Moved Horde.Storage from git to p4

diff --git a/Engine/Source/Programs/Horde/HordeStorage/Horde.Storage/Implementation/Replication/RefsReplicator.cs b/Engine/Source/Programs/Horde/HordeStorage/Horde.Storage/Implementation/Replication/RefsReplicator.cs
new file mode 100644
--- /dev/null
+++ b/Engine/Source/Programs/Horde/HordeStorage/Horde.Storage/Implementation/Replication/RefsReplicator.cs
@@ -0,0 +1,548 @@
+// Copyright Epic Games, Inc. All Rights Reserved.
+
+using System;
+using System.Collections.Generic;
+using System.IO;
+using System.Linq;
+using System.Net;
+using System.Net.Http;
+using System.Runtime.CompilerServices;
+using System.Text;
+using System.Threading;
+using System.Threading.Tasks;
+using Dasync.Collections;
+using Datadog.Trace;
+using Horde.Storage.Controllers;
+using Horde.Storage.Implementation.TransactionLog;
+using Jupiter.Implementation;
+using Microsoft.AspNetCore.Mvc;
+using Microsoft.Extensions.Options;
+using Newtonsoft.Json;
+using Nito.AsyncEx;
+using Serilog;
+
+namespace Horde.Storage.Implementation
+{
+    public class RefsReplicator : IReplicator
+    {
+        private readonly string _name;
+        private readonly ILogger _logger = Log.ForContext<RefsReplicator>();
+        private readonly FileInfo _stateFile;
+        private readonly AsyncManualResetEvent _replicationFinishedEvent = new AsyncManualResetEvent(true);
+        private readonly CancellationTokenSource _replicationTokenSource = new CancellationTokenSource();
+        private readonly ReplicatorSettings _replicatorSettings;
+        private readonly IBlobStore _blobStore;
+        private readonly IReplicationLog _replicationLog;
+        private readonly IServiceCredentials _serviceCredentials;
+        private readonly HttpClient _httpClient;
+        private RefsState _refsState;
+        private bool _replicationRunning;
+
+        public RefsReplicator(ReplicatorSettings replicatorSettings, IOptionsMonitor<ReplicationSettings> replicationSettings, IBlobStore blobStore, IHttpClientFactory httpClientFactory, IReplicationLog replicationLog, IServiceCredentials serviceCredentials)
+        {
+            _name = replicatorSettings.ReplicatorName;
+            _replicatorSettings = replicatorSettings;
+            _blobStore = blobStore;
+            _replicationLog = replicationLog;
+            _serviceCredentials = serviceCredentials;
+
+            _httpClient = httpClientFactory.CreateClient();
+            _httpClient.BaseAddress = new Uri(replicatorSettings.ConnectionString);
+
+            string stateFileName = $"{_name}.json";
+            DirectoryInfo stateRoot = new DirectoryInfo(replicationSettings.CurrentValue.StateRoot);
+            _stateFile = new FileInfo(Path.Combine(stateRoot.FullName, stateFileName));
+
+            if (_stateFile.Exists)
+            {
+                _refsState = ReadState(_stateFile)!;
+            }
+            else
+            {
+                _refsState = new RefsState();
+            }
+
+            Info = new IReplicator.ReplicatorInfo(replicatorSettings.ReplicatorName, replicatorSettings.NamespaceToReplicate, _refsState);
+        }
+
+        public void Dispose()
+        {
+            SaveState(_stateFile, _refsState);
+
+            _replicationTokenSource.Dispose();
+        }
+
+        private static RefsState ReadState(FileInfo stateFile)
+        {
+            using StreamReader streamReader = stateFile.OpenText();
+            using JsonReader reader = new JsonTextReader(streamReader);
+            JsonSerializer serializer = new JsonSerializer();
+            RefsState? state =  serializer.Deserialize<RefsState>(reader);
+            if (state == null)
+                throw new Exception("Failed to read state");
+
+            return state;
+        }
+
+        private static void SaveState(FileInfo stateFile, RefsState newState)
+        {
+            using StreamWriter writer = stateFile.CreateText();
+            JsonSerializer serializer = new JsonSerializer();
+            serializer.Serialize(writer, newState);
+        }
+
+        public async Task<bool> TriggerNewReplications()
+        {
+            if (_replicationRunning)
+            {
+                return false;
+            }
+
+            // read the state again to allow it to be modified by the admin controller / other instances of horde-storage connected to the same filesystem
+            _stateFile.Refresh();
+            if (_stateFile.Exists)
+            {
+                _refsState = ReadState(_stateFile) ?? new RefsState();
+            }
+            else
+            {
+                _refsState = new RefsState();
+            }
+
+            LogReplicationHeartbeat(0);
+
+            bool hasRun;
+
+            try
+            {
+                using Scope scope = Tracer.Instance.StartActive("replicator.run");
+                scope.Span.ResourceName =_name;
+
+                _replicationRunning = true;
+                _replicationFinishedEvent.Reset();
+                CancellationToken replicationToken = _replicationTokenSource.Token;
+
+                NamespaceId ns = _replicatorSettings.NamespaceToReplicate;
+
+                int countOfReplicationsDone = 0;
+                Info.CountOfRunningReplications = 0;
+
+                string? lastBucket = null;
+                Guid? lastEvent = null;
+
+                bool haveRunBefore = _refsState.LastBucket != null && _refsState.LastEvent != null;
+                if (!haveRunBefore)
+                {
+                    // have not run before, replicate a snapshot
+                    _logger.Information("{Name} Have not run replication before, attempting to use snapshot. State: {@State}", _name, _refsState);
+                    try
+                    {
+                        (string eventBucket, Guid eventId, int countOfEventsReplicated) = await ReplicateFromSnapshot(ns, replicationToken);
+                        countOfReplicationsDone += countOfEventsReplicated;
+
+                        // finished replicating the snapshot, persist the state
+                        _refsState.LastBucket = eventBucket;
+                        _refsState.LastEvent = eventId;
+                        SaveState(_stateFile, _refsState);
+                        hasRun = true;
+                    }
+                    catch (NoSnapshotAvailableException)
+                    {
+                        // no snapshot available so we attempt a incremental replication with no bucket set
+                        _refsState.LastBucket = null;
+                        _refsState.LastEvent = null;
+                    }
+                }
+
+                lastBucket = _refsState.LastBucket;
+                lastEvent = _refsState.LastEvent;
+
+                bool retry;
+                bool haveAttemptedSnapshot = false;
+                do
+                {
+                    retry = false;
+                    UseSnapshotException? useSnapshotException = null;
+                    try
+                    {
+                        countOfReplicationsDone += await ReplicateIncrementally(ns, lastBucket, lastEvent, replicationToken);
+                    }
+                    catch (AggregateException ae)
+                    {
+                        ae.Handle(e =>
+                        {
+                            if (e is UseSnapshotException snapshotException)
+                            {
+                                useSnapshotException = snapshotException;
+                                return true;
+                            }
+
+                            return false;
+                        });
+                    }
+                    catch (UseSnapshotException e)
+                    {
+                        useSnapshotException = e;
+                    }
+
+                    if (useSnapshotException != null)
+                    {
+                        // if we have already attempted to recover using a snapshot and we still fail we just give up and throw the exception onwards.
+                        if (haveAttemptedSnapshot)
+                            throw new AggregateException(useSnapshotException);
+
+                        // if we fail to replicate incrementally we revert to using a snapshot
+                        haveAttemptedSnapshot = true;
+                        (string eventBucket, Guid eventId, int countOfEventsReplicated) = await ReplicateFromSnapshot(ns, replicationToken, useSnapshotException.SnapshotBlob);
+                        countOfReplicationsDone += countOfEventsReplicated;
+
+                        // resume from these new events instead
+                        lastBucket = eventBucket;
+                        lastEvent = eventId;
+
+                        _refsState.LastBucket = eventBucket;
+                        _refsState.LastEvent = eventId;
+                        SaveState(_stateFile, _refsState);
+                        retry = true;
+                    }
+                    
+                } while (retry);
+                hasRun = countOfReplicationsDone != 0;
+            }
+            finally
+            {
+                _replicationRunning = false;
+                _replicationFinishedEvent.Set();
+            }
+
+            return hasRun;
+        }
+
+        private async Task<(string, Guid, int)> ReplicateFromSnapshot(NamespaceId ns, CancellationToken cancellationToken, BlobIdentifier? snapshotBlob = null)
+        {
+            // determine latest snapshot if no specific blob was specified
+            if (snapshotBlob == null)
+            {
+                HttpRequestMessage snapshotRequest = BuildHttpRequest(HttpMethod.Get, $"api/v1/replication-log/snapshots/{ns}");
+                HttpResponseMessage snapshotResponse = await _httpClient.SendAsync(snapshotRequest, cancellationToken);
+                snapshotResponse.EnsureSuccessStatusCode();
+
+                string s = await snapshotResponse.Content.ReadAsStringAsync(cancellationToken);
+                ReplicationLogSnapshots? snapshots = JsonConvert.DeserializeObject<ReplicationLogSnapshots>(s);
+                if (snapshots == null)
+                    throw new NotImplementedException();
+
+                SnapshotInfo? snapshotInfo = snapshots.Snapshots?.FirstOrDefault();
+                if (snapshotInfo == null)
+                {
+                    throw new NoSnapshotAvailableException("No snapshots found");
+                }
+
+                snapshotBlob = snapshotInfo.SnapshotBlob;
+            }
+
+            // process snapshot
+            // fetch the snapshot from the remote blob store
+            HttpRequestMessage request = BuildHttpRequest(HttpMethod.Get, $"api/v1/blobs/{ns}/{snapshotBlob}");
+            HttpResponseMessage response = await _httpClient.SendAsync(request, cancellationToken);
+            Stream blobStream = await response.Content.ReadAsStreamAsync(cancellationToken);
+
+            ReplicationLogSnapshot snapshot = await ReplicationLogSnapshot.DeserializeSnapshot(blobStream);
+
+            if (!snapshot.LastEvent.HasValue)
+                throw new Exception("No last event found");
+
+            Guid snapshotEvent = snapshot.LastEvent.Value;
+            string snapshotBucket = snapshot.LastBucket!;
+            int countOfObjectsReplicated = 0;
+            int countOfObjectsCurrentlyReplicating = 0;
+
+            // if MaxParallelReplications is set to not limit we use the default behavior of ParallelForEachAsync which is to limit based on CPUs 
+            int maxParallelism = _replicatorSettings.MaxParallelReplications != -1 ? _replicatorSettings.MaxParallelReplications : 0;
+            await snapshot.LiveObjects.ParallelForEachAsync(async snapshotLiveObject =>
+            {
+                Interlocked.Increment(ref countOfObjectsCurrentlyReplicating);
+                LogReplicationHeartbeat(countOfObjectsCurrentlyReplicating);
+
+                Info.CountOfRunningReplications = countOfObjectsCurrentlyReplicating;
+                Info.LastRun = DateTime.Now;
+
+                await ReplicateOp(ns, snapshotLiveObject.Blob, cancellationToken);
+                // TODO: Avoid adding to the replication log as we will get into infinite recursion if we do with the remote site
+                // we could add to the replication log only if the op was missing, but that could cause issue for a 3rd site replicating from us
+                // and the 3rd site is the whole reason we even add this to replication log in the first place
+                //await AddToReplicationLog(ns, snapshotLiveObject.Bucket, snapshotLiveObject.Key, snapshotLiveObject.Blob);
+                
+                Interlocked.Increment(ref countOfObjectsReplicated);
+                Interlocked.Decrement(ref countOfObjectsCurrentlyReplicating);
+            }, maxParallelism, cancellationToken);
+
+            // snapshot processed, we proceed to reading the incremental state
+            return (snapshotBucket, snapshotEvent, countOfObjectsReplicated);
+        }
+
+        private HttpRequestMessage BuildHttpRequest(HttpMethod httpMethod, string uri)
+        {
+            string? token = _serviceCredentials.GetToken();
+            HttpRequestMessage request = new HttpRequestMessage(httpMethod, uri);
+            if (!string.IsNullOrEmpty(token))
+                request.Headers.Add("Authorization", "Bearer " + token);
+            return request;
+        }
+
+        private async Task<int> ReplicateIncrementally(NamespaceId ns, string? lastBucket, Guid? lastEvent, CancellationToken replicationToken)
+        {
+            int countOfReplicationsDone = 0;
+            _logger.Information("{Name} Looking for new transaction. Previous state: {@State}", _name, _refsState);
+
+            SortedSet<long> replicationTasks = new();
+
+            // if MaxParallelReplications is set to not limit we use the default behavior of ParallelForEachAsync which is to limit based on CPUs 
+            int maxParallelism = _replicatorSettings.MaxParallelReplications != -1 ? _replicatorSettings.MaxParallelReplications : 0;
+
+            await GetRefEvents(ns, lastBucket, lastEvent, replicationToken).ParallelForEachAsync(async (ReplicationLogEvent @event) =>
+            {
+                _logger.Information("{Name} New transaction to replicate found. New op: {@Op} . Count of running replications: {CurrentReplications}", _name, @event, replicationTasks.Count);
+                
+                Info.CountOfRunningReplications = replicationTasks.Count;
+                LogReplicationHeartbeat(replicationTasks.Count);
+                long currentOffset = Interlocked.Increment(ref countOfReplicationsDone);
+
+                try
+                {
+                    string eventBucket = @event.TimeBucket;
+                    Guid eventId = @event.EventId;
+
+                    lock (replicationTasks)
+                    {
+                        replicationTasks.Add(currentOffset);
+                    }
+
+                    await ReplicateOp(@event.Namespace, @event.Blob, replicationToken);
+                    // TODO: Avoid adding to the replication log as we will get into infinite recursion if we do with the remote site
+                    // we could add to the replication log only if the op was missing, but that could cause issue for a 3rd site replicating from us
+                    // and the 3rd site is the whole reason we even add this to replication log in the first place
+                    // await AddToReplicationLog(@event.Namespace, @event.Bucket, @event.Key, @event.Blob);
+
+                    try
+                    {
+                        bool wasOldest;
+                        lock (replicationTasks)
+                        {
+                            wasOldest = currentOffset <= replicationTasks.First();
+                        }
+
+                        // only update the state when we have replicated everything up that point
+                        if (wasOldest)
+                        {
+                            // we have replicated everything up to a point and can persist this in the state
+                            _refsState.LastBucket = eventBucket;
+                            _refsState.LastEvent = eventId;
+                            SaveState(_stateFile, _refsState);
+
+                            _logger.Information("{Name} replicated all events up to {Time} . Bucket: {EventBucket} Id: {EventId}", _name, @event.Timestamp, eventBucket, eventId);
+                        }
+
+                        Info.LastRun = DateTime.Now;
+                        LogReplicationHeartbeat(replicationTasks.Count);
+                    }
+                    finally
+                    {
+                        lock (replicationTasks)
+                        {
+                            replicationTasks.Remove(currentOffset);
+                        }
+                    }
+                }
+                catch (BlobNotFoundException)
+                {
+                    _logger.Warning("{Name} Failed to replicate {@Op} in {Namespace} because blob was not present in remote store. Skipping.", _name, @event, Info.NamespaceToReplicate);
+                }
+            } , maxParallelism, breakLoopOnException: true, cancellationToken: replicationToken);
+
+            return countOfReplicationsDone;
+        }
+
+        private async Task AddToReplicationLog(NamespaceId ns, BucketId bucket, KeyId key, BlobIdentifier blob)
+        {
+            await _replicationLog.InsertAddEvent(ns, bucket, key, blob);
+        }
+
+        private async Task ReplicateOp(NamespaceId ns, BlobIdentifier blob, CancellationToken cancellationToken)
+        {
+            // We could potentially do this, but that could be dangerous if missing child references
+            // check if this blob exists locally before replicating, if it does we assume we have all of its references already
+            //if (await _blobStore.Exists(ns, blob))
+            //    return currentOffset;
+
+            HttpRequestMessage referencesRequest = BuildHttpRequest(HttpMethod.Get, $"api/v1/objects/{ns}/{blob}/references");
+            HttpResponseMessage referencesResponse = await _httpClient.SendAsync(referencesRequest, cancellationToken);
+            string body = await referencesResponse.Content.ReadAsStringAsync(cancellationToken);
+            referencesResponse.EnsureSuccessStatusCode();
+
+            ResolvedReferencesResult? refs = JsonConvert.DeserializeObject<ResolvedReferencesResult>(body);
+            if (refs == null)
+                throw new Exception($"Unable to resolve references for object {blob} in namespace {ns}");
+
+            BlobIdentifier[] potentialBlobs = new BlobIdentifier[refs.References.Length + 1];
+            Array.Copy(refs.References, potentialBlobs, refs.References.Length);
+            potentialBlobs[^1] = blob;
+
+            BlobIdentifier[] missingBlobs = await _blobStore.FilterOutKnownBlobs(ns, potentialBlobs);
+            Task[] blobReplicationTasks = new Task[missingBlobs.Length];
+            for (int i = 0; i < missingBlobs.Length; i++)
+            {
+                BlobIdentifier blobToReplicate = missingBlobs[i];
+                blobReplicationTasks[i] = Task.Run(async () =>
+                {
+                    
+                    HttpRequestMessage blobRequest = BuildHttpRequest(HttpMethod.Get, $"api/v1/blobs/{ns}/{blobToReplicate}");
+                    HttpResponseMessage blobResponse = await _httpClient.SendAsync(blobRequest, cancellationToken);
+                    byte[] blobContents = await blobResponse.Content.ReadAsByteArrayAsync(cancellationToken);
+
+                    BlobIdentifier calculatedBlob = BlobIdentifier.FromBlob(blobContents);
+                    if (!blobToReplicate.Equals(calculatedBlob))
+                    {
+                        _logger.Warning("Mismatching blob when replicating {Blob}. Determined Hash was {Hash} size was {Size}", blobToReplicate, calculatedBlob, blobContents.LongLength);
+                        // TODO: attempt to replicate again
+                        return;
+                    }
+
+                    await _blobStore.PutObject(ns, blobContents, calculatedBlob);
+                });
+            }
+
+            await Task.WhenAll(blobReplicationTasks);
+        }
+
+        private async IAsyncEnumerable<ReplicationLogEvent> GetRefEvents(NamespaceId ns, string? lastBucket, Guid? lastEvent, [EnumeratorCancellation] CancellationToken cancellationToken)
+        {
+            bool hasRunOnce = false;
+            ReplicationLogEvents logEvents;
+            do
+            {
+                if (cancellationToken.IsCancellationRequested)
+                    yield break;
+
+                if (hasRunOnce && (lastBucket == null || lastEvent == null))
+                {
+                    throw new Exception($"Failed to find state to resume from after first page of ref events, lastBucket: {lastBucket} lastEvent: {lastEvent}");
+                }
+                StringBuilder url = new StringBuilder($"api/v1/replication-log/incremental/{ns}");
+                // its okay for last bucket and last event to be null incase we have never run before, but after the first iteration we need them to keep track of where we were
+                if (lastBucket != null)
+                    url.Append($"?lastBucket={lastBucket}");
+                if (lastEvent != null)
+                    url.Append($"&lastEvent={lastEvent}");
+
+                hasRunOnce = true;
+                HttpRequestMessage request = BuildHttpRequest(HttpMethod.Get, url.ToString());
+                HttpResponseMessage response = await _httpClient.SendAsync(request, cancellationToken);
+                string body = await response.Content.ReadAsStringAsync(cancellationToken);
+                if (response.StatusCode == HttpStatusCode.BadRequest)
+                {
+                    ProblemDetails? problemDetails = JsonConvert.DeserializeObject<ProblemDetails>(body);
+                    if (problemDetails == null)
+                        throw new Exception("Unknown bad request body when reading incremental replication log. Body: {body}");
+
+                    if (problemDetails.Type == ProblemTypes.UseSnapshot)
+                    {
+                        BlobIdentifier snapshotBlob = new BlobIdentifier(problemDetails.Extensions["SnapshotId"]!.ToString()!);
+                        throw new UseSnapshotException(snapshotBlob);
+                    }
+                }
+
+                response.EnsureSuccessStatusCode();
+                ReplicationLogEvents? e = JsonConvert.DeserializeObject<ReplicationLogEvents>(body);
+                if (e == null)
+                    throw new Exception($"Unknown error when deserializing replication log events {ns} {lastBucket} {lastEvent}");
+
+                logEvents = e;
+                foreach (ReplicationLogEvent logEvent in logEvents.Events)
+                {
+                    yield return logEvent;
+
+                    lastBucket = logEvent.TimeBucket;
+                    lastEvent = logEvent.EventId;
+                }
+
+            } while (logEvents.Events.Any());
+        }
+
+        private void LogReplicationHeartbeat(int countOfCurrentReplications)
+        {
+            
+            // log message used to generate metric for how many replications are currently running
+            _logger.Information("{Name} replication has run . Count of running replications: {CurrentReplications}", _name, countOfCurrentReplications);
+
+            // log message used to verify replicators are actually running
+            _logger.Information("{Name} starting replication. Last transaction was {TransactionId} {Generation}", _name, State.ReplicatorOffset.GetValueOrDefault(0L), State.ReplicatingGeneration.GetValueOrDefault(Guid.Empty) );
+        }
+
+        public void SetReplicationOffset(long state)
+        {
+            throw new NotImplementedException();
+        }
+
+        public async Task StopReplicating()
+        {
+            _replicationTokenSource.Cancel(true);
+            await _replicationFinishedEvent.WaitAsync();
+        }
+
+        public IReplicator.ReplicatorState State
+        {
+            get { return _refsState; }
+        }
+
+        public IReplicator.ReplicatorInfo Info { get; private set; }
+
+        public Task DeleteState()
+        {
+            _refsState = new RefsState();
+            SaveState(_stateFile, _refsState);
+
+            return Task.CompletedTask;
+        }
+
+        public void SetRefState(string? lastBucket, Guid? lastEvent)
+        {
+            _refsState.LastBucket = lastBucket;
+            _refsState.LastEvent = lastEvent;
+            SaveState(_stateFile, _refsState);
+        }
+    }
+
+    internal class UseSnapshotException : Exception
+    {
+        public BlobIdentifier SnapshotBlob { get; }
+
+        public UseSnapshotException(BlobIdentifier snapshotBlob)
+        {
+            SnapshotBlob = snapshotBlob;
+        }
+
+    }
+
+    internal class NoSnapshotAvailableException : Exception
+    {
+        public NoSnapshotAvailableException(string message) : base(message)
+        {
+        }
+    }
+
+    public class RefsState : IReplicator.ReplicatorState
+    {
+        public RefsState()
+        {
+            ReplicatingGeneration = null;
+            ReplicatorOffset = 0;
+
+            LastEvent = null;
+            LastBucket = null;
+        }
+
+        public Guid? LastEvent { get; set; }
+        public string? LastBucket { get; set; }
+    }
+}
diff --git a/Engine/Source/Programs/Horde/HordeStorage/Horde.Storage/Implementation/Replication/ReplicationService.cs b/Engine/Source/Programs/Horde/HordeStorage/Horde.Storage/Implementation/Replication/ReplicationService.cs
new file mode 100644
--- /dev/null
+++ b/Engine/Source/Programs/Horde/HordeStorage/Horde.Storage/Implementation/Replication/ReplicationService.cs
@@ -0,0 +1,158 @@
+﻿// Copyright Epic Games, Inc. All Rights Reserved.
+
+using System;
+using System.Collections.Generic;
+using System.IO;
+using System.Linq;
+using System.Threading;
+using System.Threading.Tasks;
+using Jupiter;
+using Jupiter.Implementation;
+using Microsoft.Extensions.DependencyInjection;
+using Microsoft.Extensions.Options;
+using Serilog;
+
+namespace Horde.Storage.Implementation
+{
+    // ReSharper disable once ClassNeverInstantiated.Global
+    public class ReplicationService : PollingService<ReplicationService.ReplicationState>, IDisposable
+    {
+        private readonly IOptionsMonitor<ReplicationSettings> _settings;
+        private readonly ILeaderElection _leaderElection;
+        private readonly ILogger _logger = Log.ForContext<ReplicationService>();
+        private readonly Dictionary<string, Task<bool>> _currentReplications = new Dictionary<string, Task<bool>>();
+
+        public class ReplicationState
+        {
+            public List<IReplicator> Replicators { get; } = new List<IReplicator>();
+        }
+
+        public override bool ShouldStartPolling()
+        {
+            return _settings.CurrentValue.Enabled;
+        }
+
+        public ReplicationService(IOptionsMonitor<ReplicationSettings> settings, IServiceProvider provider, ILeaderElection leaderElection) :
+            base(serviceName: nameof(ReplicationService), TimeSpan.FromSeconds(settings.CurrentValue.ReplicationPollFrequencySeconds), new ReplicationState())
+        {
+            _settings = settings;
+            _leaderElection = leaderElection;
+
+            _leaderElection.OnLeaderChanged += OnLeaderChanged;
+
+            DirectoryInfo di = new DirectoryInfo(settings.CurrentValue.StateRoot);
+            Directory.CreateDirectory(di.FullName);
+
+            foreach (var replicator in settings.CurrentValue.Replicators)
+            {
+                try
+                {
+                    State.Replicators.Add(CreateReplicator(replicator, provider));
+                }
+                catch (Exception e)
+                {
+                    _logger.Error(e, "Failed to create replicator {Name}", replicator.ReplicatorName);
+                }
+            }
+
+        }
+
+        private void OnLeaderChanged(object? sender, ILeaderElection.OnLeaderChangedEventArgs e)
+        {
+            if (e.IsLeader)
+                return;
+
+            // if we are no longer the leader cancel any pending replications
+            foreach (IReplicator replicator in State.Replicators)
+            {
+                //TODO: Cancel replications
+            }
+        }
+
+        private static IReplicator CreateReplicator(ReplicatorSettings replicatorSettings, IServiceProvider provider)
+        {
+            switch (replicatorSettings.Version)
+            {
+                case ReplicatorVersion.V1:
+                    return ActivatorUtilities.CreateInstance<ReplicatorV1>(provider, replicatorSettings);
+                case ReplicatorVersion.Refs:
+                    return ActivatorUtilities.CreateInstance<RefsReplicator>(provider, replicatorSettings);
+                default:
+                    throw new ArgumentOutOfRangeException();
+            }
+        }
+
+        public override async Task<bool> OnPoll(ReplicationState state, CancellationToken cancellationToken)
+        {
+            if (!_settings.CurrentValue.Enabled)
+            {
+                _logger.Information("Skipped running replication as it is disabled");
+                return false;
+            }
+
+            if (!_leaderElection.IsThisInstanceLeader())
+            {
+                _logger.Information("Skipped running replicators because this instance was not the leader");
+                return false;
+            }
+
+            foreach (IReplicator replicator in state.Replicators)
+            {
+                if (_currentReplications.TryGetValue(replicator.Info.ReplicatorName, out Task<bool>? replicationTask))
+                {
+                    if (replicationTask.IsCompleted)
+                    {
+                        _currentReplications.Remove(replicator.Info.ReplicatorName);
+
+                        if (replicationTask.IsFaulted)
+                        {
+                            // we log the error but avoid raising it to make sure all replicators actually get to run even if there is a faulting one
+                            _logger.Error(replicationTask.Exception, "Unhandled exception in replicator {Name}", replicator.Info.ReplicatorName);
+                            continue;
+                        }
+                     
+                        await replicationTask;
+                    }
+                    else
+                    {
+                        // if the replication is still running let it continue to run
+                        continue;
+                    }
+                }
+
+                // start a new run of the replication
+                Task<bool> newReplication = replicator.TriggerNewReplications();
+                _currentReplications[replicator.Info.ReplicatorName] = newReplication;
+            }
+
+            if (state.Replicators.Count == 0)
+            {
+                _logger.Information("Finished replication poll, but no replicators configured to run.");
+            }
+
+            return state.Replicators.Count != 0;
+
+        }
+
+        protected override async Task OnStopping(ReplicationState state)
+        {
+            await Task.WhenAll(state.Replicators.Select(replicator => replicator.StopReplicating()).ToArray());
+        }
+
+        public void Dispose()
+        {
+            // we should have been stopped first so this should not be needed, but to make sure state is stored we dispose of it again
+            foreach (IReplicator replicator in State.Replicators)
+            {
+                replicator.Dispose();
+            }
+        }
+
+        public IEnumerable<IReplicator> GetReplicators(NamespaceId ns)
+        {
+            return State.Replicators
+                .Where(pair => ns == pair.Info.NamespaceToReplicate)
+                .Select(pair => pair);
+        }
+    }
+}
diff --git a/Engine/Source/Programs/Horde/HordeStorage/Horde.Storage/Implementation/Replication/ReplicationSettings.cs b/Engine/Source/Programs/Horde/HordeStorage/Horde.Storage/Implementation/Replication/ReplicationSettings.cs
new file mode 100644
--- /dev/null
+++ b/Engine/Source/Programs/Horde/HordeStorage/Horde.Storage/Implementation/Replication/ReplicationSettings.cs
@@ -0,0 +1,104 @@
+// Copyright Epic Games, Inc. All Rights Reserved.
+
+using System.Collections.Generic;
+using System.ComponentModel.DataAnnotations;
+using Jupiter.Implementation;
+
+namespace Horde.Storage
+{
+    // ReSharper disable once ClassNeverInstantiated.Global
+    public class ReplicationSettings
+    {
+        /// <summary>
+        /// Enable to start a replicating another Jupiter instance into this one
+        /// </summary>
+        public bool Enabled { get; set; } = false;
+
+        /// <summary>
+        /// Path to a directory were the local state can be kept
+        /// </summary>
+        [Required]
+        public string StateRoot { get; set; } = "";
+
+        /// <summary>
+        /// Path to a directory where state used to be stored. State is migrated from old state root to the new during startup.
+        /// </summary>
+        public string OldStateRoot { get; set; } = "";
+
+        /// <summary>
+        ///  Name of the current site, has to be globally unique across all deployments, used to avoid replication loops
+        /// </summary>
+        [Required]
+        [Key]
+        public string CurrentSite { get; set; } = "";
+
+        /// <summary>
+        /// The frequency at which to poll for new replication events
+        /// </summary>
+        [Required]
+        [Range(15, int.MaxValue)]
+        public int ReplicationPollFrequencySeconds { get; set; } = 15;
+        
+        [Required]
+        // ReSharper disable once CollectionNeverUpdated.Global
+        public List<ReplicatorSettings> Replicators { get; set; } = new List<ReplicatorSettings>();
+    }
+
+    public class ReplicatorSettings
+    {
+        /// <summary>
+        /// The namespace which this replicator is replicating, will be used both on the source and destination side 
+        /// </summary>
+        [Required]
+        public NamespaceId NamespaceToReplicate { get; set; }
+
+        /// <summary>
+        /// Name of this replicator instance, should be unique within the cluster
+        /// </summary>
+        [Required]
+        [Key]
+        public string ReplicatorName { get; set; } = "";
+
+        /// <summary>
+        /// The connection string to the remote jupiter instance which will be replicated
+        /// </summary>
+        [Required]
+        [Url]
+        public string ConnectionString { get; set; } = "";
+        public ReplicatorVersion Version { get; set; } = ReplicatorVersion.V1;
+
+        /// <summary>
+        /// Max number of replications that can run in parallel, set to -1 to disable the limit and go as wide as possible
+        /// </summary>
+        public int MaxParallelReplications { get; set; } = 64;
+    }
+
+    public enum ReplicatorVersion
+    {
+        V1, 
+        Refs
+    }
+
+    public class ServiceCredentialSettings
+    {
+        /// <summary>
+        /// The client id to use for communication with other services
+        /// </summary>
+        public string OAuthClientId { get; set; } = "";
+
+        /// <summary>
+        /// The client secret to use for communication with other services
+        /// </summary>
+        public string OAuthClientSecret { get; set; } = "";
+
+        /// <summary>
+        /// The url to login for a token to use to connect to the other services. Set to empty to disable login, assuming a unprotected service.
+        /// </summary>
+        public string OAuthLoginUrl { get; set; } = "";
+
+        /// <summary>
+        /// The scope to request
+        /// </summary>
+        public string OAuthScope { get; set; } = "cache_access";
+    }
+}
diff --git a/Engine/Source/Programs/Horde/HordeStorage/Horde.Storage/Implementation/Replication/Replicator.cs b/Engine/Source/Programs/Horde/HordeStorage/Horde.Storage/Implementation/Replication/Replicator.cs
new file mode 100644
--- /dev/null
+++ b/Engine/Source/Programs/Horde/HordeStorage/Horde.Storage/Implementation/Replication/Replicator.cs
@@ -0,0 +1,465 @@
+// Copyright Epic Games, Inc. All Rights Reserved.
+
+using System;
+using System.Collections.Generic;
+using System.IO;
+using System.Linq;
+using System.Net;
+using System.Threading;
+using System.Threading.Tasks;
+using Datadog.Trace;
+using Jupiter.Implementation;
+using Microsoft.Extensions.Options;
+using Newtonsoft.Json;
+using Nito.AsyncEx;
+using Nito.AsyncEx.Synchronous;
+using RestSharp;
+using RestSharp.Serializers.NewtonsoftJson;
+using Serilog;
+
+namespace Horde.Storage.Implementation
+{
+    public class ReplicatorV1 : Replicator<TransactionEvent>
+    {
+        private readonly ITransactionLogWriter _transactionLogWriter;
+
+        public ReplicatorV1(ReplicatorSettings replicatorSettings, IOptionsMonitor<ReplicationSettings> replicationSettings, IBlobStore blobStore, ITransactionLogWriter transactionLogWriter, IServiceCredentials serviceCredentials) : base(replicatorSettings, replicationSettings, blobStore, CreateRemoteClient(replicatorSettings, serviceCredentials))
+        {
+            _transactionLogWriter = transactionLogWriter;
+        }
+
+        internal ReplicatorV1(ReplicatorSettings replicatorSettings, IOptionsMonitor<ReplicationSettings> replicationSettings, IBlobStore blobStore, ITransactionLogWriter transactionLogWriter, IRestClient remoteClient) : base(replicatorSettings, replicationSettings, blobStore, remoteClient)
+        {
+            _transactionLogWriter = transactionLogWriter;
+        }
+
+        protected override IAsyncEnumerable<TransactionEvent> GetCallistoOp(long stateReplicatorOffset, Guid? stateReplicatingGeneration, string currentSite,
+            OpsEnumerationState enumerationState, CancellationToken replicationToken)
+        {
+            CallistoReader remoteCallistoReader = new CallistoReader(_restClient, _replicatorSettings.NamespaceToReplicate);
+
+            return remoteCallistoReader.GetOps(stateReplicatorOffset, stateReplicatingGeneration, currentSite, enumerationState: enumerationState, cancellationToken: replicationToken);
+        }
+
+        protected override async Task ReplicateOp(IRestClient remoteClient, TransactionEvent op, CancellationToken replicationToken)
+        {
+            using Scope scope = Tracer.Instance.StartActive("replicator.replicate_blobs");
+            NamespaceId ns = _replicatorSettings.NamespaceToReplicate;
+
+            // now we replicate the blobs
+            switch (op)
+            {
+                case AddTransactionEvent addEvent:
+                    BlobIdentifier[] blobs = addEvent.Blobs;
+                    await ReplicateBlobs(remoteClient, ns, blobs, replicationToken);
+                    break;
+                case RemoveTransactionEvent removeEvent:
+                    // TODO: Do we even want to do anything? we can not delete the blob in the store because it may be used by something else
+                    // so we wait for the GC to remove it.
+                    break;
+                default:
+                    throw new ArgumentOutOfRangeException();
+            }
+        }
+
+        protected override async Task<long?> ReplicateOpInline(IRestClient remoteClient, TransactionEvent op, CancellationToken replicationToken)
+        {
+            using Scope scope = Tracer.Instance.StartActive("replicator.replicate_inline");
+            //scope.Span.ResourceName = op.transactionId;
+
+            NamespaceId ns = _replicatorSettings.NamespaceToReplicate;
+
+            // Put a copy of the event in the local callisto store
+            // Make sure this is put before the io put request as we want to pin the reference before uploading the content to avoid the gc removing it right away
+            long newTransactionId = await _transactionLogWriter.Add(ns, op);
+
+            _logger.Information("{Name} Replicated Op {@Op} to local store as transaction id {@TransactionId}", _name, op, newTransactionId);
+
+            // TODO: Add Refstore replication
+            //_refStore.Add();
+
+            return newTransactionId;
+        }
+
+    }
+
+    public abstract class Replicator<T> : IReplicator
+        where T: IReplicationEvent
+    {
+        private readonly IBlobStore _blobStore;
+        protected readonly string _name;
+        private readonly string _currentSite;
+        protected readonly ReplicatorSettings _replicatorSettings;
+        protected readonly ILogger _logger = Log.ForContext<Replicator<T>>();
+
+        private readonly FileInfo _stateFile;
+        private bool _replicationRunning;
+        private readonly AsyncManualResetEvent _replicationFinishedEvent = new AsyncManualResetEvent(true);
+        private readonly CancellationTokenSource _replicationTokenSource = new CancellationTokenSource();
+        private readonly IReplicator.ReplicatorInfo _replicatorInfo;
+        protected readonly IRestClient _restClient;
+
+
+        public IReplicator.ReplicatorState State { get; private set; }
+
+        public IReplicator.ReplicatorInfo Info
+        {
+            get { return _replicatorInfo; }
+        }
+
+        protected static IRestClient CreateRemoteClient(ReplicatorSettings replicatorSettings, IServiceCredentials serviceCredentials)
+        {
+            var remoteClient = new RestClient(replicatorSettings.ConnectionString).UseSerializer(() => new JsonNetSerializer());
+            remoteClient.Authenticator = serviceCredentials.GetAuthenticator();
+            return remoteClient;
+        }
+
+        protected Replicator(ReplicatorSettings replicatorSettings, IOptionsMonitor<ReplicationSettings> replicationSettings, IBlobStore blobStore, IRestClient remoteClient)
+        {
+            _name = replicatorSettings.ReplicatorName;
+            _blobStore = blobStore;
+            _currentSite = replicationSettings.CurrentValue.CurrentSite;
+            _replicatorSettings = replicatorSettings;
+
+            string stateFileName = $"{_name}.json";
+            DirectoryInfo stateRoot = new DirectoryInfo(replicationSettings.CurrentValue.StateRoot);
+            _stateFile = new FileInfo(Path.Combine(stateRoot.FullName, stateFileName));
+
+            if (_stateFile.Exists)
+            {
+                State = ReadState(_stateFile) ?? new IReplicator.ReplicatorState {ReplicatorOffset = null, ReplicatingGeneration = null};
+            }
+            else
+            {
+                State = new IReplicator.ReplicatorState {ReplicatorOffset = null, ReplicatingGeneration = null};
+            }
+
+            _replicatorInfo = new IReplicator.ReplicatorInfo(_replicatorSettings.ReplicatorName, _replicatorSettings.NamespaceToReplicate, State);
+            _restClient = remoteClient;
+        }
+
+        ~Replicator()
+        {
+            Dispose();
+        }
+
+        /// <summary>
+        /// Attempt to run a new replication, if a replication is already in flight for this replicator this will early exist.
+        /// </summary>
+        /// <returns>True if the replication actually attempted to run</returns>
+        public async Task<bool> TriggerNewReplications()
+        {
+            if (_replicationRunning)
+            {
+                return false;
+            }
+
+            bool hasRun = false;
+            OpsEnumerationState enumerationState = new();
+            DateTime startedReplicationAt = DateTime.Now;
+            long countOfReplicatedEvents = 0L;
+            try
+            {
+                using Scope scope = Tracer.Instance.StartActive("replicator.run");
+                scope.Span.ResourceName =_name;
+
+                _replicationRunning = true;
+                _replicationFinishedEvent.Reset();
+                CancellationToken replicationToken = _replicationTokenSource.Token;
+
+                // read the state again to allow it to be modified by the admin controller
+                _stateFile.Refresh();
+                if (_stateFile.Exists)
+                {
+                    State = ReadState(_stateFile) ?? new IReplicator.ReplicatorState {ReplicatorOffset = null, ReplicatingGeneration = null};
+                }
+                else
+                {
+                    State = new IReplicator.ReplicatorState {ReplicatorOffset = null, ReplicatingGeneration = null};
+                }
+
+                LogReplicationHeartbeat(_name, State, 0);
+                _logger.Information("{Name} Looking for new transaction. Previous state: {@State}", _name, State);
+
+                SortedList<long, Task> replicationTasks = new();
+                await foreach (T op in GetCallistoOp(State.ReplicatorOffset ?? 0, State.ReplicatingGeneration, _currentSite, enumerationState, replicationToken))
+                {
+                    hasRun = true;
+                    _logger.Information("{Name} New transaction to replicate found. New op: {@Op} . Count of running replications: {CurrentReplications}", _name, op, replicationTasks.Count);
+
+                    Info.CountOfRunningReplications = replicationTasks.Count;
+
+                    if (!op.Identifier.HasValue || !op.NextIdentifier.HasValue)
+                    {
+                        _logger.Error("{Name} Missing identifier in {@Op}", _name, op);
+                        throw new Exception("Missing identifier in op. Version mismatch?");
+                    }
+                    
+                    long currentOffset = op.Identifier.Value;
+                    long nextIdentifier = op.NextIdentifier.Value;
+                    Guid replicatingGeneration = enumerationState.ReplicatingGeneration;
+                    try
+                    {
+                        // first replicate the things which are order sensitive like the transaction log
+                        long? _ = await ReplicateOpInline(_restClient, op, replicationToken);
+
+                        // next we replicate the large things that overlap with other replications like blobs
+                        Task currentReplicationTask = ReplicateOp(_restClient, op, replicationToken);
+
+                        Task OnDone(Task task)
+                        {
+                            Interlocked.Increment(ref countOfReplicatedEvents);
+                            try
+                            {
+                                bool wasOldest;
+                                lock (replicationTasks)
+                                {
+                                    wasOldest = currentOffset <= replicationTasks.First().Key;
+                                }
+
+                                if (task.IsFaulted)
+                                {
+                                    // rethrow the exception if there was any
+                                    task.WaitAndUnwrapException(replicationToken);
+                                }
+
+                                // only update the state when we have replicated everything up that point
+                                if (wasOldest)
+                                {
+                                    // we have finished replicating 
+                                    State.ReplicatorOffset = nextIdentifier;
+                                    State.ReplicatingGeneration = replicatingGeneration;
+                                    SaveState(_stateFile, State);
+                                }
+
+                                Info.LastRun = DateTime.Now;
+                                Info.CountOfRunningReplications = replicationTasks.Count;
+
+                                LogReplicationHeartbeat(_name, State, replicationTasks.Count);
+
+                                return task;
+                            }
+                            finally
+                            {
+                                lock (replicationTasks)
+                                {
+                                    replicationTasks.Remove(currentOffset);
+                                }
+                            }
+                        }
+
+                        lock (replicationTasks)
+                        {
+                            Task<Task> commitStateTask = currentReplicationTask.ContinueWith(OnDone, replicationToken);
+
+                            replicationTasks.Add(currentOffset, commitStateTask);
+                        }
+                    }
+                    catch (BlobNotFoundException)
+                    {
+                        _logger.Warning("{Name} Failed to replicate {@Op} in {Namespace} because blob was not present in remote store. Skipping.", _name, op, Info.NamespaceToReplicate);
+                    }
+
+                    // if we have reached the max amount of parallel replications we wait for one of them to finish before starting a new one
+                    // if max replications is set to -1 we do not limit the concurrency
+                    if ( _replicatorSettings.MaxParallelReplications != -1 && replicationTasks.Count >= _replicatorSettings.MaxParallelReplications)
+                    {
+                        Task[] currentReplicationTasks;
+                        lock (replicationTasks)
+                        {
+                            currentReplicationTasks = replicationTasks.Values.ToArray();
+                        }
+                        await Task.WhenAny(currentReplicationTasks);
+                    }
+
+                    Info.CountOfRunningReplications = replicationTasks.Count;
+                    LogReplicationHeartbeat(_name, State, replicationTasks.Count);
+
+                    if (replicationToken.IsCancellationRequested)
+                    {
+                        break;
+                    }
+                }
+
+                // make a copy to avoid the collection being modified while we wait for it
+                List<Task> tasksToWaitFor = replicationTasks.Values.ToList();
+                await Task.WhenAll(tasksToWaitFor);
+            }
+            finally
+            {
+                _replicationRunning = false;
+                _replicationFinishedEvent.Set();
+            }
+
+            DateTime endedReplicationAt = DateTime.Now;
+            _logger.Information("{Name} {Namespace} Replication Finished at {EndTime} and last run was at {LastTime}, replicated {CountOfEvents} events. Thus was {TimeDifference} minutes behind.", _name, Info.NamespaceToReplicate, endedReplicationAt, startedReplicationAt, countOfReplicatedEvents, (endedReplicationAt - startedReplicationAt).TotalMinutes);
+
+            if (!hasRun)
+            {
+                // if no events were found we still persist the offsets we have run thru so we don't have to read them again.
+                State.ReplicatorOffset = enumerationState.LastOffset;
+                State.ReplicatingGeneration = enumerationState.ReplicatingGeneration;
+                SaveState(_stateFile, State);
+            }
+
+            return hasRun;
+        }
+
+        protected abstract IAsyncEnumerable<T> GetCallistoOp(long stateReplicatorOffset, Guid? stateReplicatingGeneration, string currentSite, OpsEnumerationState enumerationState, CancellationToken replicationToken);
+        
+        private void LogReplicationHeartbeat(string name, IReplicator.ReplicatorState? state, int countOfCurrentReplications)
+        {
+            if (state == null)
+                return;
+
+            // log message used to generate metric for how many replications are currently running
+            _logger.Information("{Name} replication has run . Count of running replications: {CurrentReplications}", _name, countOfCurrentReplications);
+
+            // log message used to verify replicators are actually running
+            _logger.Information("{Name} starting replication. Last transaction was {TransactionId} {Generation}", name, state.ReplicatorOffset.GetValueOrDefault(0L), state.ReplicatingGeneration.GetValueOrDefault(Guid.Empty) );
+        }
+
+        protected abstract Task<long?> ReplicateOpInline(IRestClient remoteClient, T op, CancellationToken replicationToken);
+
+        protected abstract Task ReplicateOp(IRestClient remoteClient, T op, CancellationToken replicationToken);
+
+        protected async Task ReplicateBlobs(IRestClient remoteClient, NamespaceId ns, BlobIdentifier[] blobs, CancellationToken replicationToken)
+        {
+            Task[] blobReplicateTasks = new Task[blobs.Length];
+            for (int index = 0; index < blobs.Length; index++)
+            {
+                BlobIdentifier blob = blobs[index];
+                blobReplicateTasks[index] = Task.Run(async () =>
+                {
+                    // check if this blob exists locally before replicating
+                    if (await _blobStore.Exists(ns, blob))
+                        return;
+
+                    // attempt to replicate for a few tries with some delay in between
+                    // this because new transactions are written to callisto first (to establish the GC handle) before content is uploaded to io.
+                    // as such its possible (though not very likely) when we replicate the very newest transaction that we find a transaction but no content
+                    // and as the content upload can be large it can take a little time for it to exist in io.
+                    byte[]? rawContent = null;
+                    BlobIdentifier? calculatedBlob = null;
+                    const int MaxAttempts = 3;
+                    for (int attempts = 0; attempts < MaxAttempts; attempts++)
+                    {
+                        RestRequest remoteIoGet = new RestRequest("api/v1/s/{ns}/{blob}");
+                        remoteIoGet.AddUrlSegment("ns", ns);
+                        remoteIoGet.AddUrlSegment("blob", blob);
+
+                        IRestResponse response = await remoteClient.ExecuteGetAsync(remoteIoGet, replicationToken);
+
+                        if (response.StatusCode == HttpStatusCode.NotFound)
+                        {
+                            _logger.Information("{Blob} was not found in remote blob store. Retry attempt {Attempts}", blob, attempts);
+                            await Task.Delay(TimeSpan.FromSeconds(5), replicationToken);
+                            continue;
+                        }
+
+                        if (response.StatusCode == HttpStatusCode.GatewayTimeout)
+                        {
+                            _logger.Information("GatewayTimeout while replicating {Blob}, retrying. Retry attempt {Attempts}", blob, attempts);
+                            continue;
+                        }
+
+                        if (response.StatusCode == HttpStatusCode.BadRequest || response.StatusCode == HttpStatusCode.NotFound)
+                        {
+                            _logger.Warning("Remote blob {Blob} missing, unable to replicate.", blob);
+                            return;
+                        }
+
+                        if (response.IsSuccessful)
+                        {
+                            rawContent = response.RawBytes;
+
+                            calculatedBlob = BlobIdentifier.FromBlob(rawContent);
+                            if (!blob.Equals(calculatedBlob))
+                            {
+                                _logger.Warning("Mismatching blob when replicating {Blob}. Determined Hash was {Hash} size was {Size} HttpStatusCode {StatusCode} HttpMessage {HttpMessage} ResponseUri {Url}", blob, calculatedBlob, rawContent.LongLength, response.StatusCode, response.ErrorMessage, response.ResponseUri);
+                                continue; // attempt to replicate again
+                            }
+    
+                            // this blob is good and should be added to the blob store
+                            break;
+                        }
+
+                        if (response.ErrorException is WebException we)
+                        {
+                            if (we.Status == WebExceptionStatus.Timeout)
+                            {
+                                const string template = "Operation timed out while replicating {Blob}, retrying. Retry attempt {Attempts}";
+                                _logger.Information(we, template, blob, attempts);
+                                continue;
+                            }
+                        }
+
+                        throw new Exception($"Replicator \"{_name}\" failed to replicate blob {blob} due unsuccessful status from remote blob store: {response.StatusCode} . Error message: \"{response.ErrorMessage}\"", response.ErrorException);
+                    }
+
+                    if (rawContent == null || calculatedBlob == null)
+                    {
+                        _logger.Warning("Remote blob {Blob} not present in remote blob store after multiple attempts to replicate. Assuming this blob has been GCed and continuing replication.", blob);
+                        return;
+                    }
+
+                    if (!blob.Equals(calculatedBlob))
+                    {
+                        _logger.Warning("Mismatching blob when replicating {Blob}. Determined Hash was {Hash} size was {Size}. Multiple attempts failed, giving up.", blob, calculatedBlob, rawContent.LongLength);
+                        return; 
+                    }
+
+                    await _blobStore.PutObject(ns, rawContent, calculatedBlob);
+                }, replicationToken);
+            }
+
+            await Task.WhenAll(blobReplicateTasks);
+        }
+
+        private static IReplicator.ReplicatorState? ReadState(FileInfo stateFile)
+        {
+            using StreamReader streamReader = stateFile.OpenText();
+            using JsonReader reader = new JsonTextReader(streamReader);
+            JsonSerializer serializer = new JsonSerializer();
+            return serializer.Deserialize<IReplicator.ReplicatorState>(reader);
+        }
+
+        private static void SaveState(FileInfo stateFile, IReplicator.ReplicatorState newState)
+        {
+            using StreamWriter writer = stateFile.CreateText();
+            JsonSerializer serializer = new JsonSerializer();
+            serializer.Serialize(writer, newState);
+        }
+
+        public void Dispose()
+        {
+            SaveState(_stateFile, State);
+
+            _replicationTokenSource.Dispose();
+        }
+
+        public async Task StopReplicating()
+        {
+            _replicationTokenSource.Cancel(true);
+            await _replicationFinishedEvent.WaitAsync();
+        }
+
+        public void SetReplicationOffset(long offset)
+        {
+            State.ReplicatorOffset = offset;
+            SaveState(_stateFile, State);
+        }
+
+        public Task DeleteState()
+        {
+            State = new IReplicator.ReplicatorState()
+            {
+                ReplicatingGeneration = null,
+                ReplicatorOffset = 0
+            };
+            SaveState(_stateFile, State);
+
+            return Task.CompletedTask;
+        }
+    }
+}
diff --git a/Engine/Source/Programs/Horde/HordeStorage/Horde.Storage/Implementation/ServiceCredentials.cs b/Engine/Source/Programs/Horde/HordeStorage/Horde.Storage/Implementation/ServiceCredentials.cs
new file mode 100644
--- /dev/null
+++ b/Engine/Source/Programs/Horde/HordeStorage/Horde.Storage/Implementation/ServiceCredentials.cs
@@ -0,0 +1,49 @@
+﻿using System;
+using System.Collections.Generic;
+using System.Linq;
+using System.Threading.Tasks;
+using Jupiter.Common.Implementation;
+using Jupiter.Implementation;
+using Microsoft.Extensions.Options;
+using RestSharp.Authenticators;
+
+namespace Horde.Storage.Implementation
+{
+    public interface IServiceCredentials
+    {
+        IAuthenticator? GetAuthenticator();
+
+        string? GetToken();
+    }
+
+    public class ServiceCredentials : IServiceCredentials
+    {
+        private readonly ClientCredentialOAuthAuthenticator? _authenticator;
+
+        public ServiceCredentials(IOptionsMonitor<ServiceCredentialSettings> settingsMonitor, ISecretResolver secretResolver)
+        {
+            ServiceCredentialSettings settings = settingsMonitor.CurrentValue;
+            if (!string.IsNullOrEmpty(settings.OAuthLoginUrl))
+            {
+                string? clientId = secretResolver.Resolve(settings.OAuthClientId);
+                if (string.IsNullOrEmpty(clientId))
+                    throw new ArgumentException("ClientId must be set when using a service credential");
+                string? clientSecret = secretResolver.Resolve(settings.OAuthClientSecret);
+                if (string.IsNullOrEmpty(clientSecret))
+                    throw new ArgumentException("ClientSecret must be set when using a service credential");
+                _authenticator = new ClientCredentialOAuthAuthenticator(settings.OAuthLoginUrl, clientId, clientSecret, settings.OAuthScope);
+            }
+        }
+
+        public IAuthenticator? GetAuthenticator()
+        {
+            return _authenticator;
+        }
+
+        public string? GetToken()
+        {
+            return _authenticator?.Authenticate();
+        }
+    }
+
+}
diff --git a/Engine/Source/Programs/Horde/HordeStorage/Horde.Storage/Implementation/TransactionLog/CallistoTransactionLogWriter.cs b/Engine/Source/Programs/Horde/HordeStorage/Horde.Storage/Implementation/TransactionLog/CallistoTransactionLogWriter.cs
new file mode 100644
--- /dev/null
+++ b/Engine/Source/Programs/Horde/HordeStorage/Horde.Storage/Implementation/TransactionLog/CallistoTransactionLogWriter.cs
@@ -0,0 +1,91 @@
+// Copyright Epic Games, Inc. All Rights Reserved.
+
+using System;
+using System.Threading.Tasks;
+using Jupiter.Implementation;
+using Serilog;
+using Microsoft.Extensions.Options;
+using RestSharp;
+using RestSharp.Serializers.NewtonsoftJson;
+
+namespace Horde.Storage.Implementation
+{
+    internal class CallistoTransactionLogWriter : ITransactionLogWriter
+    {
+        private readonly ILogger _logger = Log.ForContext<CallistoTransactionLogWriter>();
+        private readonly IRestClient _client;
+
+        public CallistoTransactionLogWriter(IOptionsMonitor<CallistoTransactionLogSettings> settings, IServiceCredentials serviceCredentials)
+        {
+            _client = new RestClient(settings.CurrentValue.ConnectionString).UseSerializer(() => new JsonNetSerializer());
+            _client.Authenticator = serviceCredentials.GetAuthenticator();
+        }
+
+
+        public async Task<long> Add(NamespaceId ns, TransactionEvent @event)
+        {
+            const int retryAttempts = 15;
+            for (int i = 0; i < retryAttempts; i++)
+            {
+                RestRequest request = new RestRequest("api/v1/t/{ns}");
+                request.AddUrlSegment("ns", ns);
+                request.AddJsonBody(@event);
+
+                _logger.Verbose("Sending record {@Record} to callisto", @event);
+
+                IRestResponse<CallistoResponse> response = await _client.ExecutePostAsync<CallistoResponse>(request);
+                if (!response.IsSuccessful)
+                {
+                    // if callisto is down we wait a while and try again
+                    if (response.StatusCode == 0 && response.ErrorMessage.StartsWith("Connection refused"))
+                    {
+                        await Task.Delay(1000);
+                        continue; // retry
+                    }
+
+                    // unknown error, throw exception
+                    throw new Exception($"Failed to put add record in Callisto. Status: {response.StatusCode}. Message: {response.ErrorMessage}. Body: {response.Content}");
+                }
+                
+                return response.Data.Offset;
+            }
+
+            throw new Exception($"Failed to reach callisto after {retryAttempts} attempts");
+        }
+
+        public async Task<long> Delete(NamespaceId ns, BucketId bucket, KeyId key)
+        {
+            RestRequest request = new RestRequest("api/v1/t/{ns}");
+            var record = new CallistoRequest("Delete", key, bucket);
+            request.AddUrlSegment("ns", ns);
+            request.AddJsonBody(record);
+            
+            _logger.Verbose("Sending delete record {@Record} to callisto", record);
+
+            CallistoResponse response = await _client.PostAsync<CallistoResponse>(request);
+
+            if (response == null)
+                throw new Exception("Failed to put delete record in Callisto.");
+            return response.Offset;
+        }
+    }
+
+    internal class CallistoRequest
+    {
+        public CallistoRequest(string type, KeyId name, BucketId bucket)
+        {
+            Type = type;
+            Name = name;
+            Bucket = bucket;
+        }
+
+        public string Type { get; }
+        public KeyId Name { get; }
+        public BucketId Bucket { get; }
+    }
+
+    internal class CallistoResponse
+    {
+        public long Offset { get; set; }
+    }
+}
diff --git a/Engine/Source/Programs/Horde/HordeStorage/Horde.Storage/Implementation/TransactionLog/IReplicationLog.cs b/Engine/Source/Programs/Horde/HordeStorage/Horde.Storage/Implementation/TransactionLog/IReplicationLog.cs
new file mode 100644
--- /dev/null
+++ b/Engine/Source/Programs/Horde/HordeStorage/Horde.Storage/Implementation/TransactionLog/IReplicationLog.cs
@@ -0,0 +1,75 @@
+﻿// Copyright Epic Games, Inc. All Rights Reserved.
+
+using System;
+using System.Collections.Generic;
+using System.Reflection.Metadata;
+using System.Threading.Tasks;
+using Horde.Storage.Implementation.TransactionLog;
+using Jupiter.Implementation;
+using Newtonsoft.Json;
+
+namespace Horde.Storage.Implementation
+{
+    public interface IReplicationLog
+    {
+        IAsyncEnumerable<NamespaceId> GetNamespaces();
+
+        Task<(string, Guid)> InsertAddEvent(NamespaceId ns, BucketId bucket, KeyId key, BlobIdentifier objectBlob, DateTime? timeBucket = null);
+        Task<(string, Guid)> InsertDeleteEvent(NamespaceId ns, BucketId bucket, KeyId key, BlobIdentifier objectBlob, DateTime? timeBucket = null);
+        IAsyncEnumerable<ReplicationLogEvent> Get(NamespaceId ns, string? lastBucket, Guid? lastEvent);
+
+        Task AddSnapshot(SnapshotInfo snapshotHeader);
+        Task<SnapshotInfo?> GetLatestSnapshot(NamespaceId ns);
+        IAsyncEnumerable<SnapshotInfo> GetSnapshots(NamespaceId ns);
+    }
+
+    public class SnapshotInfo
+    {
+        [JsonConstructor]
+        public SnapshotInfo(NamespaceId snapshottedNamespace, NamespaceId blobNamespace, BlobIdentifier snapshotBlob)
+        {
+            SnapshottedNamespace = snapshottedNamespace;
+            BlobNamespace = blobNamespace;
+            SnapshotBlob = snapshotBlob;
+        }
+
+        public NamespaceId SnapshottedNamespace { get; set; }
+        public NamespaceId BlobNamespace { get; set; }
+        public BlobIdentifier SnapshotBlob { get; set; }
+    }
+
+    public class ReplicationLogEvent
+    {
+        public ReplicationLogEvent(NamespaceId @namespace, BucketId bucket, KeyId key, BlobIdentifier blob, Guid eventId, string timeBucket, DateTime timestamp, OpType op)
+        {
+            Namespace = @namespace;
+            Bucket = bucket;
+            Key = key;
+            Blob = blob;
+            EventId = eventId;
+            TimeBucket = timeBucket;
+            Timestamp = timestamp;
+            Op = op;
+        }
+
+        // these are serialized as ints so make sure to keep the value intact
+        public enum OpType
+        {
+            Added = 0,
+            Deleted = 1
+        };
+
+        public NamespaceId Namespace { get; }
+        public BucketId Bucket { get; }
+        public KeyId Key { get; }
+        public OpType Op { get; }
+        public DateTime Timestamp { get; }
+        public string TimeBucket { get; }
+        public Guid EventId { get; }
+        public BlobIdentifier Blob { get; }
+    }
+
+    public class IncrementalLogNotAvailableException : Exception
+    {
+    }
+}
diff --git a/Engine/Source/Programs/Horde/HordeStorage/Horde.Storage/Implementation/TransactionLog/ITransactionLogWriter.cs b/Engine/Source/Programs/Horde/HordeStorage/Horde.Storage/Implementation/TransactionLog/ITransactionLogWriter.cs
new file mode 100644
--- /dev/null
+++ b/Engine/Source/Programs/Horde/HordeStorage/Horde.Storage/Implementation/TransactionLog/ITransactionLogWriter.cs
@@ -0,0 +1,13 @@
+﻿// Copyright Epic Games, Inc. All Rights Reserved.
+
+using System.Threading.Tasks;
+using Jupiter.Implementation;
+
+namespace Horde.Storage.Implementation
+{
+    public interface ITransactionLogWriter
+    {
+        Task<long> Add(NamespaceId ns, TransactionEvent @event);
+        Task<long> Delete(NamespaceId ns, BucketId bucket, KeyId key);
+    }
+}
diff --git a/Engine/Source/Programs/Horde/HordeStorage/Horde.Storage/Implementation/TransactionLog/MemoryReplicationLog.cs b/Engine/Source/Programs/Horde/HordeStorage/Horde.Storage/Implementation/TransactionLog/MemoryReplicationLog.cs
new file mode 100644
--- /dev/null
+++ b/Engine/Source/Programs/Horde/HordeStorage/Horde.Storage/Implementation/TransactionLog/MemoryReplicationLog.cs
@@ -0,0 +1,188 @@
+// Copyright Epic Games, Inc. All Rights Reserved.
+
+using System;
+using System.Collections.Concurrent;
+using System.Collections.Generic;
+using System.Linq;
+using System.Threading.Tasks;
+using Cassandra;
+using Dasync.Collections;
+using Jupiter.Implementation;
+
+namespace Horde.Storage.Implementation
+{
+    internal class MemoryReplicationLog : IReplicationLog
+    {
+        private readonly ConcurrentDictionary<NamespaceId, SortedList<string, SortedList<TimeUuid, ReplicationLogEvent>>> _replicationEvents = new();
+        private readonly ConcurrentDictionary<NamespaceId, List<SnapshotInfo>>  _snapshots = new();
+
+        public IAsyncEnumerable<NamespaceId> GetNamespaces()
+        {
+            return _replicationEvents.Keys.ToAsyncEnumerable();
+        }
+
+        public Task<(string, Guid)> InsertAddEvent(NamespaceId ns, BucketId bucket, KeyId key, BlobIdentifier objectBlob, DateTime? timestamp)
+        {
+            return DoInsert(ns, bucket, key, objectBlob, ReplicationLogEvent.OpType.Added, timestamp);
+        }
+
+        private async Task<(string, Guid)> DoInsert(NamespaceId ns, BucketId bucket, KeyId key, BlobIdentifier hash, ReplicationLogEvent.OpType op, DateTime? lastTimestamp)
+        {
+            DateTime timestamp = lastTimestamp.GetValueOrDefault(DateTime.Now);
+
+            return await Task.Run(() =>
+            {
+                DateTime hourlyBucket = timestamp.ToHourlyBucket();
+                string bucketId = hourlyBucket.ToReplicationBucketIdentifier();
+                TimeUuid eventId = TimeUuid.NewId(timestamp);
+                ReplicationLogEvent logEvent = new ReplicationLogEvent(ns, bucket, key, hash, eventId, bucketId, hourlyBucket, op);
+
+                _replicationEvents.AddOrUpdate(ns, _ =>
+                {
+                    SortedList<string, SortedList<TimeUuid, ReplicationLogEvent>> l = new() { { bucketId, new () { {eventId, logEvent} } } };
+                    return l;
+                }, (_, buckets) =>
+                {
+                    lock (buckets)
+                    {
+                        if (buckets.TryGetValue(bucketId, out SortedList<TimeUuid, ReplicationLogEvent>? events))
+                        {
+                            lock (events)
+                            {
+                                events.Add(eventId, logEvent);
+                            }
+                        }
+                        else
+                        {
+                            buckets.Add(bucketId, new () {{eventId, logEvent} });
+                        }
+                    }
+
+                    return buckets;
+                });
+
+                return (bucketId, eventId.ToGuid());
+            });
+        }
+
+        public Task<(string, Guid)> InsertDeleteEvent(NamespaceId ns, BucketId bucket, KeyId key, BlobIdentifier objectBlob, DateTime? timestamp)
+        {
+            return DoInsert(ns, bucket, key, objectBlob, ReplicationLogEvent.OpType.Deleted, timestamp); 
+        }
+
+        public async IAsyncEnumerable<ReplicationLogEvent> Get(NamespaceId ns, string? lastBucket, Guid? lastEvent)
+        {
+            await Task.CompletedTask;
+            if (!_replicationEvents.TryGetValue(ns, out SortedList<string, SortedList<TimeUuid, ReplicationLogEvent>>? buckets))
+            {
+                throw new NamespaceNotFoundException(ns);
+            }
+
+            // if no previous bucket was specified we started from the oldest one
+            if (lastBucket == null)
+            {
+                // if we have no buckets we are done
+                if (!buckets.Any())
+                    yield break;
+
+                lastBucket = buckets.First().Key;
+            }
+
+            int currentBucketIndex = buckets.IndexOfKey(lastBucket!);
+            if (currentBucketIndex == -1)
+            {
+                // the bucket to resume from does not exist anymore, we should start from a snapshot
+                throw new IncrementalLogNotAvailableException();
+            }
+            // loop over all the buckets after this including the previous one
+            // we will then skip all the events we have previously processed
+            for (int i = currentBucketIndex; i < buckets.Keys.Count; i++)
+            {
+                string bucket = buckets.Keys[i];
+
+                if (!buckets.TryGetValue(bucket, out SortedList<TimeUuid, ReplicationLogEvent>? logEvents))
+                {
+                    yield break;
+                }
+
+                int skipCount = 0;
+                if (lastEvent.HasValue)
+                {
+                    int eventIndex = logEvents.IndexOfKey(lastEvent.Value);
+                    // if the event to resume from does not exist we need to use a snapshot instead
+                    if (eventIndex == -1)
+                    {
+                        throw new IncrementalLogNotAvailableException();
+                    }
+
+                    skipCount = eventIndex + 1;
+                }
+
+                foreach ((TimeUuid _, ReplicationLogEvent value) in logEvents.Skip(skipCount))
+                {
+                    yield return value;
+                }
+
+                // once we have enumerated the bucket we were working on we should now return all events from the remaining ones
+                lastEvent = null;
+            }
+        }
+
+        public Task AddSnapshot(SnapshotInfo snapshotHeader)
+        {
+            _snapshots.AddOrUpdate(snapshotHeader.SnapshottedNamespace, _ => new List<SnapshotInfo> { snapshotHeader }, (_, list) =>
+            {
+                // we want the newest snapshots first
+                list.Insert(0, snapshotHeader);
+                const int maxCountOfSnapshots = 10;
+                if (list.Count > maxCountOfSnapshots)
+                {
+                    // remove the last snapshots as they are the oldest
+                    list.RemoveRange(maxCountOfSnapshots, list.Count - 10);
+                }
+                return list;
+            });
+
+            return Task.CompletedTask;
+        }
+
+        public Task<SnapshotInfo?> GetLatestSnapshot(NamespaceId ns)
+        {
+            if (!_snapshots.TryGetValue(ns, out List<SnapshotInfo>? snapshots))
+            {
+                return Task.FromResult<SnapshotInfo?>(null);
+            }
+
+            return Task.FromResult<SnapshotInfo?>(snapshots.Last());
+        }
+
+        public async IAsyncEnumerable<SnapshotInfo> GetSnapshots(NamespaceId ns)
+        {
+            await Task.CompletedTask;
+
+            if (!_snapshots.TryGetValue(ns, out List<SnapshotInfo>? snapshots))
+            {
+                yield break;
+            }
+
+            foreach (SnapshotInfo snapshot in snapshots)
+            {
+                yield return snapshot;
+            }
+
+        }
+    }
+
+    public static class DateTimeUtils
+    {
+        public static DateTime ToHourlyBucket(this DateTime timestamp)
+        {
+            return new DateTime(timestamp.Year, timestamp.Month, timestamp.Day, timestamp.Hour, 0, 0);
+        }
+
+        public static string ToReplicationBucketIdentifier(this DateTime timestamp)
+        {
+            return $"rep-{timestamp.ToFileTime()}";
+        }
+    }
+}
diff --git a/Engine/Source/Programs/Horde/HordeStorage/Horde.Storage/Implementation/TransactionLog/MemoryTransactionLogWriter.cs b/Engine/Source/Programs/Horde/HordeStorage/Horde.Storage/Implementation/TransactionLog/MemoryTransactionLogWriter.cs
new file mode 100644
--- /dev/null
+++ b/Engine/Source/Programs/Horde/HordeStorage/Horde.Storage/Implementation/TransactionLog/MemoryTransactionLogWriter.cs
@@ -0,0 +1,34 @@
+// Copyright Epic Games, Inc. All Rights Reserved.
+
+using System.Collections.Concurrent;
+using System.Threading.Tasks;
+using Jupiter.Implementation;
+
+namespace Horde.Storage.Implementation
+{
+    internal class MemoryTransactionLogWriter : ITransactionLogWriter
+    {
+        private readonly ConcurrentDictionary<NamespaceId, long> _transactionIds = new ConcurrentDictionary<NamespaceId, long>();
+
+        public Task<long> Add(NamespaceId ns, TransactionEvent @event)
+        {
+            // since we are never reading back the value of the transaction log in Europa we can just fake it by adding incrementing a id
+
+            return DoInsert(ns);
+        }
+
+        private async Task<long> DoInsert(NamespaceId ns)
+        {
+            return await Task.Run(() =>
+            {
+                return _transactionIds.AddOrUpdate(ns, _ => 0,
+                    (_, transactionId) => transactionId + 1);
+            });
+        }
+
+        public Task<long> Delete(NamespaceId ns, BucketId bucket, KeyId key)
+        {
+            return DoInsert(ns); 
+        }
+    }
+}
diff --git a/Engine/Source/Programs/Horde/HordeStorage/Horde.Storage/Implementation/TransactionLog/ReplicationLogSnapshot.cs b/Engine/Source/Programs/Horde/HordeStorage/Horde.Storage/Implementation/TransactionLog/ReplicationLogSnapshot.cs
new file mode 100644
--- /dev/null
+++ b/Engine/Source/Programs/Horde/HordeStorage/Horde.Storage/Implementation/TransactionLog/ReplicationLogSnapshot.cs
@@ -0,0 +1,146 @@
+﻿// Copyright Epic Games, Inc. All Rights Reserved.
+
+using System;
+using System.Collections.Generic;
+using System.IO;
+using System.IO.Compression;
+using System.Linq;
+using System.Threading.Tasks;
+using Jupiter.Implementation;
+using Newtonsoft.Json;
+
+namespace Horde.Storage.Implementation.TransactionLog
+{
+    public class ReplicationLogSnapshot
+    {
+        private readonly NamespaceId _namespace;
+        private readonly Dictionary<string, SnapshotLiveObject> _liveObjects = new Dictionary<string, SnapshotLiveObject>();
+
+        public class ReplicationLogSnapshotState
+        {
+            public ReplicationLogSnapshotState(SnapshotHeader header, List<SnapshotLiveObject> liveObjects)
+            {
+                Header = header;
+                LiveObjects = liveObjects;
+            }
+
+            public SnapshotHeader Header { get; set; }
+            public List<SnapshotLiveObject> LiveObjects { get; set; }
+        }
+
+        public class SnapshotHeader
+        {
+            public NamespaceId Namespace { get; set; }
+            public string? LastBucket { get; set; }
+            public Guid LastEvent { get; set; }
+        }
+
+        public class SnapshotLiveObject
+        {
+            public SnapshotLiveObject(BucketId bucket, KeyId key, BlobIdentifier blob)
+            {
+                Bucket = bucket;
+                Key = key;
+                Blob = blob;
+            }
+
+            public BucketId Bucket { get; set; }
+            public KeyId Key { get; set; }
+            public BlobIdentifier Blob { get; set; }
+        }
+
+        public ReplicationLogSnapshot(NamespaceId ns)
+        {
+            _namespace = ns;
+        }
+
+        private ReplicationLogSnapshot(ReplicationLogSnapshotState snapshotState)
+        {
+            _namespace = snapshotState.Header.Namespace;
+            LastEvent = snapshotState.Header.LastEvent;
+            LastBucket = snapshotState.Header.LastBucket;
+
+            _liveObjects = snapshotState.LiveObjects.ToDictionary(o => BuildObjectKey(o.Bucket, o.Key));
+        }
+
+        public Guid? LastEvent { get; set; }
+        public string? LastBucket { get; set; }
+
+        public IEnumerable<SnapshotLiveObject> LiveObjects
+        {
+            get { return _liveObjects.Values; }
+        }
+
+        public static async Task<ReplicationLogSnapshot> DeserializeSnapshot(Stream stream)
+        {
+            await using GZipStream gzipStream = new GZipStream(stream, CompressionMode.Decompress);
+            using TextReader textReader = new StreamReader(gzipStream);
+            using JsonReader reader = new JsonTextReader(textReader);
+            JsonSerializer serializer = JsonSerializer.Create();
+            ReplicationLogSnapshotState? snapshotState = serializer.Deserialize<ReplicationLogSnapshotState>(reader);
+            if (snapshotState == null)
+            {
+                throw new NotImplementedException();
+            }
+            return new ReplicationLogSnapshot(snapshotState);
+
+        }
+
+        public async Task Serialize(Stream stream)
+        {
+            // apply gzip on the stream as we have a lot of redundant text from bucket, but also within the keys
+            await using GZipStream gzipStream = new GZipStream(stream, CompressionMode.Compress);
+            await using TextWriter textWriter = new StreamWriter(gzipStream);
+            JsonSerializer serializer = JsonSerializer.Create();
+
+            List<SnapshotLiveObject> liveObjects = LiveObjects.ToList();
+            if (!liveObjects.Any())
+                throw new Exception("You must have at least one live object to build a snapshot");
+            if (LastBucket == null)
+                throw new Exception("No last bucket found when serializing state, did you really have events?");
+            if (LastEvent == null)
+                throw new Exception("No last event found when serializing state, did you really have events?");
+
+            ReplicationLogSnapshotState state = new ReplicationLogSnapshotState(new SnapshotHeader()
+            {
+                Namespace = _namespace,
+                LastBucket = LastBucket,
+                LastEvent = LastEvent.Value,
+            }, liveObjects);
+            serializer.Serialize(textWriter, state);
+        }
+
+        public void ProcessEvent(ReplicationLogEvent entry)
+        {
+            LastBucket = entry.TimeBucket;
+            LastEvent = entry.EventId;
+
+            switch (entry.Op)
+            {
+                case ReplicationLogEvent.OpType.Added:
+                    ProcessAddEvent(entry.Bucket, entry.Key, entry.Blob);
+                    break;
+                case ReplicationLogEvent.OpType.Deleted:
+                    ProcessDeleteEvent(entry.Bucket, entry.Key);
+                    break;
+                default:
+                    throw new NotImplementedException();
+            }
+        }
+
+
+        private string BuildObjectKey(BucketId bucket, KeyId key)
+        {
+            return $"{bucket}.{key}";
+        }
+
+        private void ProcessAddEvent(BucketId bucket, KeyId key, BlobIdentifier blob)
+        {
+            _liveObjects.TryAdd(BuildObjectKey(bucket, key), new SnapshotLiveObject(bucket, key, blob));
+        }
+        private void ProcessDeleteEvent(BucketId bucket, KeyId key)
+        {
+            _liveObjects.Remove(BuildObjectKey(bucket, key));
+        }
+    }
+}
diff --git a/Engine/Source/Programs/Horde/HordeStorage/Horde.Storage/Implementation/TransactionLog/ReplicationLogSnapshotBuilder.cs b/Engine/Source/Programs/Horde/HordeStorage/Horde.Storage/Implementation/TransactionLog/ReplicationLogSnapshotBuilder.cs
new file mode 100644
--- /dev/null
+++ b/Engine/Source/Programs/Horde/HordeStorage/Horde.Storage/Implementation/TransactionLog/ReplicationLogSnapshotBuilder.cs
@@ -0,0 +1,92 @@
+﻿using System;
+using System.IO;
+using System.Threading;
+using System.Threading.Tasks;
+using Jupiter.Implementation;
+
+namespace Horde.Storage.Implementation.TransactionLog
+{
+    public class ReplicationLogSnapshotBuilder
+    {
+        private readonly IReplicationLog _replicationLog;
+        private readonly IBlobStore _blobStore;
+        private readonly IReferencesStore _referencesStore;
+
+        public ReplicationLogSnapshotBuilder(IReplicationLog replicationLog, IBlobStore blobStore, IReferencesStore referencesStore)
+        {
+            _replicationLog = replicationLog;
+            _blobStore = blobStore;
+            _referencesStore = referencesStore;
+        }
+
+        public async Task<BlobIdentifier> BuildSnapshot(NamespaceId ns, NamespaceId storeInNamespace, CancellationToken cancellationToken = default(CancellationToken))
+        {
+            // builds a snapshot and commits it to the blob store with the identifier specified
+
+            SnapshotInfo? snapshotInfo = await _replicationLog.GetLatestSnapshot(ns);
+
+            if (cancellationToken.IsCancellationRequested)
+                throw new TaskCanceledException();
+
+            ReplicationLogSnapshot snapshot;
+            string? lastBucket;
+            Guid? lastEvent;
+            if (snapshotInfo != null)
+            {
+                // append to the previous snapshot if one is available
+                await using BlobContents blobContents = await _blobStore.GetObject(snapshotInfo.BlobNamespace, snapshotInfo.SnapshotBlob);
+                if (cancellationToken.IsCancellationRequested)
+                    throw new TaskCanceledException();
+                snapshot = await ReplicationLogSnapshot.DeserializeSnapshot(blobContents.Stream);
+                lastBucket = snapshot.LastBucket;
+                lastEvent = snapshot.LastEvent;
+            }
+            else
+            {
+                snapshot = new ReplicationLogSnapshot(ns);
+
+                lastBucket = null;
+                lastEvent = null;
+            }
+
+            if (cancellationToken.IsCancellationRequested)
+                throw new TaskCanceledException();
+
+            await foreach (ReplicationLogEvent entry in _replicationLog.Get(ns, lastBucket, lastEvent))
+            {
+                if (cancellationToken.IsCancellationRequested)
+                    throw new TaskCanceledException();
+
+                snapshot.ProcessEvent(entry);
+            }
+
+            {
+                await using MemoryStream ms = new MemoryStream();
+                await snapshot.Serialize(ms);
+                byte[] buf = ms.GetBuffer();
+                BlobIdentifier blobIdentifier = BlobIdentifier.FromBlob(buf);
+
+
+                CompactBinaryWriter writer = new CompactBinaryWriter();
+                writer.BeginObject();
+                writer.AddBinaryAttachment(blobIdentifier, "snapshotBlob");
+                writer.AddDateTime(DateTime.Now, "timestamp");
+                writer.EndObject();
+
+                byte[] cbObjectBytes = writer.Save();
+                BlobIdentifier cbBlobId = BlobIdentifier.FromBlob(cbObjectBytes);
+                await _referencesStore.Put(storeInNamespace, new BucketId("snapshot"), new KeyId(blobIdentifier.ToString()), cbBlobId, cbObjectBytes, true);
+                await _blobStore.PutObject(storeInNamespace, buf, blobIdentifier);
+
+                if (cancellationToken.IsCancellationRequested)
+                    throw new TaskCanceledException();
+
+                // update the replication log with the new snapshot
+                await _replicationLog.AddSnapshot(new SnapshotInfo(ns, storeInNamespace, blobIdentifier));
+
+                return blobIdentifier;
+
+            }
+        }
+    }
+}
diff --git a/Engine/Source/Programs/Horde/HordeStorage/Horde.Storage/Implementation/TransactionLog/ReplicationSnapshotService.cs b/Engine/Source/Programs/Horde/HordeStorage/Horde.Storage/Implementation/TransactionLog/ReplicationSnapshotService.cs
new file mode 100644
--- /dev/null
+++ b/Engine/Source/Programs/Horde/HordeStorage/Horde.Storage/Implementation/TransactionLog/ReplicationSnapshotService.cs
@@ -0,0 +1,99 @@
+﻿// Copyright Epic Games, Inc. All Rights Reserved.
+
+using System;
+using System.Collections.Generic;
+using System.ComponentModel.DataAnnotations;
+using System.Threading;
+using System.Threading.Tasks;
+using Dasync.Collections;
+using Horde.Storage.Implementation.TransactionLog;
+using Jupiter;
+using Jupiter.Implementation;
+using Microsoft.Extensions.DependencyInjection;
+using Microsoft.Extensions.Options;
+using Serilog;
+
+namespace Horde.Storage.Implementation
+{
+    // ReSharper disable once ClassNeverInstantiated.Global
+    public class ReplicationSnapshotService : PollingService<ReplicationSnapshotService.SnapshotState>
+    {
+        public class SnapshotState
+        {
+        }
+
+        private readonly IServiceProvider _provider;
+        private readonly IOptionsMonitor<SnapshotSettings> _settings;
+        private readonly IReplicationLog _replicationLog;
+        private readonly ILeaderElection _leaderElection;
+        private readonly ILogger _logger = Log.ForContext<ReplicationSnapshotService>();
+        private readonly CancellationTokenSource _cancellationTokenSource = new CancellationTokenSource();
+        private Task? _snapshotBuildTask = null;
+
+        public override bool ShouldStartPolling()
+        {
+            return _settings.CurrentValue.Enabled;
+        }
+
+        public ReplicationSnapshotService(IServiceProvider provider, IOptionsMonitor<SnapshotSettings> settings, IReplicationLog replicationLog, ILeaderElection leaderElection) :
+            base(serviceName: nameof(ReplicationSnapshotService), TimeSpan.FromSeconds(settings.CurrentValue.SnapshotFrequencySeconds), new SnapshotState())
+        {
+            _provider = provider;
+            _settings = settings;
+            _replicationLog = replicationLog;
+            _leaderElection = leaderElection;
+        }
+
+        public override async Task<bool> OnPoll(SnapshotState state, CancellationToken cancellationToken)
+        {
+            if (!_settings.CurrentValue.Enabled)
+            {
+                _logger.Information("Skipped running replication snapshot service as it is disabled");
+                return false;
+            }
+
+            if (!_leaderElection.IsThisInstanceLeader())
+            {
+                _logger.Information("Skipped running snapshot service because this instance was not the leader");
+                return false;
+            }
+
+            bool ran = false;
+            _snapshotBuildTask = _replicationLog.GetNamespaces().ParallelForEachAsync(async ns =>
+            {
+                ReplicationLogSnapshotBuilder builder = ActivatorUtilities.CreateInstance<ReplicationLogSnapshotBuilder>(_provider);
+                BlobIdentifier snapshotBlob = await builder.BuildSnapshot(ns, _settings.CurrentValue.SnapshotStorageNamespace, _cancellationTokenSource.Token);
+                _logger.Information("Snapshot built for {Namespace} with id {Id}", ns, snapshotBlob);
+
+                ran = true;
+            }, _cancellationTokenSource.Token);
+            await _snapshotBuildTask;
+            _snapshotBuildTask = null;
+            return ran;
+        }
+
+        protected override async Task OnStopping(SnapshotState state)
+        {
+            _cancellationTokenSource.Cancel();
+
+            if (_snapshotBuildTask != null)
+                await _snapshotBuildTask;
+        }
+    }
+
+    public class SnapshotSettings
+    {
+        /// <summary>
+        /// Enable to start a replicating another Jupiter instance into this one
+        /// </summary>
+        public bool Enabled { get; set; } = false;
+
+        /// <summary>
+        /// The frequency at which to poll for new replication events
+        /// </summary>
+        [Required]
+        public int SnapshotFrequencySeconds { get; set; } = 3600; // default to once a day
+
+        public NamespaceId SnapshotStorageNamespace { get; set; } = new NamespaceId("jupiter-internal");
+    }
+}
