Moved Horde.Storage from git to p4

diff --git a/Engine/Source/Programs/Horde/HordeStorage/Horde.Storage/Implementation/TransactionLog/ScyllaReplicationLog.cs b/Engine/Source/Programs/Horde/HordeStorage/Horde.Storage/Implementation/TransactionLog/ScyllaReplicationLog.cs
new file mode 100644
--- /dev/null
+++ b/Engine/Source/Programs/Horde/HordeStorage/Horde.Storage/Implementation/TransactionLog/ScyllaReplicationLog.cs
@@ -0,0 +1,381 @@
+// Copyright Epic Games, Inc. All Rights Reserved.
+
+using System;
+using System.Collections.Generic;
+using System.Linq;
+using System.Threading.Tasks;
+using Cassandra;
+using Cassandra.Mapping;
+using Dasync.Collections;
+using Datadog.Trace;
+using Jupiter.Implementation;
+using Microsoft.Extensions.Options;
+
+namespace Horde.Storage.Implementation
+{
+    class ScyllaReplicationLog : IReplicationLog
+    {
+        private readonly IOptionsMonitor<ScyllaSettings> _settings;
+        private readonly ISession _session;
+        private readonly Mapper _mapper;
+
+        public ScyllaReplicationLog(IScyllaSessionManager scyllaSessionManager, IOptionsMonitor<ScyllaSettings> settings)
+        {
+            _settings = settings;
+            _session = scyllaSessionManager.GetSessionForLocalKeyspace();
+            _mapper = new Mapper(_session);
+
+            _session.Execute(new SimpleStatement(@"CREATE TABLE IF NOT EXISTS replication_log (
+                namespace varchar,
+                replication_bucket bigint,
+                replication_id timeuuid,
+                bucket varchar, 
+                key varchar, 
+                type int,
+                object_identifier blob_identifier,
+                PRIMARY KEY ((namespace, replication_bucket), replication_id)
+            );"
+            ));
+
+            _session.Execute(new SimpleStatement(@"CREATE TABLE IF NOT EXISTS replication_snapshot (
+                namespace varchar,
+                id timeuuid,
+                blob_snapshot blob_identifier,
+                blob_namespace varchar,
+                PRIMARY KEY ((namespace), id)
+            ) WITH CLUSTERING ORDER BY (id DESC);"  
+            ));
+
+            _session.Execute(new SimpleStatement(@"CREATE TABLE IF NOT EXISTS replication_namespace (
+                namespace varchar,
+                PRIMARY KEY ((namespace))
+            );"  
+            ));
+        }
+
+        public async IAsyncEnumerable<NamespaceId> GetNamespaces()
+        {
+            IEnumerable<ScyllaNamespace> namespaces = await _mapper.FetchAsync<ScyllaNamespace>("");
+
+            foreach (ScyllaNamespace scyllaNamespace in namespaces)
+            {
+                yield return new NamespaceId(scyllaNamespace.Namespace);
+            }
+        }
+
+        public async Task<(string, Guid)> InsertAddEvent(NamespaceId ns, BucketId bucket, KeyId key, BlobIdentifier objectBlob, DateTime? timestamp)
+        {
+            using Scope _ = Tracer.Instance.StartActive("scylla.insert_add_event");
+            Task addNamespaceTask = PotentiallyAddNamespace(ns);
+            DateTime timeBucket = timestamp.GetValueOrDefault(DateTime.Now);
+            ScyllaReplicationLogEvent log = new ScyllaReplicationLogEvent(ns.ToString(), bucket.ToString(), key.ToString(), timeBucket, ScyllaReplicationLogEvent.OpType.Added, objectBlob);
+            await _mapper.InsertAsync<ScyllaReplicationLogEvent>(log, insertNulls: false,  ttl: (int)_settings.CurrentValue.ReplicationLogTimeToLive.TotalSeconds);
+
+            await addNamespaceTask;
+            return (log.GetReplicationBucketIdentifier(), log.ReplicationId);
+        }
+
+        public async Task<(string, Guid)> InsertDeleteEvent(NamespaceId ns, BucketId bucket, KeyId key, BlobIdentifier objectBlob, DateTime? timestamp)
+        {
+            using Scope _ = Tracer.Instance.StartActive("scylla.insert_delete_event");
+            Task addNamespaceTask = PotentiallyAddNamespace(ns);
+            DateTime timeBucket = timestamp.GetValueOrDefault(DateTime.Now);
+            ScyllaReplicationLogEvent log = new ScyllaReplicationLogEvent(ns.ToString(), bucket.ToString(), key.ToString(), timeBucket, ScyllaReplicationLogEvent.OpType.Deleted, objectBlob);
+            await _mapper.InsertAsync<ScyllaReplicationLogEvent>(log, insertNulls: false,  ttl: (int)_settings.CurrentValue.ReplicationLogTimeToLive.TotalSeconds);
+
+            await addNamespaceTask;
+            return (log.GetReplicationBucketIdentifier(), log.ReplicationId);
+        }
+
+        private async Task PotentiallyAddNamespace(NamespaceId ns)
+        {
+            await _mapper.InsertAsync(new ScyllaNamespace(ns.ToString()));
+        }
+
+        public async IAsyncEnumerable<ReplicationLogEvent> Get(NamespaceId ns, string? lastBucket, Guid? lastEvent)
+        {
+            using Scope getReplicationLogScope = Tracer.Instance.StartActive("scylla.get_replication_log");
+            IAsyncEnumerable<long> buckets = FindReplicationBuckets(ns, lastBucket);
+
+            // loop thru the buckets starting with the oldest to try and find were lastBucket refers to
+            bool bucketFound = false;
+            await foreach (long bucketField in buckets)
+            {
+                using Scope readReplicationScope = Tracer.Instance.StartActive("scylla.read_replication_bucket");
+                DateTime t = DateTime.FromFileTimeUtc(bucketField);
+                string bucket = t.ToReplicationBucketIdentifier();
+                readReplicationScope.Span.ResourceName = bucket;
+
+                if (lastBucket != null && bucket != lastBucket)
+                    continue;
+
+                // at least one bucket was found
+                bucketFound = true;
+                // we have found the bucket to resume from, we should no do any more filtering
+                lastBucket = null;
+
+                IEnumerable<ScyllaReplicationLogEvent> events = await _mapper.FetchAsync<ScyllaReplicationLogEvent>("WHERE namespace = ? AND replication_bucket = ?", ns.ToString(), bucketField);
+
+                bool skipEvents = lastEvent.HasValue;
+                foreach (ScyllaReplicationLogEvent scyllaReplicationLog in events)
+                {
+                    // check if we are resuming from a previous event, if we are omit every event up to it
+                    if (lastEvent.HasValue && scyllaReplicationLog.ReplicationId == lastEvent)
+                    {
+                        skipEvents = false;
+                        continue; // start reading as of the next event
+                    }
+
+                    if (skipEvents)
+                        continue;
+
+                    yield return new ReplicationLogEvent(
+                        new NamespaceId(scyllaReplicationLog.Namespace),
+                        new BucketId(scyllaReplicationLog.Bucket!),
+                        new KeyId(scyllaReplicationLog.Key!),
+                        scyllaReplicationLog.ObjectIdentifier!.AsBlobIdentifier(),
+                        scyllaReplicationLog.ReplicationId,
+                        scyllaReplicationLog.GetReplicationBucketIdentifier(),
+                        scyllaReplicationLog.GetReplicationBucketTimestamp(),
+                        (ReplicationLogEvent.OpType)scyllaReplicationLog.Type);
+                }
+
+                // when continuing to the next bucket we start from the beginning of it
+                lastEvent = null;
+            }
+
+            if (!bucketFound)
+            {
+                ScyllaNamespace? namespaces = await _mapper.SingleOrDefaultAsync<ScyllaNamespace>("WHERE namespace = ?", ns.ToString());
+                if (namespaces == null)
+                {
+                    throw new NamespaceNotFoundException(ns);
+                }
+                
+                throw new IncrementalLogNotAvailableException();
+            }
+        }
+
+        private async IAsyncEnumerable<long> FindReplicationBuckets(NamespaceId ns, string? lastBucket)
+        {
+            using Scope findReplicationBucketScope = Tracer.Instance.StartActive("scylla.find_replication_buckets");
+
+            DateTime startBucketTime;
+            if (lastBucket != null)
+            {
+                long bucket = FromReplicationBucketIdentifier(lastBucket);
+                startBucketTime = DateTime.FromFileTimeUtc(bucket);
+                yield return bucket;
+            }
+            else
+            {
+                using Scope _ = Tracer.Instance.StartActive("scylla.determine_first_replication_bucket");
+                SortedSet<long> buckets = new SortedSet<long>();
+
+                // fetch all the buckets that exists and sort them based on time
+                RowSet replicationBuckets = await _session.ExecuteAsync(new SimpleStatement("SELECT replication_bucket FROM replication_log WHERE namespace = ? PER PARTITION LIMIT 1 ALLOW FILTERING", ns.ToString()));
+                foreach (Row replicationBucket in replicationBuckets)
+                {
+                    long bucketField = (long)replicationBucket["replication_bucket"];
+                    buckets.Add(bucketField);
+                }
+
+                // if there are no buckets we stop here
+                if (buckets.Count == 0)
+                    yield break;
+
+                // pick the first bucket
+                long bucket = buckets.First();
+                startBucketTime = DateTime.FromFileTimeUtc(bucket);
+                yield return bucket;
+            }
+
+            // ignore any bucket that is older then a cutoff, as that can cause us to end up scanning thru a lot of hours that will never exist (incremental logs are deleted after 7 days)
+            DateTime oldCutoff = DateTime.Now.AddDays(-14);
+            // we returned the start bucket earlier so now we start with the next one
+            DateTime bucketTime = startBucketTime.AddHours(1.0).ToHourlyBucket();
+            while(bucketTime < DateTime.Now && bucketTime > oldCutoff)
+            {
+                using Scope _ = Tracer.Instance.StartActive("scylla.determine_replication_bucket_exists");
+                // fetch all the buckets that exists and sort them based on time
+                IEnumerable<ScyllaReplicationLogEvent> logEvent = await _mapper.FetchAsync<ScyllaReplicationLogEvent>("WHERE namespace = ? AND replication_bucket = ? LIMIT 1", ns.ToString(), bucketTime.ToFileTimeUtc());
+                ScyllaReplicationLogEvent? e = logEvent.FirstOrDefault();
+                if (e != null)
+                {
+                    yield return e.ReplicationBucket;
+                }
+
+                bucketTime = bucketTime.AddHours(1.0);
+            }
+
+        }
+
+        private long FromReplicationBucketIdentifier(string bucket)
+        {
+            if (!bucket.StartsWith("rep-"))
+                throw new ArgumentException($"Invalid bucket identifier: \"{bucket}\"", nameof(bucket));
+            string timestamp = bucket.Substring(bucket.IndexOf("-") + 1);
+            long filetime = long.Parse(timestamp);
+            return filetime;
+        }
+
+        public async Task AddSnapshot(SnapshotInfo snapshotHeader)
+        {
+            await _mapper.InsertAsync<ScyllaSnapshot>(new ScyllaSnapshot(snapshotHeader.SnapshottedNamespace.ToString(), snapshotHeader.BlobNamespace.ToString(), TimeUuid.NewId(), snapshotHeader.SnapshotBlob));
+            await CleanupSnapshots(snapshotHeader.SnapshottedNamespace);
+        }
+
+        public async Task CleanupSnapshots(NamespaceId ns)
+        {
+            // determine if we have to many snapshots and remove the oldest ones if we do
+            RowSet rowSet = await _session.ExecuteAsync(new SimpleStatement("SELECT Count(*) FROM replication_snapshot WHERE namespace = ?", ns.ToString()));
+            Row row = rowSet.First();
+            int countOfSnapshots = (int)(long)row[0];
+            if (countOfSnapshots > _settings.CurrentValue.MaxSnapshotsPerNamespace)
+            {
+                int deleteCount = countOfSnapshots - _settings.CurrentValue.MaxSnapshotsPerNamespace;
+
+                List<ScyllaSnapshot> snapshots = (await _mapper.FetchAsync<ScyllaSnapshot>("WHERE namespace = ?", ns.ToString())).ToList();
+                for (int i = 0; i < deleteCount; i++)
+                {
+                    // since snapshots are sorted newest first we delete from the end
+                    ScyllaSnapshot snapshotToDelete = snapshots[^(i + 1)];
+                    await _mapper.DeleteAsync<ScyllaSnapshot>(snapshotToDelete);
+                }
+            }
+        }
+
+        public async Task<SnapshotInfo?> GetLatestSnapshot(NamespaceId ns)
+        {
+            SnapshotInfo? s = await GetSnapshots(ns).FirstOrDefaultAsync();
+
+            return s;
+        }
+
+        public async IAsyncEnumerable<SnapshotInfo> GetSnapshots(NamespaceId ns)
+        {
+            IEnumerable<ScyllaSnapshot> snapshots = await _mapper.FetchAsync<ScyllaSnapshot>("WHERE namespace = ?", ns.ToString());
+
+            foreach (ScyllaSnapshot snapshot in snapshots)
+            {
+                yield return new SnapshotInfo(ns, new NamespaceId(snapshot.BlobNamespace), snapshot.BlobSnapshot.AsBlobIdentifier());
+            }
+        }
+    }
+
+    
+    [Cassandra.Mapping.Attributes.Table("replication_log")]
+    class ScyllaReplicationLogEvent
+    {
+        public enum OpType
+        {
+            Added,
+            Deleted
+        };
+
+        public ScyllaReplicationLogEvent()
+        {
+            Namespace = null!;
+            Bucket = null!;
+            Key = null!;
+        }
+
+        public ScyllaReplicationLogEvent(string @namespace, string bucket, string key, DateTime lastTimestamp, OpType opType, BlobIdentifier objectIdentifier)
+        {
+            Namespace = @namespace;
+            Bucket = bucket;
+            Key = key;
+            ReplicationBucket = lastTimestamp.ToHourlyBucket().ToFileTimeUtc();
+            ReplicationId = TimeUuid.NewId(lastTimestamp);
+            Type = (int)opType;
+            ObjectIdentifier = new ScyllaBlobIdentifier(objectIdentifier);
+        }
+
+        [Cassandra.Mapping.Attributes.PartitionKey]
+        public string Namespace { get;set; }
+
+        // we store bucket as a long (datetime converted to filetime bucketed per hour) to make sure it sorts oldest first
+        [Cassandra.Mapping.Attributes.PartitionKey]
+        [Cassandra.Mapping.Attributes.Column("replication_bucket")]
+        public long ReplicationBucket { get; set; }
+
+        [Cassandra.Mapping.Attributes.Column("replication_id")]
+        public TimeUuid ReplicationId { get; set; }
+
+        [Cassandra.Mapping.Attributes.Column("bucket")]
+        public string Bucket { get; set; }
+
+        [Cassandra.Mapping.Attributes.Column("key")]
+        public string Key { get; set; }
+
+        [Cassandra.Mapping.Attributes.Column("type")]
+        public int Type { get; set; }
+
+        [Cassandra.Mapping.Attributes.Column("object_identifier")]
+        public ScyllaBlobIdentifier? ObjectIdentifier { get; set; }
+
+        public string GetReplicationBucketIdentifier()
+        {
+            // the replication bucket identifier is a a string to avoid people assuming its a timestamp, we do not store that in the db though as the string does not sort correctly then
+            return DateTime.FromFileTimeUtc(ReplicationBucket).ToReplicationBucketIdentifier();
+        }
+
+        public DateTime GetReplicationBucketTimestamp()
+        {
+            return DateTime.FromFileTimeUtc(ReplicationBucket);
+        }
+    }
+
+
+    
+    [Cassandra.Mapping.Attributes.Table("replication_snapshot")]
+    class ScyllaSnapshot
+    {
+        public ScyllaSnapshot()
+        {
+            Namespace = null!;
+            BlobSnapshot = null!;
+            BlobNamespace = null!;
+        }
+
+        public ScyllaSnapshot(string @namespace, string blobNamespace, TimeUuid id, BlobIdentifier objectIdentifier)
+        {
+            Namespace = @namespace;
+            BlobNamespace = blobNamespace;
+            Id = id;
+            BlobSnapshot = new ScyllaBlobIdentifier(objectIdentifier);
+        }
+
+        [Cassandra.Mapping.Attributes.PartitionKey]
+        public string Namespace { get;set; }
+
+
+        [Cassandra.Mapping.Attributes.Column("id")]
+        [Cassandra.Mapping.Attributes.ClusteringKey]
+        public TimeUuid Id { get; set; }
+
+        [Cassandra.Mapping.Attributes.Column("blob_snapshot")]
+        public ScyllaBlobIdentifier BlobSnapshot { get; set; }
+
+        [Cassandra.Mapping.Attributes.Column("blob_namespace")]
+        public string BlobNamespace { get; set; }
+
+    }
+
+    [Cassandra.Mapping.Attributes.Table("replication_namespace")]
+    class ScyllaNamespace
+    {
+        public ScyllaNamespace()
+        {
+            Namespace = null!;
+        }
+
+        public ScyllaNamespace(string @namespace)
+        {
+            Namespace = @namespace;
+        }
+
+        [Cassandra.Mapping.Attributes.PartitionKey]
+        public string Namespace { get;set; }
+    }
+}
diff --git a/Engine/Source/Programs/Horde/HordeStorage/Horde.Storage/Implementation/VersionFile.cs b/Engine/Source/Programs/Horde/HordeStorage/Horde.Storage/Implementation/VersionFile.cs
new file mode 100644
--- /dev/null
+++ b/Engine/Source/Programs/Horde/HordeStorage/Horde.Storage/Implementation/VersionFile.cs
@@ -0,0 +1,40 @@
+﻿// Copyright Epic Games, Inc. All Rights Reserved.
+
+using System.IO;
+using YamlDotNet.Serialization;
+using YamlDotNet.Serialization.NamingConventions;
+
+namespace Horde.Storage.Implementation
+{
+    public interface IVersionFile
+    {
+        public string? VersionString { get; }
+    }
+
+    public class VersionFile : IVersionFile
+    {
+        private readonly VersionFileContents _versionFileContents;
+
+        public class VersionFileContents
+        {
+            public string? Version { get; set; }
+        }
+
+        public VersionFile()
+        {
+            FileInfo assemblyFile = new FileInfo(typeof(VersionFile).Assembly.Location);
+            FileInfo fi = new FileInfo( Path.Combine(assemblyFile!.Directory!.FullName, "version.yaml"));
+            IDeserializer deserializer = new DeserializerBuilder()
+                .WithNamingConvention(CamelCaseNamingConvention.Instance)
+                .Build();
+
+            using StreamReader tr = fi.OpenText();
+            _versionFileContents = deserializer.Deserialize<VersionFileContents>(tr);
+        }
+
+        public string? VersionString
+        {
+            get { return _versionFileContents.Version; }
+        }
+    }
+}
diff --git a/Engine/Source/Programs/Horde/HordeStorage/Horde.Storage/Program.cs b/Engine/Source/Programs/Horde/HordeStorage/Horde.Storage/Program.cs
new file mode 100644
--- /dev/null
+++ b/Engine/Source/Programs/Horde/HordeStorage/Horde.Storage/Program.cs
@@ -0,0 +1,24 @@
+// Copyright Epic Games, Inc. All Rights Reserved.
+
+using System;
+using System.IO.Pipelines;
+using System.Net.Sockets;
+using Horde.Storage.Implementation.Kestrel;
+using Jupiter;
+
+namespace Horde.Storage
+{
+    // ReSharper disable once ClassNeverInstantiated.Global
+    public class Program
+    {
+        public static int Main(string[] args)
+        {
+            return BaseProgram<HordeStorageStartup>.BaseMain(args, HttpConnectionFactory);
+        }
+
+        private static BaseHttpConnection HttpConnectionFactory(IServiceProvider sp, PipeReader reader, PipeWriter writer, Socket socket)
+        {
+            return new KestrelDispatcher(sp, reader, writer, socket);
+        }
+    }
+}
diff --git a/Engine/Source/Programs/Horde/HordeStorage/Horde.Storage/Properties/launchSettings.json b/Engine/Source/Programs/Horde/HordeStorage/Horde.Storage/Properties/launchSettings.json
new file mode 100644
--- /dev/null
+++ b/Engine/Source/Programs/Horde/HordeStorage/Horde.Storage/Properties/launchSettings.json
@@ -0,0 +1,43 @@
+{
+  "iisSettings": {
+    "windowsAuthentication": false,
+    "anonymousAuthentication": true,
+    "iisExpress": {
+      "applicationUrl": "http://localhost:57001",
+      "sslPort": 44351
+    }
+  },
+  "$schema": "http://json.schemastore.org/launchsettings.json",
+  "profiles": {
+    "IIS Express": {
+      "commandName": "IISExpress",
+      "launchBrowser": true,
+      "launchUrl": "api/v1/c/ddc/test.ddc/default/foo.json",
+      "environmentVariables": {
+        "ASPNETCORE_ENVIRONMENT": "Development",
+        "DD_TRACE_ENABLED": "false"
+      }
+    },
+    "Local Deployment": {
+      "commandName": "IISExpress",
+      "launchBrowser": true,
+      "launchUrl": "api/v1/c/ddc/ue4.ddc/terrain/foo.json",
+      "environmentVariables": {
+        "CORECLR_ENABLE_PROFILING": "0",
+        "Callisto__ConnectionString": "http://localhost:57002",
+        "Horde_Storage__StorageImplementation": "Memory",
+        "Horde_Storage__TransactionLogWriterImplementation": "Callisto",
+        "ASPNETCORE_ENVIRONMENT": "Development"
+      }
+    },
+    "Kestrel": {
+      "commandName": "Project",
+      "launchUrl": "/",
+      "environmentVariables": {
+        "ASPNETCORE_ENVIRONMENT": "Development",
+        "DD_TRACE_ENABLED": "false"
+      },
+      "applicationUrl": ""
+    }
+  }
+}
diff --git a/Engine/Source/Programs/Horde/HordeStorage/Horde.Storage/Settings.cs b/Engine/Source/Programs/Horde/HordeStorage/Horde.Storage/Settings.cs
new file mode 100644
--- /dev/null
+++ b/Engine/Source/Programs/Horde/HordeStorage/Horde.Storage/Settings.cs
@@ -0,0 +1,226 @@
+﻿using System;
+using System.Collections.Generic;
+using System.ComponentModel.DataAnnotations;
+using System.Linq;
+using Microsoft.Extensions.Caching.Memory;
+
+namespace Horde.Storage
+{
+    
+    public class HordeStorageSettings
+    {
+        public enum RefDbImplementations
+        {
+            Mongo,
+            Cosmos,
+            DynamoDb
+
+        }
+
+        public enum TransactionLogWriterImplementations
+        {
+            Callisto,
+            Memory
+        }
+
+        public enum ReplicationLogWriterImplementations
+        {
+            Memory,
+            Scylla
+        }
+
+        public enum TreeStoreImplementations
+        {
+            Memory,
+            DynamoDb
+        }
+        public enum TreeRootStoreImplementations
+        {
+            Memory,
+            DynamoDb
+        }
+
+        public enum StorageBackendImplementations
+        {
+            S3,
+            Azure,
+            FileSystem,
+            Memory
+        }
+
+        public enum ReferencesDbImplementations
+        {
+            Memory,
+            Scylla
+        }
+
+        public enum ContentIdStoreImplementations
+        {
+            Memory,
+            Scylla
+        }
+
+        public enum LeaderElectionImplementations
+        {
+            Static, 
+            Kubernetes
+        }
+        
+        public class ValidStorageBackend : ValidationAttribute
+        {
+            public override string FormatErrorMessage(string name)
+            {
+                return "Need to specify at least one storage backend. Valid ones are: " +
+                       string.Join(", ", Enum.GetNames(typeof(StorageBackendImplementations)));
+            }
+
+            public override bool IsValid(object? value)
+            {
+                if (value == null) return true;
+                return value is IEnumerable<string> backends && backends.All(x => Enum.TryParse(typeof(StorageBackendImplementations), x, true, out _));
+            }
+        }
+
+        [ValidStorageBackend]
+        public string[]? StorageImplementations { get; set; }
+
+        [Required]
+        public TransactionLogWriterImplementations TransactionLogWriterImplementation { get; set; } = TransactionLogWriterImplementations.Memory;
+
+        [Required] public ReplicationLogWriterImplementations ReplicationLogWriterImplementation { get; set; } = ReplicationLogWriterImplementations.Memory;
+
+        [Required] public RefDbImplementations RefDbImplementation { get; set; } = RefDbImplementations.Mongo;
+
+        [Required]
+        public TreeStoreImplementations TreeStoreImplementation { get; set; } = TreeStoreImplementations.Memory;
+
+        [Required]
+        public TreeRootStoreImplementations TreeRootStoreImplementation { get; set; } = TreeRootStoreImplementations.Memory;
+
+        [Required]
+        public ReferencesDbImplementations ReferencesDbImplementation { get; set; } = ReferencesDbImplementations.Memory;
+
+        public LeaderElectionImplementations LeaderElectionImplementation { get; set; } = LeaderElectionImplementations.Static;
+        public ContentIdStoreImplementations ContentIdStoreImplementation { get; set; } = ContentIdStoreImplementations.Memory;
+
+        public int? MaxSingleBlobSize { get; set; } = null; // disable blob partitioning
+
+        public int LastAccessRollupFrequencySeconds { get; set; } = 900; // 15 minutes
+
+        public bool UseNewDDCEndpoints { get; set; } = false;
+    }
+
+    public class MongoSettings
+    {
+        [Required] public string ConnectionString { get; set; } = "";
+
+        public bool RequireTls12 { get; set; } = true;
+    }
+
+
+    public class DynamoDbSettings
+    {
+        [Required] public string ConnectionString { get; set; } = "";
+
+        public long ReadCapacityUnits { get; set; } = 100;
+        public long WriteCapacityUnits { get; set; } = 20;
+
+        public bool UseOndemandCapacityProvisioning { get; set; } = true;
+        
+        /// <summary>
+        /// Endpoint name for DynamoDB Accelerator (DAX). Acts as a cache in front of DynamoDB to speed up requests
+        /// Disabled when set to null.
+        /// </summary>
+        public string? DaxEndpoint { get; set; } = null;
+
+        /// <summary>
+        /// Enabling this will make Horde.Storage create missing tables on demand, works great for local tables but for global tables its easier to have terraform manage it
+        /// </summary>
+        public bool CreateTablesOnDemand { get; set; } = true;
+
+        public static (string, int) ParseDaxEndpointAsHostPort(string endpoint)
+        {
+            if (!endpoint.Contains(":"))
+            {
+                return (endpoint, 8111);
+            }
+
+            string host = endpoint.Split(":")[0];
+            int port = Convert.ToInt32(endpoint.Split(":")[1]);
+            return (host, port);
+        }
+    }
+
+
+    public class CosmosSettings
+    {
+        [Range(400, 10_000)] public int DefaultRU = 400;
+    }
+
+    public class CallistoTransactionLogSettings
+    {
+        [Required] public string ConnectionString { get; set; } = "";
+    }
+
+    public class MemoryCacheBlobSettings : MemoryCacheOptions
+    {
+        public bool Enabled { get; set; } = true;
+
+        public bool EnableSlidingExpiry { get; set; } = true;
+        public int SlidingExpirationMinutes { get; set; } = 60;
+    }
+
+    public class MemoryCacheTreeSettings : MemoryCacheOptions
+    {
+        public bool Enabled { get; set; } = true;
+        public bool EnableSlidingExpiry { get; set; } = true;
+        public int SlidingExpirationMinutes { get; set; } = 60;
+    }
+
+    public class MemoryCacheRefSettings : MemoryCacheOptions
+    {
+        public bool Enabled { get; set; } = true;
+        public bool EnableSlidingExpiry { get; set; } = true;
+        public int SlidingExpirationMinutes { get; set; } = 60;
+    }
+
+
+    public class AzureSettings
+    {
+        [Required] public string ConnectionString { get; set; } = "";
+    }
+
+    public class FilesystemSettings
+    {
+        [Required] public string RootDir { get; set; } = "";
+        public ulong MaxSizeBytes { get; set; } = 500 * 1024 * 1024;
+        public double TriggerThresholdPercentage { get; set; } = 0.95;
+        public double TargetThresholdPercentage { get; set; } = 0.85;
+    }
+
+    public class S3Settings
+    {
+        [Required] public string ConnectionString  { get; set; } = "";
+
+        [Required] public string BucketName { get; set; } = "";
+
+        public bool ForceAWSPathStyle { get; set; }
+        public bool CreateBucketIfMissing { get;set; } = true;
+
+        // Options to disable setting of bucket access policies, useful for local testing as minio does not support them.
+        public bool SetBucketPolicies { get; set; } = true;
+    }
+
+    public class GCSettings
+    {
+        public bool BlobCleanupServiceEnabled { get; set; } = true;
+        
+        public bool CleanOldRefRecords { get; set; } = false;
+        public bool CleanOldBlobs { get; set; } = true;
+
+        public TimeSpan LastAccessCutoff { get; set; } = TimeSpan.FromDays(14);
+
+        public TimeSpan BlobCleanupPollFrequency { get; set; } = TimeSpan.FromMinutes(60);
+        public TimeSpan RefCleanupPollFrequency { get; set; } = TimeSpan.FromMinutes(60);
+    }
+}
diff --git a/Engine/Source/Programs/Horde/HordeStorage/Horde.Storage/Tests/Functional/.gitignore b/Engine/Source/Programs/Horde/HordeStorage/Horde.Storage/Tests/Functional/.gitignore
new file mode 100644
--- /dev/null
+++ b/Engine/Source/Programs/Horde/HordeStorage/Horde.Storage/Tests/Functional/.gitignore
@@ -0,0 +1,5 @@
+bin
+obj
+.vs
+
+*.csproj.user
\ No newline at end of file
diff --git a/Engine/Source/Programs/Horde/HordeStorage/Horde.Storage/Tests/Functional/ContentId/CompressedBlobTests.cs b/Engine/Source/Programs/Horde/HordeStorage/Horde.Storage/Tests/Functional/ContentId/CompressedBlobTests.cs
new file mode 100644
--- /dev/null
+++ b/Engine/Source/Programs/Horde/HordeStorage/Horde.Storage/Tests/Functional/ContentId/CompressedBlobTests.cs
@@ -0,0 +1,167 @@
+// Copyright Epic Games, Inc. All Rights Reserved.
+
+using System;
+using System.Collections.Generic;
+using System.IO;
+using System.Net;
+using System.Net.Http;
+using System.Net.Http.Headers;
+using System.Net.Mime;
+using System.Threading.Tasks;
+using Horde.Storage.FunctionalTests.Storage;
+using Horde.Storage.Implementation;
+using Jupiter;
+using Jupiter.Implementation;
+using Microsoft.AspNetCore.TestHost;
+using Microsoft.AspNetCore.Hosting;
+using Microsoft.Extensions.Configuration;
+using Microsoft.Extensions.DependencyInjection;
+using Microsoft.VisualStudio.TestTools.UnitTesting;
+using Serilog;
+using Serilog.Core;
+
+namespace Horde.Storage.FunctionalTests.CompressedBlobs
+{
+    [TestClass]
+    public class ScyllaCompressedBlobTests : CompressedBlobTests
+    {
+        protected override IEnumerable<KeyValuePair<string, string>> GetSettings()
+        {
+            // Use the S3 storage backend so we can handle large blobs
+            return new[]
+            {
+                new KeyValuePair<string, string>("Horde.Storage:StorageImplementations:0", HordeStorageSettings.StorageBackendImplementations.S3.ToString()),
+                new KeyValuePair<string, string>("S3:BucketName", $"tests-{TestNamespace}")
+            };
+        }
+
+        protected override async Task Seed(IServiceProvider provider)
+        {
+            await Task.CompletedTask;
+        }
+
+        protected override async Task Teardown()
+        {
+            await Task.CompletedTask;
+        }
+    }
+
+
+    public abstract class CompressedBlobTests
+    {
+        protected TestServer? _server;
+        protected HttpClient? _httpClient;
+
+        protected const string TestNamespace = "test-namespace";
+
+        [TestInitialize]
+        public async Task Setup()
+        {
+            IConfigurationRoot configuration = new ConfigurationBuilder()
+                // we are not reading the base appSettings here as we want exact control over what runs in the tests
+                .AddJsonFile("appsettings.Testing.json", true)
+                .AddInMemoryCollection(GetSettings())
+                .AddEnvironmentVariables()
+                .Build();
+
+            Logger logger = new LoggerConfiguration()
+                .ReadFrom.Configuration(configuration)
+                .CreateLogger();
+
+            TestServer server = new TestServer(new WebHostBuilder()
+                .UseConfiguration(configuration)
+                .UseEnvironment("Testing")
+                .UseSerilog(logger)
+                .UseStartup<HordeStorageStartup>()
+            );
+            _httpClient = server.CreateClient();
+            _server = server;
+
+            // Seed storage
+            await Seed(_server.Services);
+        }
+
+        protected abstract IEnumerable<KeyValuePair<string, string>> GetSettings();
+
+        protected abstract Task Seed(IServiceProvider serverServices);
+        protected abstract Task Teardown();
+
+
+        [TestCleanup]
+        public async Task MyTeardown()
+        {
+            await Teardown();
+        }
+
+
+        [TestMethod]
+        [DataTestMethod]
+        [DataRow("7835c353d7dc67e8a0531c88fbc75ddfda10dee4", "7835c353d7dc67e8a0531c88fbc75ddfda10dee4")]
+        [DataRow("4958689fe783e02fb35b13c14b0c3d7beb91e50c", "4958689fe783e02fb35b13c14b0c3d7beb91e50c")]
+        [DataRow("dce31eb416f3dcb4c8250ac545eda3930919d3ff", "dce31eb416f3dcb4c8250ac545eda3930919d3ff")]
+        [DataRow("05d7c699a2668efdecbe48f10db0d621d736f449.uecomp", "05D7C699A2668EFDECBE48F10DB0D621D736F449")]
+        [DataRow("dce31eb416f3dcb4c8250ac545eda3930919d3ff", "dce31eb416f3dcb4c8250ac545eda3930919d3ff")]
+        [DataRow("UncompressedTexture_CAS_dea81b6c3b565bb5089695377c98ce0f1c13b0c3.udd", "DEA81B6C3B565BB5089695377C98CE0F1C13B0C3")]
+        [DataRow("Oodle-f895ea954b37217270e88d8b728bd3c09152689c", "F895EA954B37217270E88D8B728BD3C09152689C")]
+        [DataRow("Oodle-f0b9c675fe21951ca27699f9baab9f9f5040b202", "F0B9C675FE21951CA27699F9BAAB9F9F5040B202")]
+        [DataRow("OodleTexture_CAS_dbda9040e75c4674fcec173f982fddf12b021e24.udd", "DBDA9040E75C4674FCEC173F982FDDF12B021E24")]
+        public async Task PutPayloads(string payloadFilename, string uncompressedHash)
+        {
+            byte[] texturePayload = await File.ReadAllBytesAsync($"ContentId/Payloads/{payloadFilename}");
+            BlobIdentifier compressedPayloadIdentifier = BlobIdentifier.FromBlob(texturePayload);
+            BlobIdentifier uncompressedPayloadIdentifier = new BlobIdentifier(uncompressedHash);
+
+            ByteArrayContent content = new ByteArrayContent(texturePayload);
+            content.Headers.ContentType = new MediaTypeHeaderValue(CustomMediaTypeNames.UnrealCompressedBuffer);
+            HttpResponseMessage result = await _httpClient!.PutAsync($"api/v1/compressed-blobs/{TestNamespace}/{uncompressedPayloadIdentifier}", content);
+            result.EnsureSuccessStatusCode();
+
+            InsertResponse response = await result.Content.ReadAsAsync<InsertResponse>();
+            Assert.AreNotEqual(compressedPayloadIdentifier, response.Identifier);
+            Assert.AreEqual(uncompressedPayloadIdentifier, response.Identifier);
+        }
+
+        [TestMethod]
+        public async Task PutGetComplexTexture()
+        {
+            byte[] texturePayload = await File.ReadAllBytesAsync("ContentId/Payloads/UncompressedTexture_CAS_dea81b6c3b565bb5089695377c98ce0f1c13b0c3.udd");
+            BlobIdentifier compressedPayloadIdentifier = BlobIdentifier.FromBlob(texturePayload);
+            BlobIdentifier uncompressedPayloadIdentifier = new BlobIdentifier("DEA81B6C3B565BB5089695377C98CE0F1C13B0C3");
+
+            {
+                ByteArrayContent content = new ByteArrayContent(texturePayload);
+                content.Headers.ContentType = new MediaTypeHeaderValue(CustomMediaTypeNames.UnrealCompressedBuffer);
+                HttpResponseMessage result = await _httpClient!.PutAsync($"api/v1/compressed-blobs/{TestNamespace}/{uncompressedPayloadIdentifier}", content);
+                result.EnsureSuccessStatusCode();
+
+                InsertResponse response = await result.Content.ReadAsAsync<InsertResponse>();
+                Assert.AreNotEqual(compressedPayloadIdentifier, response.Identifier);
+                Assert.AreEqual(uncompressedPayloadIdentifier, response.Identifier);
+            }
+
+
+            {
+                HttpResponseMessage result = await _httpClient!.GetAsync($"api/v1/compressed-blobs/{TestNamespace}/{uncompressedPayloadIdentifier}");
+                result.EnsureSuccessStatusCode();
+
+                byte[] blobContent = await result.Content.ReadAsByteArrayAsync();
+                CollectionAssert.AreEqual(texturePayload, blobContent);
+            }
+
+            {
+                // verify the compressed blob can be retrieved in the blob store
+                HttpResponseMessage result = await _httpClient!.GetAsync($"api/v1/blobs/{TestNamespace}/{compressedPayloadIdentifier}");
+                result.EnsureSuccessStatusCode();
+
+                byte[] blobContent = await result.Content.ReadAsByteArrayAsync();
+                CollectionAssert.AreEqual(texturePayload, blobContent);
+            }
+
+            {
+                // the uncompressed payload should not be valid in the blob endpoint
+                HttpResponseMessage result = await _httpClient!.GetAsync($"api/v1/blobs/{TestNamespace}/{uncompressedPayloadIdentifier}");
+                Assert.AreEqual(HttpStatusCode.NotFound, result.StatusCode);
+            }
+        }
+    }
+}
diff --git a/Engine/Source/Programs/Horde/HordeStorage/Horde.Storage/Tests/Functional/GC/GCBlobTests.cs b/Engine/Source/Programs/Horde/HordeStorage/Horde.Storage/Tests/Functional/GC/GCBlobTests.cs
new file mode 100644
--- /dev/null
+++ b/Engine/Source/Programs/Horde/HordeStorage/Horde.Storage/Tests/Functional/GC/GCBlobTests.cs
@@ -0,0 +1,178 @@
+// Copyright Epic Games, Inc. All Rights Reserved.
+
+using System;
+using System.Collections.Generic;
+using System.Linq;
+using System.Net;
+using System.Net.Http;
+using System.Threading;
+using System.Threading.Tasks;
+using Dasync.Collections;
+using Jupiter.Implementation;
+using Horde.Storage.Implementation;
+using Microsoft.AspNetCore.TestHost;
+using Microsoft.AspNetCore.Hosting;
+using Microsoft.Extensions.Configuration;
+using Microsoft.Extensions.DependencyInjection;
+using Microsoft.VisualStudio.TestTools.UnitTesting;
+using RestSharp;
+using Serilog;
+using Serilog.Core;
+using Moq;
+
+namespace Horde.Storage.FunctionalTests.GC
+{
+    [TestClass]
+    public class GCBlobTests
+    {
+        private TestServer? _server;
+        private HttpClient? _httpClient;
+
+        private readonly BlobIdentifier object0id = new BlobIdentifier("0000000000000000000000000000000000000000");
+        private readonly BlobIdentifier object1id = new BlobIdentifier("1111111111111111111111111111111111111111");
+        private readonly BlobIdentifier object2id = new BlobIdentifier("2222222222222222222222222222222222222222");
+        private readonly BlobIdentifier object3id = new BlobIdentifier("3333333333333333333333333333333333333333");
+        private readonly BlobIdentifier object4id = new BlobIdentifier("4444444444444444444444444444444444444444");
+        private readonly BlobIdentifier object5id = new BlobIdentifier("5555555555555555555555555555555555555555");
+        private readonly BlobIdentifier object6id = new BlobIdentifier("6666666666666666666666666666666666666666");
+        private Mock<IRestClient> _callistoBlobMock = null!;
+        private IBlobStore? _blobStore;
+
+        private readonly NamespaceId TestNamespace = new NamespaceId("test-namespace");
+
+        [TestInitialize]
+        public async Task Setup()
+        {
+            IConfigurationRoot configuration = new ConfigurationBuilder()
+                // we are not reading the base appSettings here as we want exact control over what runs in the tests
+                .AddJsonFile("appsettings.Testing.json", false)
+                .AddEnvironmentVariables()
+                .Build();
+
+            Logger logger = new LoggerConfiguration()
+                .ReadFrom.Configuration(configuration)
+                .CreateLogger();
+
+
+            _callistoBlobMock = GetCallistoMock(new CallistoReader.CallistoGetResponse(new TransactionEvent[]
+            {
+                new AddTransactionEvent("object0", "default", new[] {object0id}, identifier: 0, nextIdentifier: 1),
+                new AddTransactionEvent("object2", "default", new[] {object2id}, identifier: 1, nextIdentifier: 2),
+                new RemoveTransactionEvent("object2", "default", identifier: 2, nextIdentifier: 3),
+                new AddTransactionEvent("object3", "default", new[] {object3id}, identifier: 3, nextIdentifier: 4),
+                new AddTransactionEvent("object6", "default", new[] {object6id}, identifier: 4, nextIdentifier: 5)
+            }.ToList(), Guid.Empty, 10));
+
+            TestServer server = new TestServer(new WebHostBuilder()
+                .UseConfiguration(configuration)
+                .UseEnvironment("Testing")
+                .UseSerilog(logger)
+                .UseStartup<HordeStorageStartup>()
+                .ConfigureTestServices(collection =>
+                {
+                    collection.AddSingleton<IBlobCleanup>(provider =>
+                    {
+                        IBlobStore blobStore = provider.GetService<IBlobStore>()!;
+                        return new OrphanBlobCleanup(blobStore, new LeaderElectionStub(true), _callistoBlobMock.Object);
+                    });
+
+                    // make sure we have a memory blob store as we do some funky calling to it below
+                    collection.AddSingleton<IBlobStore>(provider => ActivatorUtilities.CreateInstance<MemoryBlobStore>(provider));
+                })
+            );
+
+            _httpClient = server.CreateClient();
+            _server = server;
+
+            _blobStore = server.Services.GetService<IBlobStore>()!;
+
+            MemoryBlobStore memoryBlobStore = (MemoryBlobStore) _blobStore;
+            byte[] emptyContents = new byte[0];
+            await memoryBlobStore.PutObject(TestNamespace, emptyContents, object0id);
+            await memoryBlobStore.PutObject(TestNamespace, emptyContents, object1id);// this is not in callisto
+            await memoryBlobStore.PutObject(TestNamespace, emptyContents, object2id);
+            await memoryBlobStore.PutObject(TestNamespace, emptyContents, object3id);
+            await memoryBlobStore.PutObject(TestNamespace, emptyContents, object4id); // this is not in callisto
+            await memoryBlobStore.PutObject(TestNamespace, emptyContents, object5id); // this is not in callisto
+            await memoryBlobStore.PutObject(TestNamespace, emptyContents, object6id);
+
+            // set all objects to be old, only the orphaned blobs should be deleted
+            memoryBlobStore.SetLastModifiedTime(TestNamespace, object0id, DateTime.Now.AddDays(-2));
+            memoryBlobStore.SetLastModifiedTime(TestNamespace, object1id, DateTime.Now.AddDays(-2));
+            memoryBlobStore.SetLastModifiedTime(TestNamespace, object2id, DateTime.Now.AddDays(-2));
+            memoryBlobStore.SetLastModifiedTime(TestNamespace, object3id, DateTime.Now.AddDays(-2));
+            memoryBlobStore.SetLastModifiedTime(TestNamespace, object4id, DateTime.Now.AddDays(-2));
+            memoryBlobStore.SetLastModifiedTime(TestNamespace, object5id, DateTime.Now.AddDays(-2));
+            memoryBlobStore.SetLastModifiedTime(TestNamespace, object6id, DateTime.Now.AddDays(-2));
+        }
+
+        private Mock<IRestClient> GetCallistoMock(CallistoReader.CallistoGetResponse responseData)
+        {
+            // the callisto reader will do 2 requests to fetch the events (one to determine the generation and the other to actually fetch the objects)
+            // after that we insert a empty response to indicate the end of the log was reached
+            Mock<IRestClient> mock = new Mock<IRestClient> { DefaultValue = DefaultValue.Empty };
+
+            
+            Mock<IRestResponse<OrphanBlobCleanup.ListNamespaceResponse>> listNsResponse = new Mock<IRestResponse<OrphanBlobCleanup.ListNamespaceResponse>>();
+            listNsResponse.Setup(_ => _.StatusCode).Returns(HttpStatusCode.OK);
+            listNsResponse.Setup(_ => _.IsSuccessful).Returns(true);
+            listNsResponse.Setup(_ => _.Data).Returns(new OrphanBlobCleanup.ListNamespaceResponse() { Logs = new [] {TestNamespace}});
+            listNsResponse.Setup(_ => _.ResponseUri).Returns(new Uri("http://localhost"));
+            
+            mock.Setup(x =>
+                x.ExecuteGetAsync<OrphanBlobCleanup.ListNamespaceResponse>(It.IsAny<IRestRequest>(),
+                    It.IsAny<CancellationToken>())).ReturnsAsync(listNsResponse.Object).Verifiable();
+            
+            
+            Mock<IRestResponse<CallistoReader.CallistoGetResponse>> response = new Mock<IRestResponse<CallistoReader.CallistoGetResponse>>();
+            response.Setup(_ => _.StatusCode).Returns(HttpStatusCode.OK);
+            response.Setup(_ => _.IsSuccessful).Returns(true);
+            response.Setup(_ => _.Data).Returns(responseData);
+            response.Setup(_ => _.ResponseUri).Returns(new Uri("http://localhost"));
+            
+            mock.Setup(x =>
+                x.ExecuteGetAsync<CallistoReader.CallistoGetResponse>(It.IsAny<IRestRequest>(),
+                    It.IsAny<CancellationToken>())).ReturnsAsync(response.Object).Verifiable();
+
+            // when we do the second call to callisto with a offset we return a empty object indicating we have reached the end
+            Mock<IRestResponse<CallistoReader.CallistoGetResponse>> responseSecondPage = new Mock<IRestResponse<CallistoReader.CallistoGetResponse>>();
+            responseSecondPage.Setup(_ => _.StatusCode).Returns(HttpStatusCode.OK);
+            responseSecondPage.Setup(_ => _.IsSuccessful).Returns(true);
+            responseSecondPage.Setup(_ => _.Data).Returns(new CallistoReader.CallistoGetResponse(
+                new List<TransactionEvent>(),
+                responseData.Generation,
+                responseData.CurrentOffset + 10 // set the offset to some arbitrary higher number to indicate we are at the end of the log
+            ));
+            responseSecondPage.Setup(_ => _.ResponseUri).Returns(new Uri("http://localhost"));
+            mock.Setup(x => x.ExecuteGetAsync<CallistoReader.CallistoGetResponse>(
+                    It.Is<IRestRequest>(request =>
+                        request.Parameters.Any(parameter =>
+                            parameter.Name == "offset" && parameter.Value!.ToString() != "0")),
+                    It.IsAny<CancellationToken>()))
+                .ReturnsAsync(responseSecondPage.Object).Verifiable();
+
+            return mock;
+        }
+
+        [TestMethod]
+        public async Task RunBlobCleanup()
+        {
+            OrphanBlobCleanup cleanup = new OrphanBlobCleanup(_blobStore!, new LeaderElectionStub(true), _callistoBlobMock.Object);
+            List<NamespaceId> namespaces = await cleanup.ListNamespaces().ToListAsync();
+            Assert.AreEqual(1, namespaces.Count);
+
+            CancellationTokenSource cts = new CancellationTokenSource();
+            List<RemovedBlobs> removedBlobs = await cleanup.Cleanup(cts.Token);
+            Assert.AreEqual(4, removedBlobs.Count);
+            
+            BlobIdentifier[] expectedIdentifiers = {object1id, object2id, object4id, object5id};
+            foreach (RemovedBlobs removedBlob in removedBlobs)
+            {
+                Assert.IsTrue(expectedIdentifiers.Contains(removedBlob.BlobIdentifier));
+            }
+            
+            _callistoBlobMock.Verify();
+            _callistoBlobMock.VerifyNoOtherCalls();
+        }
+    }
+}
diff --git a/Engine/Source/Programs/Horde/HordeStorage/Horde.Storage/Tests/Functional/GC/GCRefTests.cs b/Engine/Source/Programs/Horde/HordeStorage/Horde.Storage/Tests/Functional/GC/GCRefTests.cs
new file mode 100644
--- /dev/null
+++ b/Engine/Source/Programs/Horde/HordeStorage/Horde.Storage/Tests/Functional/GC/GCRefTests.cs
@@ -0,0 +1,99 @@
+// Copyright Epic Games, Inc. All Rights Reserved.
+
+using System;
+using System.Collections.Generic;
+using System.Linq;
+using System.Net.Http;
+using System.Threading.Tasks;
+using Dasync.Collections;
+using Horde.Storage.Controllers;
+using Horde.Storage.Implementation;
+using Jupiter.Implementation;
+using Microsoft.AspNetCore.TestHost;
+using Microsoft.AspNetCore.Hosting;
+using Microsoft.Extensions.Configuration;
+using Microsoft.Extensions.DependencyInjection;
+using Microsoft.VisualStudio.TestTools.UnitTesting;
+using Moq;
+using Serilog;
+using Serilog.Core;
+
+namespace Horde.Storage.FunctionalTests.GC
+{
+    [TestClass]
+
+    public class GCRefTests
+    {
+        private TestServer? _server;
+        private HttpClient? _httpClient;
+        private Mock<IRefsStore>? _refMock;
+
+        private readonly NamespaceId TestNamespace = new NamespaceId("test-namespace");
+        private readonly BucketId DefaultBucket = new BucketId("default");
+
+        [TestInitialize]
+        public void Setup()
+        {
+            IConfigurationRoot configuration = new ConfigurationBuilder()
+                            // we are not reading the base appSettings here as we want exact control over what runs in the tests
+                            .AddJsonFile("appsettings.Testing.json", false)
+                            .AddEnvironmentVariables()
+                            .Build();
+
+            Logger logger = new LoggerConfiguration()
+                .ReadFrom.Configuration(configuration)
+                .CreateLogger();
+
+            OldRecord[] oldRecords =
+            {
+                new OldRecord(TestNamespace, DefaultBucket, new KeyId("object2")),
+                new OldRecord(TestNamespace, DefaultBucket, new KeyId("object5")),
+                new OldRecord(TestNamespace, DefaultBucket, new KeyId("object6")),
+            };
+            _refMock = new Mock<IRefsStore>();
+            _refMock.Setup(store => store.GetOldRecords(TestNamespace, It.IsAny<TimeSpan>())).Returns(oldRecords.ToAsyncEnumerable()).Verifiable();
+            _refMock.Setup(store => store.Delete(TestNamespace, DefaultBucket, It.IsAny<KeyId>())).Verifiable();
+
+            TestServer server = new TestServer(new WebHostBuilder()
+                .UseConfiguration(configuration)
+                .UseEnvironment("Testing")
+                .UseSerilog(logger)
+                .UseStartup<HordeStorageStartup>()
+                .ConfigureTestServices(collection =>
+                {
+                    // use our refs mock instead of the actual refs store so we can control which records are invalid
+                    collection.AddSingleton<IRefsStore>(_refMock.Object);
+                })
+            );
+            _httpClient = server.CreateClient();
+            _server = server;
+            
+        }
+
+
+        [TestMethod]
+        public async Task RunRefCleanup()
+        {
+            // trigger the cleanup
+            var cleanupResponse = await _httpClient!.PostAsync($"api/v1/admin/refCleanup/{TestNamespace}", new StringContent(string.Empty));
+            cleanupResponse.EnsureSuccessStatusCode();
+            RemovedRefRecordsResponse removedRefRecords = await cleanupResponse.Content.ReadAsAsync<RemovedRefRecordsResponse>();
+
+            KeyId[] expectedRemovedRefRecords = new[]
+            {
+                new KeyId("object2"),
+                new KeyId("object5"),
+                new KeyId("object6")
+            };
+            Assert.AreEqual(3, removedRefRecords.RemovedRecords.Length);
+            foreach (RemovedRefRecordsResponse.RemovedRecord removedRecord in removedRefRecords.RemovedRecords)
+            {
+                Assert.AreEqual("default", removedRecord.Bucket.ToString());
+                Assert.IsTrue(expectedRemovedRefRecords.Contains(removedRecord.Name));
+            }
+
+            _refMock!.Verify();
+            _refMock.VerifyNoOtherCalls();
+        }
+    }
+}
diff --git a/Engine/Source/Programs/Horde/HordeStorage/Horde.Storage/Tests/Functional/Horde.Storage.Tests.Functional.csproj b/Engine/Source/Programs/Horde/HordeStorage/Horde.Storage/Tests/Functional/Horde.Storage.Tests.Functional.csproj
new file mode 100644
--- /dev/null
+++ b/Engine/Source/Programs/Horde/HordeStorage/Horde.Storage/Tests/Functional/Horde.Storage.Tests.Functional.csproj
@@ -0,0 +1,62 @@
+<Project Sdk="Microsoft.NET.Sdk">
+
+  <PropertyGroup>
+    <TargetFramework>net6.0</TargetFramework>
+
+    <IsPackable>false</IsPackable>
+    <Nullable>enable</Nullable>
+
+    <TreatWarningsAsErrors>true</TreatWarningsAsErrors>
+    <LangVersion>10</LangVersion>
+  </PropertyGroup>
+
+  <ItemGroup>
+    <PackageReference Include="Grpc.Net.Client" Version="2.37.0" />
+    <PackageReference Include="Microsoft.AspNetCore.TestHost" Version="6.0.0" />
+    <PackageReference Include="Microsoft.NET.Test.Sdk" Version="17.0.0" />
+    <PackageReference Include="MSTest.TestAdapter" Version="2.2.7" />
+    <PackageReference Include="MSTest.TestFramework" Version="2.2.7" />
+    <PackageReference Include="coverlet.collector" Version="3.0.2">
+      <PrivateAssets>all</PrivateAssets>
+      <IncludeAssets>runtime; build; native; contentfiles; analyzers; buildtransitive</IncludeAssets>
+    </PackageReference>
+    <PackageReference Include="SerilogAnalyzer" Version="0.15.0" />
+    <PackageReference Include="Moq" Version="4.16.1" />
+    <PackageReference Include="Moq.Contrib.HttpClient" Version="1.3.0" />
+  </ItemGroup>
+
+  <ItemGroup>
+    <ProjectReference Include="..\..\Horde.Storage.csproj" />
+  </ItemGroup>
+
+  <ItemGroup>
+    <None Update="ContentId\Payloads\05d7c699a2668efdecbe48f10db0d621d736f449.uecomp">
+      <CopyToOutputDirectory>PreserveNewest</CopyToOutputDirectory>
+    </None>
+    <None Update="ContentId\Payloads\0b29a878f71cc1eb6a26663a1152537e8fe7dca9.uecomp">
+      <CopyToOutputDirectory>PreserveNewest</CopyToOutputDirectory>
+    </None>
+    <None Update="ContentId\Payloads\4958689fe783e02fb35b13c14b0c3d7beb91e50c">
+      <CopyToOutputDirectory>PreserveNewest</CopyToOutputDirectory>
+    </None>
+    <None Update="ContentId\Payloads\7835c353d7dc67e8a0531c88fbc75ddfda10dee4">
+      <CopyToOutputDirectory>PreserveNewest</CopyToOutputDirectory>
+    </None>
+    <None Update="ContentId\Payloads\dce31eb416f3dcb4c8250ac545eda3930919d3ff">
+      <CopyToOutputDirectory>PreserveNewest</CopyToOutputDirectory>
+    </None>
+    <None Update="ContentId\Payloads\Oodle-f0b9c675fe21951ca27699f9baab9f9f5040b202">
+      <CopyToOutputDirectory>PreserveNewest</CopyToOutputDirectory>
+    </None>
+    <None Update="ContentId\Payloads\Oodle-f895ea954b37217270e88d8b728bd3c09152689c">
+      <CopyToOutputDirectory>PreserveNewest</CopyToOutputDirectory>
+    </None>
+    <None Update="ContentId\Payloads\OodleTexture_CAS_dbda9040e75c4674fcec173f982fddf12b021e24.udd">
+      <CopyToOutputDirectory>PreserveNewest</CopyToOutputDirectory>
+    </None>
+    <None Update="ContentId\Payloads\UncompressedTexture_CAS_dea81b6c3b565bb5089695377c98ce0f1c13b0c3.udd">
+      <CopyToOutputDirectory>PreserveNewest</CopyToOutputDirectory>
+    </None>
+  </ItemGroup>
+
+</Project>
diff --git a/Engine/Source/Programs/Horde/HordeStorage/Horde.Storage/Tests/Functional/Objects/RefReplicatorTests.cs b/Engine/Source/Programs/Horde/HordeStorage/Horde.Storage/Tests/Functional/Objects/RefReplicatorTests.cs
new file mode 100644
--- /dev/null
+++ b/Engine/Source/Programs/Horde/HordeStorage/Horde.Storage/Tests/Functional/Objects/RefReplicatorTests.cs
@@ -0,0 +1,414 @@
+// Copyright Epic Games, Inc. All Rights Reserved.
+
+using System;
+using System.Collections.Generic;
+using System.IO;
+using System.Linq;
+using System.Net;
+using System.Net.Http;
+using System.Text;
+using System.Text.Json;
+using System.Threading.Tasks;
+using Cassandra;
+using Horde.Storage.Controllers;
+using Horde.Storage.Implementation;
+using Horde.Storage.Implementation.TransactionLog;
+using Jupiter.Implementation;
+using Microsoft.AspNetCore.TestHost;
+using Microsoft.AspNetCore.Hosting;
+using Microsoft.AspNetCore.Mvc;
+using Microsoft.Extensions.Configuration;
+using Microsoft.Extensions.DependencyInjection;
+using Microsoft.VisualStudio.TestTools.UnitTesting;
+using Moq;
+using Moq.Contrib.HttpClient;
+using Newtonsoft.Json;
+using Serilog;
+using Logger = Serilog.Core.Logger;
+
+namespace Horde.Storage.FunctionalTests.Replication
+{
+
+    [TestClass]
+    public class RefReplicatorTests
+    {
+        private static TestServer? _server;
+        private static HttpClient? _httpClient;
+        protected IBlobStore _blobStore = null!;
+        protected IReferencesStore _referencesStore = null!;
+        private IReplicationLog _replicationLog = null!;
+
+        protected readonly NamespaceId TestNamespace = new NamespaceId("test-namespace");
+        protected readonly NamespaceId SnapshotNamespace = new NamespaceId("snapshot-namespace");
+        protected readonly BucketId TestBucket = new BucketId("test");
+
+        [TestInitialize]
+        public async Task Setup()
+        {
+            IConfigurationRoot configuration = new ConfigurationBuilder()
+                // we are not reading the base appSettings here as we want exact control over what runs in the tests
+                .AddJsonFile("appsettings.Testing.json", true)
+                .AddInMemoryCollection(GetSettings())
+                .AddEnvironmentVariables()
+                .Build();
+
+            Logger logger = new LoggerConfiguration()
+                .ReadFrom.Configuration(configuration)
+                .CreateLogger();
+
+            TestServer server = new TestServer(new WebHostBuilder()
+                .UseConfiguration(configuration)
+                .UseEnvironment("Testing")
+                .UseSerilog(logger)
+                .UseStartup<HordeStorageStartup>()
+            );
+            _httpClient = server.CreateClient();
+            _server = server;
+
+            _blobStore = _server.Services.GetService<IBlobStore>()!;
+            _referencesStore = _server.Services.GetService<IReferencesStore>()!;
+            _replicationLog = _server.Services.GetService<IReplicationLog>()!;
+
+            await Task.CompletedTask;
+        }
+
+        private IEnumerable<KeyValuePair<string, string>> GetSettings()
+        {
+            return new[]
+            {
+                new KeyValuePair<string, string>("Horde.Storage:ReferencesDbImplementation", HordeStorageSettings.ReferencesDbImplementations.Scylla.ToString()),
+                new KeyValuePair<string, string>("Horde.Storage:ReplicationLogWriterImplementation", HordeStorageSettings.ReplicationLogWriterImplementations.Scylla.ToString()),
+            };
+        }
+
+
+        private async Task TeardownDb(IServiceProvider provider)
+        {
+            IScyllaSessionManager scyllaSessionManager = provider.GetService<IScyllaSessionManager>()!;
+            ISession session = scyllaSessionManager.GetSessionForLocalKeyspace();
+            
+            // remove replication log table as we expect it to be empty when starting the tests
+            await session.ExecuteAsync(new SimpleStatement("DROP TABLE IF EXISTS replication_log;"));
+
+            // remove the snapshots
+            await session.ExecuteAsync(new SimpleStatement("DROP TABLE IF EXISTS replication_snapshot;"));
+        }
+
+
+        [TestCleanup]
+        public async Task Teardown()
+        {
+            if (_server != null) 
+                await TeardownDb(_server.Services);
+        }
+        
+        [TestMethod]
+        public async Task ReplicationIncrementalState()
+        {
+            ReplicatorSettings replicatorSettings = new()
+            {
+                ConnectionString = "http://localhost",
+                MaxParallelReplications = 16,
+                NamespaceToReplicate = TestNamespace,
+                ReplicatorName = "test-replicator",
+                Version = ReplicatorVersion.Refs
+            };
+
+            List<ReplicationLogEvent> replicationEvents = new();
+            Dictionary<BlobIdentifier, byte[]> blobs = new();
+
+            const int countOfTestEvents = 100;
+            for (int i = 0; i < countOfTestEvents; i++)
+            {
+                byte[] blobContents = Encoding.UTF8.GetBytes($"random content {i}");
+                BlobIdentifier blob = BlobIdentifier.FromBlob(blobContents);
+                blobs.Add(blob, blobContents);
+                replicationEvents.Add(new ReplicationLogEvent(TestNamespace, TestBucket, new KeyId($"event-{i}"), blob, Guid.NewGuid(), "refs-000", DateTime.Now, ReplicationLogEvent.OpType.Added));
+            }
+            string lastBucket = "refs-000";
+            Guid lastEvent = Guid.NewGuid();
+
+            Mock<HttpMessageHandler> handler = new Mock<HttpMessageHandler>();
+            string s = JsonConvert.SerializeObject(new ReplicationLogEvents(replicationEvents));
+            handler.SetupRequest($"http://localhost/api/v1/replication-log/incremental/{TestNamespace}?lastBucket={lastBucket}&lastEvent={lastEvent}").ReturnsResponse(s, "application/json");
+            // this will get called again with the last event in replication events list as the lastEvent, at which point we should return a empty list as there are no more objects available
+            handler.SetupRequest($"http://localhost/api/v1/replication-log/incremental/{TestNamespace}?lastBucket={replicationEvents.Last().TimeBucket}&lastEvent={replicationEvents.Last().EventId}").ReturnsResponse(JsonConvert.SerializeObject(new ReplicationLogEvents(new List<ReplicationLogEvent>())), "application/json");
+
+            foreach (BlobIdentifier blob in blobs.Keys)
+            {
+                handler.SetupRequest($"http://localhost/api/v1/objects/{TestNamespace}/{blob}/references").ReturnsResponse(JsonConvert.SerializeObject(new ResolvedReferencesResult(Array.Empty<BlobIdentifier>())), "application/json").Verifiable();
+            }
+
+            foreach ((BlobIdentifier key, byte[] blobContent) in blobs)
+            {
+                handler.SetupRequest($"http://localhost/api/v1/blobs/{TestNamespace}/{key}").ReturnsResponse(blobContent, "application/octet-stream").Verifiable();
+            }
+
+            IHttpClientFactory httpClientFactory = handler.CreateClientFactory();
+            RefsReplicator replicator = ActivatorUtilities.CreateInstance<RefsReplicator>(_server!.Services, replicatorSettings, httpClientFactory);
+            replicator.SetRefState(lastBucket, lastEvent);
+
+            bool didRun = await replicator.TriggerNewReplications();
+
+            Assert.IsTrue(didRun);
+
+            handler.Verify();
+
+            // Verify that the objects are present
+            BlobIdentifier[] missingBlobs = await _blobStore.FilterOutKnownBlobs(TestNamespace, blobs.Keys.ToArray());
+            Assert.IsFalse(missingBlobs.Any());
+        }
+
+        [TestMethod]
+        public async Task ReplicationSnapshotState()
+        {
+            ReplicatorSettings replicatorSettings = new()
+            {
+                ConnectionString = "http://localhost",
+                MaxParallelReplications = 16,
+                NamespaceToReplicate = TestNamespace,
+                ReplicatorName = "test-replicator",
+                Version = ReplicatorVersion.Refs
+            };
+
+            List<ReplicationLogEvent> replicationEvents = new();
+            Dictionary<BlobIdentifier, byte[]> blobs = new();
+
+            const int countOfTestEvents = 100;
+            for (int i = 0; i < countOfTestEvents; i++)
+            {
+                byte[] blobContents = Encoding.UTF8.GetBytes($"random content {i}");
+                BlobIdentifier blob = BlobIdentifier.FromBlob(blobContents);
+                blobs.Add(blob, blobContents);
+                replicationEvents.Add(new ReplicationLogEvent(TestNamespace, TestBucket, new KeyId($"event-{i}"), blob, Guid.NewGuid(), "refs-000", DateTime.Now, ReplicationLogEvent.OpType.Added));
+            }
+
+            // Build snapshot
+            ReplicationLogSnapshot snapshot = new ReplicationLogSnapshot(TestNamespace);
+            foreach (ReplicationLogEvent logEvent in replicationEvents)
+            {
+                snapshot.ProcessEvent(logEvent);
+            }
+
+            byte[] snapshotContent;
+            {
+                await using MemoryStream ms = new MemoryStream();
+                await snapshot.Serialize(ms);
+                snapshotContent = ms.ToArray();
+            }
+
+            BlobIdentifier snapshotBlob = BlobIdentifier.FromBlob(snapshotContent);
+
+            Mock<HttpMessageHandler> handler = new Mock<HttpMessageHandler>();
+            string s = JsonConvert.SerializeObject(new ReplicationLogSnapshots(new List<SnapshotInfo>{new SnapshotInfo(TestNamespace, SnapshotNamespace, snapshotBlob)}));
+            handler.SetupRequest($"http://localhost/api/v1/replication-log/snapshots/{TestNamespace}").ReturnsResponse(s, "application/json");
+
+            // after processing a snapshot it will attempt to incrementally replicate from there, which should be empty
+            handler.SetupRequest($"http://localhost/api/v1/replication-log/incremental/{TestNamespace}?lastBucket={replicationEvents.Last().TimeBucket}&lastEvent={replicationEvents.Last().EventId}").ReturnsResponse(JsonConvert.SerializeObject(new ReplicationLogEvents(new List<ReplicationLogEvent>())), "application/json");
+
+            handler.SetupRequest($"http://localhost/api/v1/blobs/{TestNamespace}/{snapshotBlob}").ReturnsResponse(snapshotContent, "application/octet-stream").Verifiable();
+
+            foreach (BlobIdentifier blob in blobs.Keys)
+            {
+                handler.SetupRequest($"http://localhost/api/v1/objects/{TestNamespace}/{blob}/references").ReturnsResponse(JsonConvert.SerializeObject(new ResolvedReferencesResult(Array.Empty<BlobIdentifier>())), "application/json").Verifiable();
+            }
+
+            foreach ((BlobIdentifier key, byte[] blobContent) in blobs)
+            {
+                handler.SetupRequest($"http://localhost/api/v1/blobs/{TestNamespace}/{key}").ReturnsResponse(blobContent, "application/octet-stream").Verifiable();
+            }
+
+            IHttpClientFactory httpClientFactory = handler.CreateClientFactory();
+            RefsReplicator replicator = ActivatorUtilities.CreateInstance<RefsReplicator>(_server!.Services, replicatorSettings, httpClientFactory);
+            // no previous state, will download it from a snapshot
+            replicator.SetRefState(null, null);
+
+            bool didRun = await replicator.TriggerNewReplications();
+
+            Assert.IsTrue(didRun);
+
+            handler.Verify();
+
+            // Verify that the objects are present
+            BlobIdentifier[] missingBlobs = await _blobStore.FilterOutKnownBlobs(TestNamespace, blobs.Keys.ToArray());
+            Assert.IsFalse(missingBlobs.Any());
+        }
+
+
+        [TestMethod]
+        public async Task ReplicationStateBoth()
+        {
+            ReplicatorSettings replicatorSettings = new()
+            {
+                ConnectionString = "http://localhost",
+                MaxParallelReplications = 16,
+                NamespaceToReplicate = TestNamespace,
+                ReplicatorName = "test-replicator",
+                Version = ReplicatorVersion.Refs
+            };
+
+            List<ReplicationLogEvent> snapshotEvents = new();
+            List<ReplicationLogEvent> incrementalEvents = new();
+            Dictionary<BlobIdentifier, byte[]> blobs = new();
+
+            const int countOfTestEvents = 100;
+            for (int i = 0; i < countOfTestEvents; i++)
+            {
+                byte[] blobContents = Encoding.UTF8.GetBytes($"random content in snapshot {i}");
+                BlobIdentifier blob = BlobIdentifier.FromBlob(blobContents);
+                blobs.Add(blob, blobContents);
+                snapshotEvents.Add(new ReplicationLogEvent(TestNamespace, TestBucket, new KeyId($"event-{i}"), blob, Guid.NewGuid(), "refs-000", DateTime.Now, ReplicationLogEvent.OpType.Added));
+            }
+
+            for (int i = 0; i < countOfTestEvents; i++)
+            {
+                byte[] blobContents = Encoding.UTF8.GetBytes($"random content {i}");
+                BlobIdentifier blob = BlobIdentifier.FromBlob(blobContents);
+                blobs.Add(blob, blobContents);
+                incrementalEvents.Add(new ReplicationLogEvent(TestNamespace, TestBucket, new KeyId($"incremental-event-{i}"), blob, Guid.NewGuid(), "refs-000", DateTime.Now, ReplicationLogEvent.OpType.Added));
+            }
+            
+            // Build snapshot
+            ReplicationLogSnapshot snapshot = new ReplicationLogSnapshot(TestNamespace);
+            foreach (ReplicationLogEvent logEvent in snapshotEvents)
+            {
+                snapshot.ProcessEvent(logEvent);
+            }
+
+            byte[] snapshotContent;
+            {
+                await using MemoryStream ms = new MemoryStream();
+                await snapshot.Serialize(ms);
+                snapshotContent = ms.ToArray();
+            }
+
+            BlobIdentifier snapshotBlob = BlobIdentifier.FromBlob(snapshotContent);
+
+            Mock<HttpMessageHandler> handler = new Mock<HttpMessageHandler>();
+            string s = JsonConvert.SerializeObject(new ReplicationLogSnapshots(new List<SnapshotInfo>{new SnapshotInfo(TestNamespace, SnapshotNamespace, snapshotBlob)}));
+            handler.SetupRequest($"http://localhost/api/v1/replication-log/snapshots/{TestNamespace}").ReturnsResponse(s, "application/json");
+
+            // when the snapshot has been processed we have a set of incremental events as well
+            handler.SetupRequest($"http://localhost/api/v1/replication-log/incremental/{TestNamespace}?lastBucket={snapshotEvents.Last().TimeBucket}&lastEvent={snapshotEvents.Last().EventId}").ReturnsResponse(JsonConvert.SerializeObject(new ReplicationLogEvents(incrementalEvents)), "application/json");
+
+            // after processing the incremental events there is nothing more to find
+            handler.SetupRequest($"http://localhost/api/v1/replication-log/incremental/{TestNamespace}?lastBucket={incrementalEvents.Last().TimeBucket}&lastEvent={incrementalEvents.Last().EventId}").ReturnsResponse(JsonConvert.SerializeObject(new ReplicationLogEvents(new List<ReplicationLogEvent>())), "application/json");
+
+            handler.SetupRequest($"http://localhost/api/v1/blobs/{TestNamespace}/{snapshotBlob}").ReturnsResponse(snapshotContent, "application/octet-stream").Verifiable();
+
+            foreach (BlobIdentifier blob in blobs.Keys)
+            {
+                handler.SetupRequest($"http://localhost/api/v1/objects/{TestNamespace}/{blob}/references").ReturnsResponse(JsonConvert.SerializeObject(new ResolvedReferencesResult(Array.Empty<BlobIdentifier>())), "application/json").Verifiable();
+            }
+
+            foreach ((BlobIdentifier key, byte[] blobContent) in blobs)
+            {
+                handler.SetupRequest($"http://localhost/api/v1/blobs/{TestNamespace}/{key}").ReturnsResponse(blobContent, "application/octet-stream").Verifiable();
+            }
+
+            IHttpClientFactory httpClientFactory = handler.CreateClientFactory();
+            RefsReplicator replicator = ActivatorUtilities.CreateInstance<RefsReplicator>(_server!.Services, replicatorSettings, httpClientFactory);
+            // no previous state, will download it from a snapshot
+            replicator.SetRefState(null, null);
+
+            bool didRun = await replicator.TriggerNewReplications();
+
+            Assert.IsTrue(didRun);
+
+            handler.Verify();
+
+            // Verify that the objects are present
+            BlobIdentifier[] missingBlobs = await _blobStore.FilterOutKnownBlobs(TestNamespace, blobs.Keys.ToArray());
+            Assert.IsFalse(missingBlobs.Any());
+        }
+
+
+        
+        [TestMethod]
+        public async Task ReplicationStateSnapshotFallback()
+        {
+            ReplicatorSettings replicatorSettings = new()
+            {
+                ConnectionString = "http://localhost",
+                MaxParallelReplications = 16,
+                NamespaceToReplicate = TestNamespace,
+                ReplicatorName = "test-replicator",
+                Version = ReplicatorVersion.Refs
+            };
+
+            List<ReplicationLogEvent> snapshotEvents = new();
+            Dictionary<BlobIdentifier, byte[]> blobs = new();
+
+            const int countOfTestEvents = 100;
+            for (int i = 0; i < countOfTestEvents; i++)
+            {
+                byte[] blobContents = Encoding.UTF8.GetBytes($"random content in snapshot {i}");
+                BlobIdentifier blob = BlobIdentifier.FromBlob(blobContents);
+                blobs.Add(blob, blobContents);
+                snapshotEvents.Add(new ReplicationLogEvent(TestNamespace, TestBucket, new KeyId($"event-{i}"), blob, Guid.NewGuid(), "refs-000", DateTime.Now, ReplicationLogEvent.OpType.Added));
+            }
+
+            // Build snapshot
+            ReplicationLogSnapshot snapshot = new ReplicationLogSnapshot(TestNamespace);
+            foreach (ReplicationLogEvent logEvent in snapshotEvents)
+            {
+                snapshot.ProcessEvent(logEvent);
+            }
+
+            byte[] snapshotContent;
+            {
+                await using MemoryStream ms = new MemoryStream();
+                await snapshot.Serialize(ms);
+                snapshotContent = ms.ToArray();
+            }
+
+            BlobIdentifier snapshotBlob = BlobIdentifier.FromBlob(snapshotContent);
+
+            string missingBucket = "this-does-not-exist";
+            Guid missingId = Guid.NewGuid();
+
+            Mock<HttpMessageHandler> handler = new Mock<HttpMessageHandler>();
+            string s = JsonConvert.SerializeObject(new ReplicationLogSnapshots(new List<SnapshotInfo>{new SnapshotInfo(TestNamespace, SnapshotNamespace, snapshotBlob)}));
+            handler.SetupRequest($"http://localhost/api/v1/replication-log/snapshots/{TestNamespace}").ReturnsResponse(s, "application/json");
+
+            // mock a error being generated due to the lastBucket/event being to old
+            handler.SetupRequest($"http://localhost/api/v1/replication-log/incremental/{TestNamespace}?lastBucket={missingBucket}&lastEvent={missingId}").ReturnsResponse(HttpStatusCode.BadRequest, JsonConvert.SerializeObject(new ProblemDetails
+            {
+                Title = $"Log file is not available, use snapshot {snapshotBlob} instead",
+                Type = "http://jupiter.epicgames.com/replication/useSnapshot",
+                Extensions = { { "SnapshotId", snapshotBlob } }
+            }), "application/json");
+
+
+            // after processing the snapshot we do not replicate anything more
+            handler.SetupRequest($"http://localhost/api/v1/replication-log/incremental/{TestNamespace}?lastBucket={snapshotEvents.Last().TimeBucket}&lastEvent={snapshotEvents.Last().EventId}").ReturnsResponse(JsonConvert.SerializeObject(new ReplicationLogEvents(new List<ReplicationLogEvent>())), "application/json");
+
+            handler.SetupRequest($"http://localhost/api/v1/blobs/{TestNamespace}/{snapshotBlob}").ReturnsResponse(snapshotContent, "application/octet-stream").Verifiable();
+
+            foreach (BlobIdentifier blob in blobs.Keys)
+            {
+                handler.SetupRequest($"http://localhost/api/v1/objects/{TestNamespace}/{blob}/references").ReturnsResponse(JsonConvert.SerializeObject(new ResolvedReferencesResult(Array.Empty<BlobIdentifier>())), "application/json").Verifiable();
+            }
+
+            foreach ((BlobIdentifier key, byte[] blobContent) in blobs)
+            {
+                handler.SetupRequest($"http://localhost/api/v1/blobs/{TestNamespace}/{key}").ReturnsResponse(blobContent, "application/octet-stream").Verifiable();
+            }
+
+            IHttpClientFactory httpClientFactory = handler.CreateClientFactory();
+            RefsReplicator replicator = ActivatorUtilities.CreateInstance<RefsReplicator>(_server!.Services, replicatorSettings, httpClientFactory);
+            // specify a bucket that does not exist, we should fallback to replicating a snapshot
+            replicator.SetRefState(missingBucket, missingId);
+
+            bool didRun = await replicator.TriggerNewReplications();
+
+            Assert.IsTrue(didRun);
+
+            handler.Verify();
+
+            // Verify that the objects are present
+            BlobIdentifier[] missingBlobs = await _blobStore.FilterOutKnownBlobs(TestNamespace, blobs.Keys.ToArray());
+            Assert.IsFalse(missingBlobs.Any());
+        }
+    }
+}
