Moved Horde.Storage from git to p4

diff --git a/Engine/Source/Programs/Horde/HordeStorage/Horde.Storage/Tests/Functional/Objects/ReferencesTests.cs b/Engine/Source/Programs/Horde/HordeStorage/Horde.Storage/Tests/Functional/Objects/ReferencesTests.cs
new file mode 100644
--- /dev/null
+++ b/Engine/Source/Programs/Horde/HordeStorage/Horde.Storage/Tests/Functional/Objects/ReferencesTests.cs
@@ -0,0 +1,1228 @@
+// Copyright Epic Games, Inc. All Rights Reserved.
+
+using System;
+using System.Collections.Generic;
+using System.IO;
+using System.Linq;
+using System.Net;
+using System.Net.Http;
+using System.Net.Http.Headers;
+using System.Net.Mime;
+using System.Text;
+using System.Threading.Tasks;
+using Cassandra;
+using Dasync.Collections;
+using Horde.Storage.Controllers;
+using Horde.Storage.Implementation;
+using Jupiter;
+using Jupiter.Implementation;
+using Microsoft.AspNetCore.TestHost;
+using Microsoft.AspNetCore.Hosting;
+using Microsoft.AspNetCore.Mvc;
+using Microsoft.Extensions.Configuration;
+using Microsoft.Extensions.DependencyInjection;
+using Microsoft.VisualStudio.TestTools.UnitTesting;
+using Newtonsoft.Json;
+using Newtonsoft.Json.Linq;
+using Serilog;
+using Logger = Serilog.Core.Logger;
+
+namespace Horde.Storage.FunctionalTests.References
+{
+
+    [TestClass]
+    public class ScyllaReferencesTests : ReferencesTests
+    {
+        protected override IEnumerable<KeyValuePair<string, string>> GetSettings()
+        {
+            return new[]
+            {
+                new KeyValuePair<string, string>("Horde.Storage:ReferencesDbImplementation", HordeStorageSettings.ReferencesDbImplementations.Scylla.ToString()),
+                new KeyValuePair<string, string>("Horde.Storage:ContentIdStoreImplementation", HordeStorageSettings.ContentIdStoreImplementations.Scylla.ToString()),
+                new KeyValuePair<string, string>("Horde.Storage:ReplicationLogWriterImplementation", HordeStorageSettings.ReplicationLogWriterImplementations.Scylla.ToString()),
+            };
+        }
+
+        protected override async Task SeedDb(IServiceProvider provider)
+        {
+            IReferencesStore referencesStore = provider.GetService<IReferencesStore>()!;
+            //verify we are using the expected store
+            Assert.IsTrue(referencesStore.GetType() == typeof(ScyllaReferencesStore));
+
+            IContentIdStore contentIdStore = provider.GetService<IContentIdStore>()!;
+            //verify we are using the expected store
+            Assert.IsTrue(contentIdStore.GetType() == typeof(ScyllaContentIdStore));
+
+            IReplicationLog replicationLog = provider.GetService<IReplicationLog>()!;
+            //verify we are using the replication log writer
+            Assert.IsTrue(replicationLog.GetType() == typeof(ScyllaReplicationLog));
+
+            await SeedTestData(referencesStore);
+        }
+
+        protected override async Task TeardownDb(IServiceProvider provider)
+        {
+            IScyllaSessionManager scyllaSessionManager = provider.GetService<IScyllaSessionManager>()!;
+
+            // Tests generally do not require the database to be empty and creating the tables all the times slows down the tests significantly
+            // Still this can be useful to have for debugging purposes.
+
+            //await session.ExecuteAsync(new SimpleStatement("DROP TABLE IF EXISTS objects"));
+            //await session.ExecuteAsync(new SimpleStatement("DROP TABLE IF EXISTS buckets"));
+
+            // we need to clear out the state we modified in the replication log, otherwise the replication log tests will fail
+            ISession localKeyspace = scyllaSessionManager.GetSessionForLocalKeyspace();
+            await Task.WhenAll(
+                // remove replication log table as we expect it to be empty when starting the tests
+                localKeyspace.ExecuteAsync(new SimpleStatement("DROP TABLE IF EXISTS replication_log;")),
+                // remove the snapshots
+                localKeyspace.ExecuteAsync(new SimpleStatement("DROP TABLE IF EXISTS replication_snapshot;")),
+                // remove the namespaces
+                localKeyspace.ExecuteAsync(new SimpleStatement("DROP TABLE IF EXISTS replication_namespace;"))
+            );
+
+            await Task.CompletedTask;
+        }
+    }
+
+
+    [TestClass]
+    public class MemoryReferencesTests : ReferencesTests
+    {
+        protected override IEnumerable<KeyValuePair<string, string>> GetSettings()
+        {
+            return new[]
+            {
+                new KeyValuePair<string, string>("Horde.Storage:ReferencesDbImplementation", HordeStorageSettings.ReferencesDbImplementations.Memory.ToString()),
+                new KeyValuePair<string, string>("Horde.Storage:ContentIdStoreImplementation", HordeStorageSettings.ContentIdStoreImplementations.Memory.ToString()),
+                new KeyValuePair<string, string>("Horde.Storage:ReplicationLogWriterImplementation", HordeStorageSettings.ReplicationLogWriterImplementations.Memory.ToString()),
+            };
+        }
+
+        protected override async Task SeedDb(IServiceProvider provider)
+        {
+            IReferencesStore referencesStore = provider.GetService<IReferencesStore>()!;
+            //verify we are using the expected refs store
+            Assert.IsTrue(referencesStore.GetType() == typeof(MemoryReferencesStore));
+
+            IContentIdStore contentIdStore = provider.GetService<IContentIdStore>()!;
+            //verify we are using the expected store
+            Assert.IsTrue(contentIdStore.GetType() == typeof(MemoryContentIdStore));
+
+            IReplicationLog replicationLog = provider.GetService<IReplicationLog>()!;
+            //verify we are using the replication log writer
+            Assert.IsTrue(replicationLog.GetType() == typeof(MemoryReplicationLog));
+
+            await SeedTestData(referencesStore);
+        }
+
+        protected override Task TeardownDb(IServiceProvider provider)
+        {
+            return Task.CompletedTask;
+        }
+    }
+    
+    public abstract class ReferencesTests
+    {
+        private static TestServer? _server;
+        private static HttpClient? _httpClient;
+        protected IBlobStore _blobStore = null!;
+        protected IReferencesStore ReferencesStore = null!;
+
+        protected readonly NamespaceId TestNamespace = new NamespaceId("test-namespace");
+
+        [TestInitialize]
+        public async Task Setup()
+
+        {
+            IConfigurationRoot configuration = new ConfigurationBuilder()
+                // we are not reading the base appSettings here as we want exact control over what runs in the tests
+                .AddJsonFile("appsettings.Testing.json", true)
+                .AddInMemoryCollection(GetSettings())
+                .AddEnvironmentVariables()
+                .Build();
+
+            Logger logger = new LoggerConfiguration()
+                .ReadFrom.Configuration(configuration)
+                .CreateLogger();
+
+            TestServer server = new TestServer(new WebHostBuilder()
+                .UseConfiguration(configuration)
+                .UseEnvironment("Testing")
+                .UseSerilog(logger)
+                .UseStartup<HordeStorageStartup>()
+            );
+            _httpClient = server.CreateClient();
+            _server = server;
+
+            _blobStore = _server.Services.GetService<IBlobStore>()!;
+            ReferencesStore = _server.Services.GetService<IReferencesStore>()!;
+            await SeedDb(server.Services);
+        }
+
+        protected async Task SeedTestData(IReferencesStore referencesStore)
+        {
+            await Task.CompletedTask;
+        }
+
+
+        [TestCleanup]
+        public async Task Teardown()
+        {
+            if (_server != null) 
+                await TeardownDb(_server.Services);
+        }
+
+        protected abstract IEnumerable<KeyValuePair<string, string>> GetSettings();
+
+        protected abstract Task SeedDb(IServiceProvider provider);
+        protected abstract Task TeardownDb(IServiceProvider provider);
+
+        [TestMethod]
+        public async Task PutGetBlob()
+        {
+            const string objectContents = "This is treated as a opaque blob";
+            byte[] data = Encoding.ASCII.GetBytes(objectContents);
+            BlobIdentifier objectHash = BlobIdentifier.FromBlob(data);
+
+            HttpContent requestContent = new ByteArrayContent(data);
+            requestContent.Headers.ContentType = new MediaTypeHeaderValue(MediaTypeNames.Application.Octet);
+            requestContent.Headers.Add(CommonHeaders.HashHeaderName, objectHash.ToString());
+
+            HttpResponseMessage result = await _httpClient!.PutAsync(requestUri: $"api/v1/refs/{TestNamespace}/bucket/newBlobObject", requestContent);
+            result.EnsureSuccessStatusCode();
+
+            {
+                HttpResponseMessage getResponse = await _httpClient.GetAsync($"api/v1/refs/{TestNamespace}/bucket/newBlobObject.raw");
+                getResponse.EnsureSuccessStatusCode();
+                await using MemoryStream ms = new MemoryStream();
+                await getResponse.Content.CopyToAsync(ms);
+
+                byte[] roundTrippedBuffer = ms.ToArray();
+                string roundTrippedPayload = Encoding.ASCII.GetString(roundTrippedBuffer);
+
+                Assert.AreEqual(objectContents, roundTrippedPayload);
+                CollectionAssert.AreEqual(data, roundTrippedBuffer);
+                Assert.AreEqual(objectHash, BlobIdentifier.FromBlob(roundTrippedBuffer));
+            }
+
+            {
+                BlobIdentifier attachment;
+                {
+                    HttpResponseMessage getResponse = await _httpClient.GetAsync($"api/v1/refs/{TestNamespace}/bucket/newBlobObject.uecb");
+                    getResponse.EnsureSuccessStatusCode();
+                    await using MemoryStream ms = new MemoryStream();
+                    await getResponse.Content.CopyToAsync(ms);
+
+                    byte[] roundTrippedBuffer = ms.ToArray();
+                    ReadOnlyMemory<byte> localMemory = new ReadOnlyMemory<byte>(roundTrippedBuffer);
+                    CompactBinaryObject cb = CompactBinaryObject.Load(ref localMemory);
+                    List<CompactBinaryField> fields = cb.GetFields().ToList();
+
+                    Assert.AreEqual(1, fields.Count);
+                    CompactBinaryField payloadField = fields[0];
+                    Assert.IsNotNull(payloadField);
+                    Assert.IsTrue(payloadField.IsBinaryAttachment());
+                    attachment = payloadField.AsBinaryAttachment()!;
+                }
+
+                {
+                    HttpResponseMessage getAttachment = await _httpClient.GetAsync($"api/v1/blobs/{TestNamespace}/{attachment}");
+                    getAttachment.EnsureSuccessStatusCode();
+                    await using MemoryStream ms = new MemoryStream();
+                    await getAttachment.Content.CopyToAsync(ms);
+                    byte[] roundTrippedBuffer = ms.ToArray();
+                    string roundTrippedString = Encoding.ASCII.GetString(roundTrippedBuffer);
+
+                    Assert.AreEqual(objectContents, roundTrippedString);
+                    Assert.AreEqual(objectHash, BlobIdentifier.FromBlob(roundTrippedBuffer));
+                }
+            }
+
+            {
+                HttpResponseMessage getResponse = await _httpClient.GetAsync($"api/v1/refs/{TestNamespace}/bucket/newBlobObject.json");
+                getResponse.EnsureSuccessStatusCode();
+                await using MemoryStream ms = new MemoryStream();
+                await getResponse.Content.CopyToAsync(ms);
+
+                byte[] roundTrippedBuffer = ms.ToArray();
+                string s = Encoding.ASCII.GetString(roundTrippedBuffer);
+                JObject jObject = JObject.Parse(s);
+                Assert.AreEqual(1, jObject.Children().Count());
+                JToken? childToken = jObject.Children().First();
+                Assert.IsNotNull(childToken);
+                Assert.AreEqual(JTokenType.Property, childToken.Type);
+                JProperty property = jObject.Children<JProperty>().First();
+                Assert.IsNotNull(property);
+
+                Assert.AreEqual("", property!.Name);
+
+                string value = property.Value.Value<string>()!;
+                Assert.IsNotNull(value);
+                Assert.AreEqual(objectHash, new BlobIdentifier(value));
+            }
+
+            {
+                // request the object as a json response using accept instead of the format filter
+                var request = new HttpRequestMessage(HttpMethod.Get, $"api/v1/refs/{TestNamespace}/bucket/newBlobObject");
+                request.Headers.Accept.Add(new MediaTypeWithQualityHeaderValue(MediaTypeNames.Application.Json));
+                HttpResponseMessage getResponse = await _httpClient.SendAsync(request);
+                getResponse.EnsureSuccessStatusCode();
+                Assert.AreEqual(MediaTypeNames.Application.Json, getResponse.Content.Headers.ContentType?.MediaType);
+
+                await using MemoryStream ms = new MemoryStream();
+                await getResponse.Content.CopyToAsync(ms);
+
+                byte[] roundTrippedBuffer = ms.ToArray();
+                string s = Encoding.ASCII.GetString(roundTrippedBuffer);
+                JObject jObject = JObject.Parse(s);
+                Assert.AreEqual(1, jObject.Children().Count());
+                JToken? childToken = jObject.Children().First();
+                Assert.IsNotNull(childToken);
+                Assert.AreEqual(JTokenType.Property, childToken.Type);
+                JProperty property = jObject.Children<JProperty>().First();
+                Assert.IsNotNull(property);
+
+                Assert.AreEqual("", property!.Name);
+
+                string? value = property.Value.Value<string>()!;
+                Assert.IsNotNull(value);
+                Assert.AreEqual(objectHash, new BlobIdentifier(value));
+            }
+        }
+
+        [TestMethod]
+        public async Task PutGetCompactBinary()
+        {
+            CompactBinaryWriter writer = new CompactBinaryWriter();
+            writer.BeginObject();
+            writer.AddString("thisIsAField", "stringField");
+            writer.EndObject();
+
+            byte[] objectData = writer.Save();
+            BlobIdentifier objectHash = BlobIdentifier.FromBlob(objectData);
+
+            HttpContent requestContent = new ByteArrayContent(objectData);
+            requestContent.Headers.ContentType = new MediaTypeHeaderValue(CustomMediaTypeNames.UnrealCompactBinary);
+            requestContent.Headers.Add(CommonHeaders.HashHeaderName, objectHash.ToString());
+
+            HttpResponseMessage result = await _httpClient!.PutAsync(requestUri: $"api/v1/refs/{TestNamespace}/bucket/newReferenceObject.uecb", requestContent);
+            result.EnsureSuccessStatusCode();
+
+            {
+                Assert.AreEqual(CustomMediaTypeNames.UnrealCompactBinary, result!.Content.Headers.ContentType!.MediaType);
+                // check that no blobs are missing
+                await using MemoryStream ms = new MemoryStream();
+                await result.Content.CopyToAsync(ms);
+                byte[] roundTrippedBuffer = ms.ToArray();
+                ReadOnlyMemory<byte> localMemory = new ReadOnlyMemory<byte>(roundTrippedBuffer);
+                CompactBinaryObject cb = CompactBinaryObject.Load(ref localMemory);
+                CompactBinaryField? needsField = cb["needs"];
+                Assert.IsNotNull(needsField);
+                List<BlobIdentifier?> missingBlobs = needsField!.AsArray().Select(field => field.AsHash()).ToList();
+                Assert.AreEqual(0, missingBlobs.Count);
+            }
+
+            {
+                BucketId bucket = new BucketId("bucket");
+                KeyId key = new KeyId("newReferenceObject");
+
+                ObjectRecord objectRecord = await ReferencesStore.Get(TestNamespace, bucket, key);
+
+                Assert.IsTrue(objectRecord.IsFinalized);
+                Assert.AreEqual("newReferenceObject", objectRecord.Name.ToString());
+                Assert.AreEqual(objectHash, objectRecord.BlobIdentifier);
+                Assert.IsNotNull(objectRecord.InlinePayload);
+            }
+
+            {
+                HttpResponseMessage getResponse = await _httpClient.GetAsync($"api/v1/refs/{TestNamespace}/bucket/newReferenceObject.raw");
+                getResponse.EnsureSuccessStatusCode();
+                await using MemoryStream ms = new MemoryStream();
+                await getResponse.Content.CopyToAsync(ms);
+
+                byte[] roundTrippedBuffer = ms.ToArray();
+
+                CollectionAssert.AreEqual(objectData, roundTrippedBuffer);
+                Assert.AreEqual(objectHash, BlobIdentifier.FromBlob(roundTrippedBuffer));
+            }
+
+            {
+                HttpResponseMessage getResponse = await _httpClient.GetAsync($"api/v1/refs/{TestNamespace}/bucket/newReferenceObject.uecb");
+                getResponse.EnsureSuccessStatusCode();
+                await using MemoryStream ms = new MemoryStream();
+                await getResponse.Content.CopyToAsync(ms);
+
+                byte[] roundTrippedBuffer = ms.ToArray();
+                ReadOnlyMemory<byte> localMemory = new ReadOnlyMemory<byte>(roundTrippedBuffer);
+                CompactBinaryObject cb = CompactBinaryObject.Load(ref localMemory);
+                List<CompactBinaryField> fields = cb.GetFields().ToList();
+
+                Assert.AreEqual(1, fields.Count);
+                CompactBinaryField stringField = fields[0];
+                Assert.AreEqual("stringField", stringField.Name);
+                Assert.AreEqual("thisIsAField", stringField.AsString());
+            }
+
+
+            {
+                HttpResponseMessage getResponse = await _httpClient.GetAsync($"api/v1/refs/{TestNamespace}/bucket/newReferenceObject.json");
+                getResponse.EnsureSuccessStatusCode();
+                await using MemoryStream ms = new MemoryStream();
+                await getResponse.Content.CopyToAsync(ms);
+
+                byte[] roundTrippedBuffer = ms.ToArray();
+
+                string s = Encoding.ASCII.GetString(roundTrippedBuffer);
+                JObject jObject = JObject.Parse(s);
+                Assert.AreEqual(1, jObject.Children().Count());
+
+                JToken? childToken = jObject.Children().First();
+                Assert.IsNotNull(childToken);
+                Assert.AreEqual(JTokenType.Property, childToken.Type);
+                JProperty property = jObject.Children<JProperty>().First();
+                Assert.IsNotNull(property);
+
+                Assert.AreEqual("stringField", property!.Name);
+
+                string? value = property.Value.Value<string>();
+                Assert.IsNotNull(value);
+                Assert.AreEqual("thisIsAField", value);
+            }
+        }
+
+        
+        [TestMethod]
+        public async Task PutGetCompactBinaryHierarchy()
+        {
+            CompactBinaryWriter childObjectWriter = new CompactBinaryWriter();
+            childObjectWriter.BeginObject();
+            childObjectWriter.AddString("thisIsAField", "stringField");
+            childObjectWriter.EndObject();
+            byte[] childObjectData = childObjectWriter.Save();
+            BlobIdentifier childObjectHash = BlobIdentifier.FromBlob(childObjectData);
+
+            CompactBinaryWriter parentObjectWriter = new CompactBinaryWriter();
+            parentObjectWriter.BeginObject();
+            parentObjectWriter.AddCompactBinaryAttachment(childObjectHash, "childObject");
+            parentObjectWriter.EndObject();
+            byte[] parentObjectData = parentObjectWriter.Save();
+            BlobIdentifier parentObjectHash = BlobIdentifier.FromBlob(parentObjectData);
+
+            // this first upload should fail with the child object missing
+            {
+                HttpContent requestContent = new ByteArrayContent(parentObjectData);
+                requestContent.Headers.ContentType = new MediaTypeHeaderValue(CustomMediaTypeNames.UnrealCompactBinary);
+                requestContent.Headers.Add(CommonHeaders.HashHeaderName, parentObjectHash.ToString());
+
+                HttpResponseMessage result = await _httpClient!.PutAsync(requestUri: $"api/v1/refs/{TestNamespace}/bucket/newHierarchyObject.uecb", requestContent);
+                result.EnsureSuccessStatusCode();
+
+                Assert.AreEqual(CustomMediaTypeNames.UnrealCompactBinary, result!.Content.Headers.ContentType!.MediaType);
+                // check that one blobs is missing
+                await using MemoryStream ms = new MemoryStream();
+                await result.Content.CopyToAsync(ms);
+                byte[] roundTrippedBuffer = ms.ToArray();
+                ReadOnlyMemory<byte> localMemory = new ReadOnlyMemory<byte>(roundTrippedBuffer);
+                CompactBinaryObject cb = CompactBinaryObject.Load(ref localMemory);
+                CompactBinaryField? needsField = cb["needs"];
+                Assert.IsNotNull(needsField);
+                List<BlobIdentifier?> missingBlobs = needsField!.AsArray().Select(field => field.AsHash()).ToList();
+                Assert.AreEqual(1, missingBlobs.Count);
+                Assert.AreEqual(childObjectHash, missingBlobs[0]);
+            }
+
+            // upload the child object
+            {
+                HttpRequestMessage request = new HttpRequestMessage(HttpMethod.Put, $"api/v1/objects/{TestNamespace}/{childObjectHash}");
+                request.Headers.Accept.Add(new MediaTypeWithQualityHeaderValue(CustomMediaTypeNames.UnrealCompactBinary));
+
+                HttpContent requestContent = new ByteArrayContent(childObjectData);
+                requestContent.Headers.ContentType = new MediaTypeHeaderValue(CustomMediaTypeNames.UnrealCompactBinary);
+                requestContent.Headers.Add(CommonHeaders.HashHeaderName, childObjectHash.ToString());
+                request.Content = requestContent;
+
+                HttpResponseMessage result = await _httpClient!.SendAsync(request);
+                result.EnsureSuccessStatusCode();
+
+                Assert.AreEqual(CustomMediaTypeNames.UnrealCompactBinary, result!.Content.Headers.ContentType!.MediaType);
+                // check that one blobs is missing
+                await using MemoryStream ms = new MemoryStream();
+                await result.Content.CopyToAsync(ms);
+                byte[] roundTrippedBuffer = ms.ToArray();
+                ReadOnlyMemory<byte> localMemory = new ReadOnlyMemory<byte>(roundTrippedBuffer);
+                CompactBinaryObject cb = CompactBinaryObject.Load(ref localMemory);
+                CompactBinaryField? value = cb["identifier"];
+                Assert.IsNotNull(value);
+                Assert.AreEqual(childObjectHash, value!.AsHash());
+            }
+
+            // since we have now uploaded the child object putting the object again should result in no missing references
+            {
+                HttpContent requestContent = new ByteArrayContent(parentObjectData);
+                requestContent.Headers.ContentType = new MediaTypeHeaderValue(CustomMediaTypeNames.UnrealCompactBinary);
+                requestContent.Headers.Add(CommonHeaders.HashHeaderName, parentObjectHash.ToString());
+
+                HttpResponseMessage result = await _httpClient!.PutAsync(requestUri: $"api/v1/refs/{TestNamespace}/bucket/newHierarchyObject.uecb", requestContent);
+                result.EnsureSuccessStatusCode();
+
+                Assert.AreEqual(result!.Content.Headers.ContentType!.MediaType, CustomMediaTypeNames.UnrealCompactBinary);
+                // check that one blobs is missing
+                await using MemoryStream ms = new MemoryStream();
+                await result.Content.CopyToAsync(ms);
+                byte[] roundTrippedBuffer = ms.ToArray();
+                ReadOnlyMemory<byte> localMemory = new ReadOnlyMemory<byte>(roundTrippedBuffer);
+                CompactBinaryObject cb = CompactBinaryObject.Load(ref localMemory);
+                CompactBinaryField? needsField = cb["needs"];
+                Assert.IsNotNull(needsField);
+                List<BlobIdentifier?> missingBlobs = needsField!.AsArray().Select(field => field.AsHash()).ToList();
+                Assert.AreEqual(0, missingBlobs.Count);
+            }
+
+            {
+                BucketId bucket = new BucketId("bucket");
+                KeyId key = new KeyId("newHierarchyObject");
+
+                ObjectRecord objectRecord = await ReferencesStore.Get(TestNamespace, bucket, key);
+
+                Assert.IsTrue(objectRecord.IsFinalized);
+                Assert.AreEqual("newHierarchyObject", objectRecord.Name.ToString());
+                Assert.AreEqual(parentObjectHash, objectRecord.BlobIdentifier);
+                Assert.IsNotNull(objectRecord.InlinePayload);
+            }
+
+            {
+                HttpResponseMessage getResponse = await _httpClient.GetAsync($"api/v1/refs/{TestNamespace}/bucket/newHierarchyObject.raw");
+                getResponse.EnsureSuccessStatusCode();
+                await using MemoryStream ms = new MemoryStream();
+                await getResponse.Content.CopyToAsync(ms);
+
+                byte[] roundTrippedBuffer = ms.ToArray();
+
+                CollectionAssert.AreEqual(parentObjectData, roundTrippedBuffer);
+                Assert.AreEqual(parentObjectHash, BlobIdentifier.FromBlob(roundTrippedBuffer));
+            }
+
+            {
+                HttpResponseMessage getResponse = await _httpClient.GetAsync($"api/v1/refs/{TestNamespace}/bucket/newHierarchyObject.uecb");
+                getResponse.EnsureSuccessStatusCode();
+                await using MemoryStream ms = new MemoryStream();
+                await getResponse.Content.CopyToAsync(ms);
+
+                byte[] roundTrippedBuffer = ms.ToArray();
+                ReadOnlyMemory<byte> localMemory = new ReadOnlyMemory<byte>(roundTrippedBuffer);
+                CompactBinaryObject cb = CompactBinaryObject.Load(ref localMemory);
+                List<CompactBinaryField> fields = cb.GetFields().ToList();
+
+                Assert.AreEqual(1, fields.Count);
+                CompactBinaryField childObjectField = fields[0];
+                Assert.AreEqual("childObject", childObjectField.Name);
+                Assert.AreEqual(childObjectHash, childObjectField.AsHash());
+            }
+
+
+            {
+                HttpResponseMessage getResponse = await _httpClient.GetAsync($"api/v1/refs/{TestNamespace}/bucket/newHierarchyObject.json");
+                getResponse.EnsureSuccessStatusCode();
+                await using MemoryStream ms = new MemoryStream();
+                await getResponse.Content.CopyToAsync(ms);
+
+                byte[] roundTrippedBuffer = ms.ToArray();
+
+                string s = Encoding.ASCII.GetString(roundTrippedBuffer);
+                JObject jObject = JObject.Parse(s);
+                Assert.AreEqual(1, jObject.Children().Count());
+
+                JToken? childToken = jObject.Children().First();
+                Assert.IsNotNull(childToken);
+                Assert.AreEqual(JTokenType.Property, childToken.Type);
+                JProperty property = jObject.Children<JProperty>().First();
+                Assert.IsNotNull(property);
+
+                Assert.AreEqual("childObject", property!.Name);
+
+                string? value = property.Value.Value<string>();
+                Assert.IsNotNull(value);
+                Assert.AreEqual(childObjectHash.ToString(), value);
+            }
+
+            {
+                HttpResponseMessage getResponse = await _httpClient.GetAsync($"api/v1/objects/{TestNamespace}/{childObjectHash}");
+                getResponse.EnsureSuccessStatusCode();
+                await using MemoryStream ms = new MemoryStream();
+                await getResponse.Content.CopyToAsync(ms);
+
+                byte[] roundTrippedBuffer = ms.ToArray();
+
+                CollectionAssert.AreEqual(childObjectData, roundTrippedBuffer);
+                Assert.AreEqual(childObjectHash, BlobIdentifier.FromBlob(roundTrippedBuffer));
+            }
+
+        }
+
+        [TestMethod]
+        public async Task ExistsChecks()
+        {
+            const string objectContents = "This is treated as a opaque blob";
+            byte[] data = Encoding.ASCII.GetBytes(objectContents);
+            BlobIdentifier objectHash = BlobIdentifier.FromBlob(data);
+
+            HttpContent requestContent = new ByteArrayContent(data);
+            requestContent.Headers.ContentType = new MediaTypeHeaderValue(CustomMediaTypeNames.UnrealCompactBinary);
+            requestContent.Headers.Add(CommonHeaders.HashHeaderName, objectHash.ToString());
+
+            {
+                HttpResponseMessage result = await _httpClient!.PutAsync(requestUri: $"api/v1/refs/{TestNamespace}/bucket/newObject.raw", requestContent);
+                result.EnsureSuccessStatusCode();
+            }
+
+            {
+                HttpResponseMessage result = await _httpClient!.SendAsync(new HttpRequestMessage(HttpMethod.Head, $"api/v1/refs/{TestNamespace}/bucket/newObject"));
+                result.EnsureSuccessStatusCode();
+            }
+
+            {
+                HttpResponseMessage result = await _httpClient!.SendAsync(new HttpRequestMessage(HttpMethod.Head, $"api/v1/refs/{TestNamespace}/bucket/missingObject"));
+                Assert.AreEqual(HttpStatusCode.NotFound, result.StatusCode);
+            }
+        }
+
+
+        [TestMethod]
+        public async Task PutGetObjectHierarchy()
+        {
+            string blobContents = "This is a string that is referenced as a blob";
+            byte[] blobData = Encoding.ASCII.GetBytes(blobContents);
+            BlobIdentifier blobHash = BlobIdentifier.FromBlob(blobData);
+            await _blobStore.PutObject(TestNamespace, blobData, blobHash);
+
+            string blobContentsChild = "This string is also referenced as a blob but from a child object";
+            byte[] dataChild = Encoding.ASCII.GetBytes(blobContentsChild);
+            BlobIdentifier blobHashChild = BlobIdentifier.FromBlob(dataChild);
+            await _blobStore.PutObject(TestNamespace, dataChild, blobHashChild);
+
+            CompactBinaryWriter writerChild = new CompactBinaryWriter();
+            writerChild.BeginObject();
+            writerChild.AddBinaryAttachment(blobHashChild, "blob");
+            writerChild.EndObject();
+
+            byte[] childDataObject = writerChild.Save();
+            BlobIdentifier childDataObjectHash = BlobIdentifier.FromBlob(childDataObject);
+            await _blobStore.PutObject(TestNamespace, childDataObject, childDataObjectHash);
+
+            CompactBinaryWriter writerParent = new CompactBinaryWriter();
+            writerParent.BeginObject();
+            
+            writerParent.AddBinaryAttachment(blobHash, "blobAttachment");
+            writerParent.AddCompactBinaryAttachment(childDataObjectHash, "objectAttachment");
+            writerParent.EndObject();
+
+            byte[] objectData = writerParent.Save();
+            BlobIdentifier objectHash = BlobIdentifier.FromBlob(objectData);
+
+            HttpContent requestContent = new ByteArrayContent(objectData);
+            requestContent.Headers.ContentType = new MediaTypeHeaderValue(CustomMediaTypeNames.UnrealCompactBinary);
+            requestContent.Headers.Add(CommonHeaders.HashHeaderName, objectHash.ToString());
+
+            HttpResponseMessage result = await _httpClient!.PutAsync(requestUri: $"api/v1/refs/{TestNamespace}/bucket/newHierarchyObject.uecb", requestContent);
+            result.EnsureSuccessStatusCode();
+
+            // check the response
+            {
+                Assert.AreEqual(result!.Content.Headers.ContentType!.MediaType, CustomMediaTypeNames.UnrealCompactBinary);
+                // check that no blobs are missing
+                await using MemoryStream ms = new MemoryStream();
+                await result.Content.CopyToAsync(ms);
+                byte[] roundTrippedBuffer = ms.ToArray();
+                ReadOnlyMemory<byte> localMemory = new ReadOnlyMemory<byte>(roundTrippedBuffer);
+                CompactBinaryObject cb = CompactBinaryObject.Load(ref localMemory);
+                CompactBinaryField? needsField = cb["needs"];
+                Assert.IsNotNull(needsField);
+                List<BlobIdentifier?> missingBlobs = needsField!.AsArray().Select(field => field.AsHash()).ToList();
+                Assert.AreEqual(0, missingBlobs.Count);
+            }
+
+            // check that actual internal representation
+            {
+                BucketId bucket = new BucketId("bucket");
+                KeyId key = new KeyId("newHierarchyObject");
+
+                ObjectRecord objectRecord = await ReferencesStore.Get(TestNamespace, bucket, key);
+
+                Assert.IsTrue(objectRecord.IsFinalized);
+                Assert.AreEqual("newHierarchyObject", objectRecord.Name.ToString());
+                Assert.AreEqual(objectHash, objectRecord.BlobIdentifier);
+                Assert.IsNotNull(objectRecord.InlinePayload);
+            }
+
+            // verify attachments
+            {
+                HttpResponseMessage getResponse = await _httpClient.GetAsync($"api/v1/refs/{TestNamespace}/bucket/newHierarchyObject.raw");
+                getResponse.EnsureSuccessStatusCode();
+                await using MemoryStream ms = new MemoryStream();
+                await getResponse.Content.CopyToAsync(ms);
+
+                byte[] roundTrippedBuffer = ms.ToArray();
+
+                CollectionAssert.AreEqual(objectData, roundTrippedBuffer);
+                Assert.AreEqual(objectHash, BlobIdentifier.FromBlob(roundTrippedBuffer));
+            }
+
+            {
+                BlobIdentifier blobAttachment;
+                BlobIdentifier objectAttachment;
+                {
+                    HttpResponseMessage getResponse = await _httpClient.GetAsync($"api/v1/refs/{TestNamespace}/bucket/newHierarchyObject.uecb");
+                    getResponse.EnsureSuccessStatusCode();
+                    await using MemoryStream ms = new MemoryStream();
+                    await getResponse.Content.CopyToAsync(ms);
+
+                    byte[] roundTrippedBuffer = ms.ToArray();
+                    ReadOnlyMemory<byte> localMemory = new ReadOnlyMemory<byte>(roundTrippedBuffer);
+                    CompactBinaryObject cb = CompactBinaryObject.Load(ref localMemory);
+                    Assert.AreEqual(2, cb.GetFields().Count());
+
+                    CompactBinaryField? blobAttachmentField = cb["blobAttachment"];
+                    Assert.IsNotNull(blobAttachmentField);
+                    blobAttachment = blobAttachmentField!.AsBinaryAttachment()!;
+                    CompactBinaryField? objectAttachmentField = cb["objectAttachment"];
+                    Assert.IsNotNull(objectAttachmentField);
+                    objectAttachment = objectAttachmentField!.AsCompactBinaryAttachment()!;
+                }
+
+                {
+                    HttpResponseMessage getAttachment = await _httpClient.GetAsync($"api/v1/blobs/{TestNamespace}/{blobAttachment}");
+                    getAttachment.EnsureSuccessStatusCode();
+                    await using MemoryStream ms = new MemoryStream();
+                    await getAttachment.Content.CopyToAsync(ms);
+                    byte[] roundTrippedBuffer = ms.ToArray();
+                    string roundTrippedString = Encoding.ASCII.GetString(roundTrippedBuffer);
+
+                    Assert.AreEqual(blobContents, roundTrippedString);
+                    Assert.AreEqual(blobHash, BlobIdentifier.FromBlob(roundTrippedBuffer));
+                }
+
+                BlobIdentifier attachedBlobIdentifier;
+                {
+                    HttpResponseMessage getAttachment = await _httpClient.GetAsync($"api/v1/blobs/{TestNamespace}/{objectAttachment}");
+                    getAttachment.EnsureSuccessStatusCode();
+                    await using MemoryStream ms = new MemoryStream();
+                    await getAttachment.Content.CopyToAsync(ms);
+                    byte[] roundTrippedBuffer = ms.ToArray();
+                    ReadOnlyMemory<byte> localMemory = new ReadOnlyMemory<byte>(roundTrippedBuffer);
+                    CompactBinaryObject cb = CompactBinaryObject.Load(ref localMemory);
+                    Assert.AreEqual(1, cb.GetFields().Count());
+
+                    CompactBinaryField? blobField = cb["blob"];
+                    Assert.IsNotNull(blobField);
+
+                    attachedBlobIdentifier = blobField!.AsBinaryAttachment()!;
+                }
+
+                {
+                    HttpResponseMessage getAttachment = await _httpClient.GetAsync($"api/v1/blobs/{TestNamespace}/{attachedBlobIdentifier}");
+                    getAttachment.EnsureSuccessStatusCode();
+                    await using MemoryStream ms = new MemoryStream();
+                    await getAttachment.Content.CopyToAsync(ms);
+                    byte[] roundTrippedBuffer = ms.ToArray();
+                    string roundTrippedString = Encoding.ASCII.GetString(roundTrippedBuffer);
+
+                    Assert.AreEqual(blobContentsChild, roundTrippedString);
+                    Assert.AreEqual(blobHashChild, BlobIdentifier.FromBlob(roundTrippedBuffer));
+                }
+            }
+
+            // check json representation
+            {
+                HttpResponseMessage getResponse = await _httpClient.GetAsync($"api/v1/refs/{TestNamespace}/bucket/newHierarchyObject.json");
+                getResponse.EnsureSuccessStatusCode();
+                await using MemoryStream ms = new MemoryStream();
+                await getResponse.Content.CopyToAsync(ms);
+
+                byte[] roundTrippedBuffer = ms.ToArray();
+                string s = Encoding.ASCII.GetString(roundTrippedBuffer);
+                JObject jObject = JObject.Parse(s);
+                Assert.AreEqual(2, jObject.Children().Count());
+
+                JToken? blobAttachment = jObject["blobAttachment"];
+                string blobAttachmentString = blobAttachment!.Value<string>()!;
+                Assert.AreEqual(blobHash, new BlobIdentifier(blobAttachmentString));
+
+                JToken? objectAttachment = jObject["objectAttachment"];
+                string objectAttachmentString = objectAttachment!.Value<string>()!;
+                Assert.AreEqual(childDataObjectHash, new BlobIdentifier(objectAttachmentString));
+
+            }
+        }
+
+
+        
+        [TestMethod]
+        public async Task PutPartialHierarchy()
+        {
+            // do not submit the content of the blobs, which should be reported in the response of the put
+            string blobContents = "This is a string that is referenced as a blob";
+            byte[] blobData = Encoding.ASCII.GetBytes(blobContents);
+            BlobIdentifier blobHash = BlobIdentifier.FromBlob(blobData);
+
+            string blobContentsChild = "This string is also referenced as a blob but from a child object";
+            byte[] dataChild = Encoding.ASCII.GetBytes(blobContentsChild);
+            BlobIdentifier blobHashChild = BlobIdentifier.FromBlob(dataChild);
+
+            CompactBinaryWriter writerChild = new CompactBinaryWriter();
+            writerChild.BeginObject();
+            writerChild.AddBinaryAttachment(blobHashChild, "blob");
+            writerChild.EndObject();
+
+            byte[] childDataObject = writerChild.Save();
+            BlobIdentifier childDataObjectHash = BlobIdentifier.FromBlob(childDataObject);
+            await _blobStore.PutObject(TestNamespace, childDataObject, childDataObjectHash);
+
+            CompactBinaryWriter writerParent = new CompactBinaryWriter();
+            writerParent.BeginObject();
+            
+            writerParent.AddBinaryAttachment(blobHash, "blobAttachment");
+            writerParent.AddCompactBinaryAttachment(childDataObjectHash, "objectAttachment");
+            writerParent.EndObject();
+
+            byte[] objectData = writerParent.Save();
+            BlobIdentifier objectHash = BlobIdentifier.FromBlob(objectData);
+
+            {
+                HttpContent requestContent = new ByteArrayContent(objectData);
+                requestContent.Headers.ContentType = new MediaTypeHeaderValue(CustomMediaTypeNames.UnrealCompactBinary);
+                requestContent.Headers.Add(CommonHeaders.HashHeaderName, objectHash.ToString());
+
+                HttpResponseMessage result = await _httpClient!.PutAsync(requestUri: $"api/v1/refs/{TestNamespace}/bucket/newHierarchyObject.uecb", requestContent);
+                result.EnsureSuccessStatusCode();
+
+                {
+                    Assert.AreEqual(result!.Content.Headers.ContentType!.MediaType, CustomMediaTypeNames.UnrealCompactBinary);
+                    await using MemoryStream ms = new MemoryStream();
+                    await result.Content.CopyToAsync(ms);
+                    byte[] roundTrippedBuffer = ms.ToArray();
+                    ReadOnlyMemory<byte> localMemory = new ReadOnlyMemory<byte>(roundTrippedBuffer);
+                    CompactBinaryObject cb = CompactBinaryObject.Load(ref localMemory);
+                    CompactBinaryField? needsField = cb["needs"];
+                    Assert.IsNotNull(needsField);
+                    List<BlobIdentifier?> missingBlobs = needsField!.AsArray().Select(field => field.AsHash()).ToList();
+                    Assert.AreEqual(2, missingBlobs.Count);
+                    Assert.IsTrue(missingBlobs.Contains(blobHash));
+                    Assert.IsTrue(missingBlobs.Contains(blobHashChild));
+                }
+            }
+
+            {
+                HttpContent requestContent = new ByteArrayContent(objectData);
+                requestContent.Headers.ContentType = new MediaTypeHeaderValue(CustomMediaTypeNames.UnrealCompactBinary);
+                requestContent.Headers.Add(CommonHeaders.HashHeaderName, objectHash.ToString());
+
+                HttpResponseMessage result = await _httpClient!.PutAsync(requestUri: $"api/v1/refs/{TestNamespace}/bucket/newHierarchyObject.json", requestContent);
+                result.EnsureSuccessStatusCode();
+
+                {
+                    Assert.AreEqual(result!.Content.Headers.ContentType!.MediaType, MediaTypeNames.Application.Json);
+                    await using MemoryStream ms = new MemoryStream();
+                    await result.Content.CopyToAsync(ms);
+                    byte[] roundTrippedBuffer = ms.ToArray();
+                    string s = Encoding.ASCII.GetString(roundTrippedBuffer);
+                    JObject jObject = JObject.Parse(s);
+                    Assert.AreEqual(1, jObject.Children().Count());
+
+                    JToken? needs = jObject["needs"];
+                    BlobIdentifier[] missingBlobs = needs!.ToArray().Select(token => new BlobIdentifier(token.Value<string>()!)).ToArray();
+                    Assert.AreEqual(2, missingBlobs.Length);
+                    Assert.IsTrue(missingBlobs.Contains(blobHash));
+                    Assert.IsTrue(missingBlobs.Contains(blobHashChild));
+                }
+            }
+        }
+
+        [TestMethod]
+        public async Task PutAndFinalize()
+        {
+            BucketId bucket = new BucketId("bucket");
+            KeyId key = new KeyId("willFinalizeObject");
+
+            // do not submit the content of the blobs, which should be reported in the response of the put
+            string blobContents = "This is a string that is referenced as a blob";
+            byte[] blobData = Encoding.ASCII.GetBytes(blobContents);
+            BlobIdentifier blobHash = BlobIdentifier.FromBlob(blobData);
+
+            string blobContentsChild = "This string is also referenced as a blob but from a child object";
+            byte[] dataChild = Encoding.ASCII.GetBytes(blobContentsChild);
+            BlobIdentifier blobHashChild = BlobIdentifier.FromBlob(dataChild);
+
+            CompactBinaryWriter writerChild = new CompactBinaryWriter();
+            writerChild.BeginObject();
+            writerChild.AddBinaryAttachment(blobHashChild, "blob");
+            writerChild.EndObject();
+
+            byte[] childDataObject = writerChild.Save();
+            BlobIdentifier childDataObjectHash = BlobIdentifier.FromBlob(childDataObject);
+            await _blobStore.PutObject(TestNamespace, childDataObject, childDataObjectHash);
+
+            CompactBinaryWriter writerParent = new CompactBinaryWriter();
+            writerParent.BeginObject();
+            
+            writerParent.AddBinaryAttachment(blobHash, "blobAttachment");
+            writerParent.AddCompactBinaryAttachment(childDataObjectHash, "objectAttachment");
+            writerParent.EndObject();
+
+            byte[] objectData = writerParent.Save();
+            BlobIdentifier objectHash = BlobIdentifier.FromBlob(objectData);
+
+            {
+                HttpContent requestContent = new ByteArrayContent(objectData);
+                requestContent.Headers.ContentType = new MediaTypeHeaderValue(CustomMediaTypeNames.UnrealCompactBinary);
+                requestContent.Headers.Add(CommonHeaders.HashHeaderName, objectHash.ToString());
+
+                HttpResponseMessage result = await _httpClient!.PutAsync(requestUri: $"api/v1/refs/{TestNamespace}/bucket/willFinalizeObject.uecb", requestContent);
+                result.EnsureSuccessStatusCode();
+
+                {
+                    Assert.AreEqual(result!.Content.Headers.ContentType!.MediaType, CustomMediaTypeNames.UnrealCompactBinary);
+                    await using MemoryStream ms = new MemoryStream();
+                    await result.Content.CopyToAsync(ms);
+                    byte[] roundTrippedBuffer = ms.ToArray();
+                    ReadOnlyMemory<byte> localMemory = new ReadOnlyMemory<byte>(roundTrippedBuffer);
+                    CompactBinaryObject cb = CompactBinaryObject.Load(ref localMemory);
+                    CompactBinaryField? needsField = cb["needs"];
+                    Assert.IsNotNull(needsField);
+                    List<BlobIdentifier?> missingBlobs = needsField!.AsArray().Select(field => field.AsHash()).ToList();
+                    Assert.AreEqual(2, missingBlobs.Count);
+                    Assert.IsTrue(missingBlobs.Contains(blobHash));
+                    Assert.IsTrue(missingBlobs.Contains(blobHashChild));
+                }
+            }
+
+            // check that actual internal representation
+            {
+                ObjectRecord objectRecord = await ReferencesStore.Get(TestNamespace, bucket, key);
+
+                Assert.IsFalse(objectRecord.IsFinalized);
+                Assert.AreEqual("willFinalizeObject", objectRecord.Name.ToString());
+                Assert.AreEqual(objectHash, objectRecord.BlobIdentifier);
+                Assert.IsNotNull(objectRecord.InlinePayload);
+            }
+
+            // upload missing pieces
+            {
+                await _blobStore.PutObject(TestNamespace, blobData, blobHash);
+                await _blobStore.PutObject(TestNamespace, dataChild, blobHashChild);
+            }
+
+            // finalize the object as no pieces is now missing
+            {
+                HttpContent requestContent = new ByteArrayContent(Array.Empty<byte>());
+
+                HttpResponseMessage result = await _httpClient!.PostAsync(requestUri: $"api/v1/refs/{TestNamespace}/bucket/willFinalizeObject/finalize/{objectHash}.uecb", requestContent);
+                result.EnsureSuccessStatusCode();
+
+                {
+                    Assert.AreEqual(result!.Content.Headers.ContentType!.MediaType, CustomMediaTypeNames.UnrealCompactBinary);
+                    await using MemoryStream ms = new MemoryStream();
+                    await result.Content.CopyToAsync(ms);
+                    byte[] roundTrippedBuffer = ms.ToArray();
+                    ReadOnlyMemory<byte> localMemory = new ReadOnlyMemory<byte>(roundTrippedBuffer);
+                    CompactBinaryObject cb = CompactBinaryObject.Load(ref localMemory);
+                    CompactBinaryField? needsField = cb["needs"];
+                    Assert.IsNotNull(needsField);
+                    List<BlobIdentifier?> missingBlobs = needsField!.AsArray().Select(field => field.AsHash()).ToList();
+                    Assert.AreEqual(0, missingBlobs.Count);
+                }
+            }
+
+            // check that actual internal representation has updated its state
+            {
+                ObjectRecord objectRecord = await ReferencesStore.Get(TestNamespace, bucket, key);
+
+                Assert.IsTrue(objectRecord.IsFinalized);
+                Assert.AreEqual("willFinalizeObject", objectRecord.Name.ToString());
+                Assert.AreEqual(objectHash, objectRecord.BlobIdentifier);
+            }
+        }
+
+        [TestMethod]
+        public async Task GetMissingBlobRecord()
+        {
+            string blobContents = "This is a blob";
+            byte[] blobData = Encoding.ASCII.GetBytes(blobContents);
+            BlobIdentifier blobHash = BlobIdentifier.FromBlob(blobData);
+
+            HttpContent requestContent = new ByteArrayContent(blobData);
+            requestContent.Headers.ContentType = new MediaTypeHeaderValue(MediaTypeNames.Application.Octet);
+            requestContent.Headers.Add(CommonHeaders.HashHeaderName, blobHash.ToString());
+
+            HttpResponseMessage result = await _httpClient!.PutAsync(requestUri: $"api/v1/refs/{TestNamespace}/bucket/newReferenceObject.uecb", requestContent);
+            result.EnsureSuccessStatusCode();
+
+            {
+                Assert.AreEqual(CustomMediaTypeNames.UnrealCompactBinary, result!.Content.Headers.ContentType!.MediaType);
+                // check that no blobs are missing
+                await using MemoryStream ms = new MemoryStream();
+                await result.Content.CopyToAsync(ms);
+                byte[] roundTrippedBuffer = ms.ToArray();
+                ReadOnlyMemory<byte> localMemory = new ReadOnlyMemory<byte>(roundTrippedBuffer);
+                CompactBinaryObject cb = CompactBinaryObject.Load(ref localMemory);
+                CompactBinaryField? needsField = cb["needs"];
+                Assert.IsNotNull(needsField);
+                List<BlobIdentifier?> missingBlobs = needsField!.AsArray().Select(field => field.AsHash()).ToList();
+                Assert.AreEqual(0, missingBlobs.Count);
+            }
+
+            {
+                // verify we can fetch the blob properly
+                HttpResponseMessage getResponse = await _httpClient.GetAsync($"api/v1/refs/{TestNamespace}/bucket/newReferenceObject.raw");
+                getResponse.EnsureSuccessStatusCode();
+                await using MemoryStream ms = new MemoryStream();
+                await getResponse.Content.CopyToAsync(ms);
+
+                byte[] roundTrippedBuffer = ms.ToArray();
+
+                CollectionAssert.AreEqual(blobData, roundTrippedBuffer);
+                Assert.AreEqual(blobHash, BlobIdentifier.FromBlob(roundTrippedBuffer));
+            }
+
+            {
+                // delete the blob 
+                await _blobStore.DeleteObject(TestNamespace, blobHash);
+            }
+
+            {
+                // the compact binary object still exists, so this returns a success (trying to resolve the attachment will fail)
+                HttpResponseMessage getResponse = await _httpClient.GetAsync($"api/v1/refs/{TestNamespace}/bucket/newReferenceObject.uecb");
+                getResponse.EnsureSuccessStatusCode();
+            }
+
+            {
+                // the compact binary object still exists, so this returns a success (trying to resolve the attachment will fail)
+                HttpResponseMessage getResponse = await _httpClient.GetAsync($"api/v1/refs/{TestNamespace}/bucket/newReferenceObject.json");
+                getResponse.EnsureSuccessStatusCode();
+            }
+
+            {
+                // we should now see a 404 as the blob is missing
+                HttpResponseMessage getResponse = await _httpClient.GetAsync($"api/v1/refs/{TestNamespace}/bucket/newReferenceObject.raw");
+                Assert.AreEqual(HttpStatusCode.NotFound, getResponse.StatusCode);
+                Assert.AreEqual("application/problem+json", getResponse.Content.Headers.ContentType!.MediaType);
+                string s = await getResponse.Content.ReadAsStringAsync();
+                ProblemDetails? problem = JsonConvert.DeserializeObject<ProblemDetails>(s);
+                Assert.IsNotNull(problem);
+                Assert.AreEqual($"Object {blobHash} in {TestNamespace} not found", problem.Title);
+            }
+        }
+        
+        
+        [TestMethod]
+        public async Task DeleteObject()
+        {
+            const string objectContents = "This is treated as a opaque blob";
+            byte[] data = Encoding.ASCII.GetBytes(objectContents);
+            BlobIdentifier objectHash = BlobIdentifier.FromBlob(data);
+
+            HttpContent requestContent = new ByteArrayContent(data);
+            requestContent.Headers.ContentType = new MediaTypeHeaderValue(MediaTypeNames.Application.Octet);
+            requestContent.Headers.Add(CommonHeaders.HashHeaderName, objectHash.ToString());
+
+            // submit the object
+            {
+                HttpResponseMessage result = await _httpClient!.PutAsync(requestUri: $"api/v1/refs/{TestNamespace}/bucket/deletableObject", requestContent);
+                result.EnsureSuccessStatusCode();
+            }
+
+            // verify it is present
+            {
+                HttpResponseMessage result = await _httpClient!.GetAsync(requestUri: $"api/v1/refs/{TestNamespace}/bucket/deletableObject.raw");
+                result.EnsureSuccessStatusCode();
+            }
+
+            // delete the object
+            {
+                HttpResponseMessage result = await _httpClient!.DeleteAsync(requestUri: $"api/v1/refs/{TestNamespace}/bucket/deletableObject");
+                result.EnsureSuccessStatusCode();
+            }
+
+            // ensure the object is not present anymore
+            {
+                HttpResponseMessage result = await _httpClient!.GetAsync(requestUri: $"api/v1/refs/{TestNamespace}/bucket/deletableObject.raw");
+                Assert.AreEqual(HttpStatusCode.NotFound, result.StatusCode);
+            }
+        }
+
+        
+        [TestMethod]
+        public async Task DropBucket()
+        {
+            const string BucketToDelete = "delete-bucket";
+
+            const string objectContents = "This is treated as a opaque blob";
+            byte[] data = Encoding.ASCII.GetBytes(objectContents);
+            BlobIdentifier objectHash = BlobIdentifier.FromBlob(data);
+
+            HttpContent requestContent = new ByteArrayContent(data);
+            requestContent.Headers.ContentType = new MediaTypeHeaderValue(MediaTypeNames.Application.Octet);
+            requestContent.Headers.Add(CommonHeaders.HashHeaderName, objectHash.ToString());
+
+            // submit the object into multiple buckets
+            {
+                HttpResponseMessage result = await _httpClient!.PutAsync(requestUri: $"api/v1/refs/{TestNamespace}/{BucketToDelete}/deletableObject", requestContent);
+                result.EnsureSuccessStatusCode();
+            }
+
+            {
+                HttpResponseMessage result = await _httpClient!.PutAsync(requestUri: $"api/v1/refs/{TestNamespace}/bucket/deletableObject", requestContent);
+                result.EnsureSuccessStatusCode();
+            }
+
+            // verify it is present
+            {
+                HttpResponseMessage result = await _httpClient!.GetAsync(requestUri: $"api/v1/refs/{TestNamespace}/bucket/deletableObject.raw");
+                result.EnsureSuccessStatusCode();
+            }
+
+            // verify it is present
+            {
+                HttpResponseMessage result = await _httpClient!.GetAsync(requestUri: $"api/v1/refs/{TestNamespace}/{BucketToDelete}/deletableObject.raw");
+                result.EnsureSuccessStatusCode();
+            }
+
+            // delete the bucket
+            {
+                HttpResponseMessage result = await _httpClient!.DeleteAsync(requestUri: $"api/v1/refs/{TestNamespace}/{BucketToDelete}");
+                result.EnsureSuccessStatusCode();
+            }
+
+            // ensure the object is present in not deleted bucket
+            {
+                HttpResponseMessage result = await _httpClient!.GetAsync(requestUri: $"api/v1/refs/{TestNamespace}/bucket/deletableObject.raw");
+                result.EnsureSuccessStatusCode();
+            }
+
+            // ensure the object is not present anymore
+            {
+                HttpResponseMessage result = await _httpClient!.GetAsync(requestUri: $"api/v1/refs/{TestNamespace}/{BucketToDelete}/deletableObject.raw");
+                Assert.AreEqual(HttpStatusCode.NotFound, result.StatusCode);
+            }
+        }
+
+        
+        [TestMethod]
+        public async Task DeleteNamespace()
+        {
+            const string NamespaceToBeDeleted = "test-delete-namespace";
+
+            const string objectContents = "This is treated as a opaque blob";
+            byte[] data = Encoding.ASCII.GetBytes(objectContents);
+            BlobIdentifier objectHash = BlobIdentifier.FromBlob(data);
+
+            HttpContent requestContent = new ByteArrayContent(data);
+            requestContent.Headers.ContentType = new MediaTypeHeaderValue(MediaTypeNames.Application.Octet);
+            requestContent.Headers.Add(CommonHeaders.HashHeaderName, objectHash.ToString());
+
+            // submit the object into multiple namespaces
+            {
+                HttpResponseMessage result = await _httpClient!.PutAsync(requestUri: $"api/v1/refs/{TestNamespace}/bucket/deletableObject", requestContent);
+                result.EnsureSuccessStatusCode();
+            }
+
+            {
+                HttpResponseMessage result = await _httpClient!.PutAsync(requestUri: $"api/v1/refs/{NamespaceToBeDeleted}/bucket/deletableObject", requestContent);
+                result.EnsureSuccessStatusCode();
+            }
+
+            // verify it is present
+            {
+                HttpResponseMessage result = await _httpClient!.GetAsync(requestUri: $"api/v1/refs/{TestNamespace}/bucket/deletableObject.raw");
+                result.EnsureSuccessStatusCode();
+            }
+
+            // verify it is present
+            {
+                HttpResponseMessage result = await _httpClient!.GetAsync(requestUri: $"api/v1/refs/{NamespaceToBeDeleted}/bucket/deletableObject.raw");
+                result.EnsureSuccessStatusCode();
+            }
+
+            // delete the namespace
+            {
+                HttpResponseMessage result = await _httpClient!.DeleteAsync(requestUri: $"api/v1/refs/{NamespaceToBeDeleted}");
+                result.EnsureSuccessStatusCode();
+            }
+
+            // ensure the object is present in not deleted namespace
+            {
+                HttpResponseMessage result = await _httpClient!.GetAsync(requestUri: $"api/v1/refs/{TestNamespace}/bucket/deletableObject.raw");
+                result.EnsureSuccessStatusCode();
+            }
+
+            // ensure the object is not present anymore
+            {
+                HttpResponseMessage result = await _httpClient!.GetAsync(requestUri: $"api/v1/refs/{NamespaceToBeDeleted}/bucket/deletableObject.raw");
+                Assert.AreEqual(HttpStatusCode.NotFound, result.StatusCode);
+            }
+
+            // make sure the namespace is not considered a valid namespace anymore
+            {
+                HttpResponseMessage result = await _httpClient!.GetAsync(requestUri: $"api/v1/refs");
+                result.EnsureSuccessStatusCode();
+                Assert.AreEqual(result!.Content.Headers.ContentType!.MediaType, MediaTypeNames.Application.Json);
+                GetNamespacesResponse response = await result.Content.ReadAsAsync<GetNamespacesResponse>();
+                CollectionAssert.DoesNotContain(response.Namespaces, NamespaceToBeDeleted);
+            }
+        }
+
+
+        
+        [TestMethod]
+        public async Task ListNamespaces()
+        {
+            const string objectContents = "This is treated as a opaque blob";
+            byte[] data = Encoding.ASCII.GetBytes(objectContents);
+            BlobIdentifier objectHash = BlobIdentifier.FromBlob(data);
+
+            HttpContent requestContent = new ByteArrayContent(data);
+            requestContent.Headers.ContentType = new MediaTypeHeaderValue(MediaTypeNames.Application.Octet);
+            requestContent.Headers.Add(CommonHeaders.HashHeaderName, objectHash.ToString());
+
+            // submit a object to make sure a namespace is created
+            {
+                HttpResponseMessage result = await _httpClient!.PutAsync(requestUri: $"api/v1/refs/{TestNamespace}/bucket/notUsedObject", requestContent);
+                result.EnsureSuccessStatusCode();
+            }
+
+            {
+                HttpResponseMessage result = await _httpClient!.GetAsync(requestUri: $"api/v1/refs");
+                result.EnsureSuccessStatusCode();
+                Assert.AreEqual(result!.Content.Headers.ContentType!.MediaType, MediaTypeNames.Application.Json);
+                GetNamespacesResponse response = await result.Content.ReadAsAsync<GetNamespacesResponse>();
+                Assert.IsTrue(response.Namespaces.Contains(TestNamespace));
+            }
+        }
+
+        
+        [TestMethod]
+        public async Task GetOldRecords()
+        {
+            const string objectContents = "This is treated as a opaque blob";
+            byte[] data = Encoding.ASCII.GetBytes(objectContents);
+            BlobIdentifier objectHash = BlobIdentifier.FromBlob(data);
+
+            HttpContent requestContent = new ByteArrayContent(data);
+            requestContent.Headers.ContentType = new MediaTypeHeaderValue(MediaTypeNames.Application.Octet);
+            requestContent.Headers.Add(CommonHeaders.HashHeaderName, objectHash.ToString());
+
+            // submit some contents
+            {
+                HttpResponseMessage result = await _httpClient!.PutAsync(requestUri: $"api/v1/refs/{TestNamespace}/bucket/oldRecord", requestContent);
+                result.EnsureSuccessStatusCode();
+            }
+
+            List<ObjectRecord> records = await ReferencesStore.GetOldestRecords(TestNamespace).ToListAsync();
+
+            ObjectRecord? oldRecord = records.Find(record => record.Name == new KeyId("oldRecord"));
+            Assert.IsNotNull(oldRecord);
+            Assert.AreEqual("oldRecord", oldRecord!.Name.ToString());
+            Assert.AreEqual("bucket", oldRecord.Bucket.ToString());
+        }
+    }
+}
