Moved Horde.Storage from git to p4

diff --git a/Engine/Source/Programs/Horde/HordeStorage/Horde.Storage/Tests/Functional/Objects/ReplicationTests.cs b/Engine/Source/Programs/Horde/HordeStorage/Horde.Storage/Tests/Functional/Objects/ReplicationTests.cs
new file mode 100644
--- /dev/null
+++ b/Engine/Source/Programs/Horde/HordeStorage/Horde.Storage/Tests/Functional/Objects/ReplicationTests.cs
@@ -0,0 +1,744 @@
+// Copyright Epic Games, Inc. All Rights Reserved.
+
+using System;
+using System.Collections.Generic;
+using System.Linq;
+using System.Net;
+using System.Net.Http;
+using System.Net.Http.Headers;
+using System.Net.Http.Json;
+using System.Net.Mime;
+using System.Threading.Tasks;
+using async_enumerable_dotnet;
+using Cassandra;
+using Dasync.Collections;
+using Horde.Storage.Controllers;
+using Horde.Storage.Implementation;
+using Horde.Storage.Implementation.TransactionLog;
+using Jupiter;
+using Jupiter.Implementation;
+using Microsoft.AspNetCore.TestHost;
+using Microsoft.AspNetCore.Hosting;
+using Microsoft.AspNetCore.Mvc;
+using Microsoft.Extensions.Configuration;
+using Microsoft.Extensions.DependencyInjection;
+using Microsoft.VisualStudio.TestTools.UnitTesting;
+using Serilog;
+using Logger = Serilog.Core.Logger;
+
+namespace Horde.Storage.FunctionalTests.References
+{
+
+    [TestClass]
+    public class ScyllaReplicationTests : ReplicationTests
+    {
+        protected override IEnumerable<KeyValuePair<string, string>> GetSettings()
+        {
+            return new[]
+            {
+                new KeyValuePair<string, string>("Horde.Storage:ReferencesDbImplementation", HordeStorageSettings.ReferencesDbImplementations.Scylla.ToString()),
+                new KeyValuePair<string, string>("Horde.Storage:ReplicationLogWriterImplementation", HordeStorageSettings.ReplicationLogWriterImplementations.Scylla.ToString()),
+            };
+        }
+
+        protected override async Task SeedDb(IServiceProvider provider)
+        {
+            IReferencesStore referencesStore = provider.GetService<IReferencesStore>()!;
+            Assert.IsTrue(referencesStore.GetType() == typeof(ScyllaReferencesStore));
+
+            IReplicationLog replicationLog = provider.GetService<IReplicationLog>()!;
+            Assert.IsTrue(replicationLog.GetType() == typeof(ScyllaReplicationLog));
+
+            await SeedTestData(referencesStore);
+        }
+
+        protected override async Task TeardownDb(IServiceProvider provider)
+        {
+            IScyllaSessionManager scyllaSessionManager = provider.GetService<IScyllaSessionManager>()!;
+
+            ISession localKeyspace = scyllaSessionManager.GetSessionForLocalKeyspace();
+            
+            await Task.WhenAll(
+                // remove replication log table as we expect it to be empty when starting the tests
+                localKeyspace.ExecuteAsync(new SimpleStatement("DROP TABLE IF EXISTS replication_log;")),
+                // remove the snapshots
+                localKeyspace.ExecuteAsync(new SimpleStatement("DROP TABLE IF EXISTS replication_snapshot;")),
+                // remove the namespaces
+                localKeyspace.ExecuteAsync(new SimpleStatement("DROP TABLE IF EXISTS replication_namespace;"))
+            );
+
+
+        }
+    }
+
+
+    [TestClass]
+    public class MemoryReplicationTests : ReplicationTests
+    {
+        protected override IEnumerable<KeyValuePair<string, string>> GetSettings()
+        {
+            return new[]
+            {
+                new KeyValuePair<string, string>("Horde.Storage:ReferencesDbImplementation", HordeStorageSettings.ReferencesDbImplementations.Memory.ToString()),
+                new KeyValuePair<string, string>("Horde.Storage:ReplicationLogWriterImplementation", HordeStorageSettings.ReplicationLogWriterImplementations.Memory.ToString()),
+            };
+        }
+
+        protected override async Task SeedDb(IServiceProvider provider)
+        {
+            IReferencesStore referencesStore = provider.GetService<IReferencesStore>()!;
+            //verify we are using the expected refs store
+            Assert.IsTrue(referencesStore.GetType() == typeof(MemoryReferencesStore));
+
+            IReplicationLog replicationLog = provider.GetService<IReplicationLog>()!;
+            Assert.IsTrue(replicationLog.GetType() == typeof(MemoryReplicationLog));
+
+            await SeedTestData(referencesStore);
+        }
+
+        protected override Task TeardownDb(IServiceProvider provider)
+        {
+            return Task.CompletedTask;
+        }
+    }
+    
+    public abstract class ReplicationTests
+    {
+        private static TestServer? _server;
+        private static HttpClient? _httpClient;
+        protected IBlobStore _blobStore = null!;
+        protected IReferencesStore _referencesStore = null!;
+        private IReplicationLog _replicationLog = null!;
+
+        protected readonly NamespaceId TestNamespace = new NamespaceId("test-namespace");
+        protected readonly NamespaceId SnapshotNamespace = new NamespaceId("snapshot-namespace");
+        protected readonly BucketId TestBucket = new BucketId("default");
+
+        [TestInitialize]
+        public async Task Setup()
+
+        {
+            IConfigurationRoot configuration = new ConfigurationBuilder()
+                // we are not reading the base appSettings here as we want exact control over what runs in the tests
+                .AddJsonFile("appsettings.Testing.json", true)
+                .AddInMemoryCollection(GetSettings())
+                .AddEnvironmentVariables()
+                .Build();
+
+            Logger logger = new LoggerConfiguration()
+                .ReadFrom.Configuration(configuration)
+                .CreateLogger();
+
+            TestServer server = new TestServer(new WebHostBuilder()
+                .UseConfiguration(configuration)
+                .UseEnvironment("Testing")
+                .UseSerilog(logger)
+                .UseStartup<HordeStorageStartup>()
+            );
+            _httpClient = server.CreateClient();
+            _server = server;
+
+            _blobStore = _server.Services.GetService<IBlobStore>()!;
+            _referencesStore = _server.Services.GetService<IReferencesStore>()!;
+            _replicationLog = _server.Services.GetService<IReplicationLog>()!;
+
+            await SeedDb(server.Services);
+        }
+
+        protected async Task SeedTestData(IReferencesStore referencesStore)
+        {
+            await Task.CompletedTask;
+        }
+
+
+        [TestCleanup]
+        public async Task Teardown()
+        {
+            if (_server != null) 
+                await TeardownDb(_server.Services);
+        }
+
+        protected abstract IEnumerable<KeyValuePair<string, string>> GetSettings();
+
+        protected abstract Task SeedDb(IServiceProvider provider);
+        protected abstract Task TeardownDb(IServiceProvider provider);
+        
+        
+        [TestMethod]
+        public async Task ReplicationLogCreation()
+        {
+            CompactBinaryWriter writer = new CompactBinaryWriter();
+            writer.BeginObject();
+            writer.AddString("thisIsAField", "stringField");
+            writer.EndObject();
+
+            byte[] objectData = writer.Save();
+            BlobIdentifier objectHash = BlobIdentifier.FromBlob(objectData);
+
+            HttpContent requestContent = new ByteArrayContent(objectData);
+            requestContent.Headers.ContentType = new MediaTypeHeaderValue(CustomMediaTypeNames.UnrealCompactBinary);
+            requestContent.Headers.Add(CommonHeaders.HashHeaderName, objectHash.ToString());
+
+            {
+                HttpResponseMessage result = await _httpClient!.PutAsync(requestUri: $"api/v1/refs/{TestNamespace}/{TestBucket}/newReferenceObject.uecb", requestContent);
+                result.EnsureSuccessStatusCode();
+            }
+
+            {
+                HttpResponseMessage result = await _httpClient!.PutAsync(requestUri: $"api/v1/refs/{TestNamespace}/{TestBucket}/secondObject.uecb", requestContent);
+                result.EnsureSuccessStatusCode();
+            }
+
+            {
+                HttpResponseMessage result = await _httpClient!.PutAsync(requestUri: $"api/v1/refs/{TestNamespace}/{TestBucket}/thirdObject.uecb", requestContent);
+                result.EnsureSuccessStatusCode();
+            }
+            
+            {
+                HttpResponseMessage result = await _httpClient!.GetAsync(requestUri: $"api/v1/replication-log/incremental/{TestNamespace}");
+                result.EnsureSuccessStatusCode();
+                string s = await result.Content.ReadAsStringAsync();
+                
+                Assert.AreEqual(result!.Content.Headers.ContentType!.MediaType, MediaTypeNames.Application.Json);
+
+                ReplicationLogEvents? events = await result.Content.ReadAsAsync<ReplicationLogEvents>();
+                Assert.IsNotNull(events);
+                Assert.AreEqual(3, events!.Events.Count);
+
+                // parse the events returned, make sure they are in the right order
+                {
+                    ReplicationLogEvent e = events.Events[0];
+                    Assert.AreEqual(TestNamespace, e.Namespace);
+                    Assert.AreEqual(TestBucket, e.Bucket);
+                    Assert.AreEqual("newReferenceObject", e.Key.ToString());
+                    Assert.AreEqual(objectHash, e.Blob);
+                }
+
+                {
+                    ReplicationLogEvent e = events.Events[1];
+                    Assert.AreEqual(TestNamespace, e.Namespace);
+                    Assert.AreEqual(TestBucket, e.Bucket);
+                    Assert.AreEqual("secondObject", e.Key.ToString());
+                    Assert.AreEqual(objectHash, e.Blob);
+                }
+
+                {
+                    ReplicationLogEvent e = events.Events[2];
+                    Assert.AreEqual(TestNamespace, e.Namespace);
+                    Assert.AreEqual(TestBucket, e.Bucket);
+                    Assert.AreEqual("thirdObject", e.Key.ToString());
+                    Assert.AreEqual(objectHash, e.Blob);
+                }
+            }
+
+            CollectionAssert.AreEqual(new [] {TestNamespace}, await _replicationLog.GetNamespaces().ToArrayAsync());
+        }
+
+         
+        [TestMethod]
+        public async Task ReplicationLogReading()
+        {
+            CompactBinaryWriter writer = new CompactBinaryWriter();
+            writer.BeginObject();
+            writer.AddString("thisIsAField", "stringField");
+            writer.EndObject();
+
+            byte[] objectData = writer.Save();
+            BlobIdentifier objectHash = BlobIdentifier.FromBlob(objectData);
+
+            DateTime oldestTimestamp = DateTime.Now.AddDays(-1);
+            (string eventBucket, Guid eventId) = await _replicationLog.InsertAddEvent(TestNamespace, TestBucket, new KeyId("firstObject"), objectHash, oldestTimestamp);
+            await _replicationLog.InsertAddEvent(TestNamespace, TestBucket, new KeyId("secondObject"), objectHash, oldestTimestamp.AddHours(2.0));
+            await _replicationLog.InsertAddEvent(TestNamespace, TestBucket, new KeyId("thirdObject"), objectHash, oldestTimestamp.AddHours(3.0));
+            await _replicationLog.InsertAddEvent(TestNamespace, TestBucket, new KeyId("fourthObject"), objectHash, oldestTimestamp.AddDays(0.9));
+
+            {
+                HttpResponseMessage result = await _httpClient!.GetAsync(requestUri: $"api/v1/replication-log/incremental/{TestNamespace}?lastBucket={eventBucket}&lastEvent={eventId}");
+                result.EnsureSuccessStatusCode();
+                string s = await result.Content.ReadAsStringAsync();
+                
+                Assert.AreEqual(result!.Content.Headers.ContentType!.MediaType, MediaTypeNames.Application.Json);
+
+                ReplicationLogEvents? events = await result.Content.ReadAsAsync<ReplicationLogEvents>();
+                Assert.IsNotNull(events);
+                Assert.AreEqual(3, events!.Events.Count);
+
+                // we will not get the first event, as if we ever were to fetch all events we could potentially have events that are missed and snapshots are used instead
+
+                // parse the events returned, make sure they are in the right order
+                {
+                    ReplicationLogEvent e = events.Events[0];
+                    Assert.AreEqual(TestNamespace, e.Namespace);
+                    Assert.AreEqual(TestBucket, e.Bucket);
+                    Assert.AreEqual("secondObject", e.Key.ToString());
+                    Assert.AreEqual(objectHash, e.Blob);
+                }
+
+                {
+                    ReplicationLogEvent e = events.Events[1];
+                    Assert.AreEqual(TestNamespace, e.Namespace);
+                    Assert.AreEqual(TestBucket, e.Bucket);
+                    Assert.AreEqual("thirdObject", e.Key.ToString());
+                    Assert.AreEqual(objectHash, e.Blob);
+                }
+
+                {
+                    ReplicationLogEvent e = events.Events[2];
+                    Assert.AreEqual(TestNamespace, e.Namespace);
+                    Assert.AreEqual(TestBucket, e.Bucket);
+                    Assert.AreEqual("fourthObject", e.Key.ToString());
+                    Assert.AreEqual(objectHash, e.Blob);
+                }
+            }
+         
+            CollectionAssert.AreEqual(new [] {TestNamespace}, await _replicationLog.GetNamespaces().ToArrayAsync());
+        }
+
+        [TestMethod]
+        public async Task ReplicationLogResume()
+        {
+            CompactBinaryWriter writer = new CompactBinaryWriter();
+            writer.BeginObject();
+            writer.AddString("thisIsAField", "stringField");
+            writer.EndObject();
+
+            byte[] objectData = writer.Save();
+            BlobIdentifier objectHash = BlobIdentifier.FromBlob(objectData);
+
+            DateTime oldestTimestamp = DateTime.Now.AddDays(-3);
+            // insert multiple objects in the same time bucket, verifying that we correctly get only the objects after this
+            await _replicationLog.InsertAddEvent(TestNamespace, TestBucket, new KeyId("firstObject"), objectHash, oldestTimestamp);
+            await _replicationLog.InsertAddEvent(TestNamespace, TestBucket, new KeyId("secondObject"), objectHash, oldestTimestamp.AddHours(2.0));
+            (string eventBucket, Guid eventId) = await _replicationLog.InsertAddEvent(TestNamespace, TestBucket, new KeyId("thirdObject"), objectHash, oldestTimestamp.AddHours(2.1));
+            await _replicationLog.InsertAddEvent(TestNamespace, TestBucket, new KeyId("fourthObject"), objectHash, oldestTimestamp.AddHours(2.11));
+            await _replicationLog.InsertAddEvent(TestNamespace, TestBucket, new KeyId("fifthObject"), objectHash, oldestTimestamp.AddHours(2.12));
+            await _replicationLog.InsertAddEvent(TestNamespace, TestBucket, new KeyId("sixthObject"), objectHash, oldestTimestamp.AddDays(2.13));
+
+            {
+                HttpResponseMessage result = await _httpClient!.GetAsync(requestUri: $"api/v1/replication-log/incremental/{TestNamespace}?lastBucket={eventBucket}&lastEvent={eventId}");
+                result.EnsureSuccessStatusCode();
+                string s = await result.Content.ReadAsStringAsync();
+                
+                Assert.AreEqual(result!.Content.Headers.ContentType!.MediaType, MediaTypeNames.Application.Json);
+
+                ReplicationLogEvents? events = await result.Content.ReadAsAsync<ReplicationLogEvents>();
+                Assert.IsNotNull(events);
+                Assert.AreEqual(3, events!.Events.Count);
+
+                // we will not get the first event, as if we ever were to fetch all events we could potentially have events that are missed and snapshots are used instead
+
+                // parse the events returned, make sure they are in the right order
+                {
+                    ReplicationLogEvent e = events.Events[0];
+                    Assert.AreEqual(TestNamespace, e.Namespace);
+                    Assert.AreEqual(TestBucket, e.Bucket);
+                    Assert.AreEqual("fourthObject", e.Key.ToString());
+                    Assert.AreEqual(objectHash, e.Blob);
+                }
+
+                {
+                    ReplicationLogEvent e = events.Events[1];
+                    Assert.AreEqual(TestNamespace, e.Namespace);
+                    Assert.AreEqual(TestBucket, e.Bucket);
+                    Assert.AreEqual("fifthObject", e.Key.ToString());
+                    Assert.AreEqual(objectHash, e.Blob);
+                }
+
+                {
+                    ReplicationLogEvent e = events.Events[2];
+                    Assert.AreEqual(TestNamespace, e.Namespace);
+                    Assert.AreEqual(TestBucket, e.Bucket);
+                    Assert.AreEqual("sixthObject", e.Key.ToString());
+                    Assert.AreEqual(objectHash, e.Blob);
+                }
+            }
+         
+            CollectionAssert.AreEqual(new [] {TestNamespace}, await _replicationLog.GetNamespaces().ToArrayAsync());
+        }
+
+        [TestMethod]
+        public async Task ReplicationLogReadingLimit()
+        {
+            CompactBinaryWriter writer = new CompactBinaryWriter();
+            writer.BeginObject();
+            writer.AddString("thisIsAField", "stringField");
+            writer.EndObject();
+
+            byte[] objectData = writer.Save();
+            BlobIdentifier objectHash = BlobIdentifier.FromBlob(objectData);
+
+            DateTime oldestTimestamp = DateTime.Now.AddDays(-1);
+            (string eventBucket, Guid eventId) = await _replicationLog.InsertAddEvent(TestNamespace, TestBucket, new KeyId("firstObject"), objectHash, oldestTimestamp);
+            await _replicationLog.InsertAddEvent(TestNamespace, TestBucket, new KeyId("secondObject"), objectHash, oldestTimestamp.AddHours(1));
+            await _replicationLog.InsertAddEvent(TestNamespace, TestBucket, new KeyId("thirdObject"), objectHash, oldestTimestamp.AddHours(1.5));
+            await _replicationLog.InsertAddEvent(TestNamespace, TestBucket, new KeyId("fourthObject"), objectHash, oldestTimestamp.AddDays(0.9));
+
+            // start from the second event
+            const int eventsToFetch = 2;
+            {
+                HttpResponseMessage result = await _httpClient!.GetAsync(requestUri: $"api/v1/replication-log/incremental/{TestNamespace}?lastBucket={eventBucket}&lastEvent={eventId}&count={eventsToFetch}");
+                result.EnsureSuccessStatusCode();
+                string s = await result.Content.ReadAsStringAsync();
+                
+                Assert.AreEqual(result!.Content.Headers.ContentType!.MediaType, MediaTypeNames.Application.Json);
+
+                ReplicationLogEvents? events = await result.Content.ReadAsAsync<ReplicationLogEvents>();
+                Assert.IsNotNull(events);
+                Assert.AreEqual(eventsToFetch, events!.Events.Count);
+
+                // parse the events returned, make sure they are in the right order
+                {
+                    ReplicationLogEvent e = events.Events[0];
+                    Assert.AreEqual(TestNamespace, e.Namespace);
+                    Assert.AreEqual(TestBucket, e.Bucket);
+                    Assert.AreEqual("secondObject", e.Key.ToString());
+                    Assert.AreEqual(objectHash, e.Blob);
+                }
+
+                {
+                    ReplicationLogEvent e = events.Events[1];
+                    Assert.AreEqual(TestNamespace, e.Namespace);
+                    Assert.AreEqual(TestBucket, e.Bucket);
+                    Assert.AreEqual("thirdObject", e.Key.ToString());
+                    Assert.AreEqual(objectHash, e.Blob);
+                }
+            }
+            
+            CollectionAssert.AreEqual(new [] {TestNamespace}, await _replicationLog.GetNamespaces().ToArrayAsync());
+        }
+
+        [TestMethod]
+        public async Task ReplicationLogInvalidBucket()
+        {
+            // the namespace exists but the bucket does not
+            CompactBinaryWriter writer = new CompactBinaryWriter();
+            writer.BeginObject();
+            writer.AddString("thisIsAField", "stringField");
+            writer.EndObject();
+
+            byte[] objectData = writer.Save();
+            BlobIdentifier objectHash = BlobIdentifier.FromBlob(objectData);
+
+            await _replicationLog.InsertAddEvent(TestNamespace, TestBucket, new KeyId("firstObject"), objectHash, DateTime.Now.AddDays(-1));
+
+            string eventBucket = "rep-00000000";
+            Guid eventId = Guid.NewGuid();
+
+            {
+                HttpResponseMessage result = await _httpClient!.GetAsync(requestUri: $"api/v1/replication-log/incremental/{TestNamespace}?lastBucket={eventBucket}&lastEvent={eventId}");
+                Assert.AreEqual(HttpStatusCode.BadRequest, result.StatusCode);
+                
+                string s = await result.Content.ReadAsStringAsync();
+                
+                ProblemDetails? problem = await result.Content.ReadFromJsonAsync<ProblemDetails?>();
+                Assert.IsNotNull(problem);
+            }
+
+            CollectionAssert.AreEqual(new [] {TestNamespace}, await _replicationLog.GetNamespaces().ToArrayAsync());
+        }
+
+        [TestMethod]
+        public async Task ReplicationLogEmptyLog()
+        {
+            // the namespace does not exist
+            {
+                HttpResponseMessage result = await _httpClient!.GetAsync(requestUri: $"api/v1/replication-log/incremental/{TestNamespace}");
+                Assert.AreEqual(HttpStatusCode.NotFound, result.StatusCode);
+                string s = await result.Content.ReadAsStringAsync();
+                
+                ProblemDetails? problem = await result.Content.ReadFromJsonAsync<ProblemDetails?>();
+                Assert.IsNotNull(problem);
+            }
+
+            CollectionAssert.AreEqual(Array.Empty<NamespaceId>(), await _replicationLog.GetNamespaces().ToArrayAsync());
+        }
+
+        [TestMethod]
+        public async Task ReplicationLogNoIncrementalLogAvailable()
+        {
+            CompactBinaryWriter writer = new CompactBinaryWriter();
+            writer.BeginObject();
+            writer.AddString("thisIsAField", "stringField");
+            writer.EndObject();
+
+            byte[] objectData = writer.Save();
+            BlobIdentifier objectHash = BlobIdentifier.FromBlob(objectData);
+
+            DateTime oldestTimestamp = DateTime.Now.AddDays(-1);
+            await _replicationLog.InsertAddEvent(TestNamespace, TestBucket, new KeyId("firstObject"), objectHash, oldestTimestamp);
+            await _replicationLog.InsertAddEvent(TestNamespace, TestBucket, new KeyId("secondObject"), objectHash, oldestTimestamp.AddHours(1));
+            await _replicationLog.InsertAddEvent(TestNamespace, TestBucket, new KeyId("thirdObject"), objectHash, oldestTimestamp.AddHours(1.5));
+            await _replicationLog.InsertAddEvent(TestNamespace, TestBucket, new KeyId("fourthObject"), objectHash, oldestTimestamp.AddDays(0.9));
+
+            // create a snapshot
+            ReplicationLogSnapshotBuilder snapshotBuilder = ActivatorUtilities.CreateInstance<ReplicationLogSnapshotBuilder>(_server!.Services);
+            Assert.IsNotNull(snapshotBuilder);
+            BlobIdentifier snapshotBlobId = await snapshotBuilder.BuildSnapshot(TestNamespace, SnapshotNamespace);
+            Assert.IsTrue(await _blobStore.Exists(SnapshotNamespace, snapshotBlobId));
+
+            // use a bucket that does not exist, should raise a message to use a snapshot instead
+            string bucketThatDoesNotExist = "rep-0000";
+            {
+                HttpResponseMessage result = await _httpClient!.GetAsync(requestUri: $"api/v1/replication-log/incremental/{TestNamespace}?lastBucket={bucketThatDoesNotExist}&lastEvent={Guid.NewGuid()}");
+                Assert.AreEqual(HttpStatusCode.BadRequest, result.StatusCode);
+                string s = await result.Content.ReadAsStringAsync();
+                
+                ProblemDetails? problem = await result.Content.ReadFromJsonAsync<ProblemDetails?>();
+                Assert.IsNotNull(problem);
+                Assert.AreEqual("http://jupiter.epicgames.com/replication/useSnapshot", problem!.Type);
+                Assert.IsTrue(problem.Extensions.ContainsKey("SnapshotId"));
+                Assert.AreEqual(snapshotBlobId, new BlobIdentifier(problem.Extensions["SnapshotId"]!.ToString()!));
+            }
+            
+        }
+
+        [TestMethod]
+        public async Task ReplicationLogSnapshotCreation()
+        {
+            CompactBinaryWriter writer = new CompactBinaryWriter();
+            writer.BeginObject();
+            writer.AddString("thisIsAField", "stringField");
+            writer.EndObject();
+
+            byte[] objectData = writer.Save();
+            BlobIdentifier objectHash = BlobIdentifier.FromBlob(objectData);
+
+            DateTime oldestTimestamp = DateTime.Now.AddDays(-1);
+            await _replicationLog.InsertAddEvent(TestNamespace, TestBucket, new KeyId("firstObject"), objectHash, oldestTimestamp);
+            await _replicationLog.InsertAddEvent(TestNamespace, TestBucket, new KeyId("secondObject"), objectHash, oldestTimestamp.AddHours(1));
+            await _replicationLog.InsertAddEvent(TestNamespace, TestBucket, new KeyId("thirdObject"), objectHash, oldestTimestamp.AddHours(1.5));
+            (string lastEventBucket, Guid lastEventId) = await _replicationLog.InsertAddEvent(TestNamespace, TestBucket, new KeyId("fourthObject"), objectHash, oldestTimestamp.AddDays(0.7));
+
+            // verify the objects were added
+            {
+                List<ReplicationLogEvent> logEvents = await _replicationLog.Get(TestNamespace, null, null).ToListAsync();
+                Assert.AreEqual(4, logEvents.Count);
+
+                // parse the events returned, make sure they are in the right order
+                {
+                    ReplicationLogEvent e = logEvents[0];
+                    Assert.AreEqual(TestNamespace, e.Namespace);
+                    Assert.AreEqual(TestBucket, e.Bucket);
+                    Assert.AreEqual("firstObject", e.Key.ToString());
+                    Assert.AreEqual(objectHash, e.Blob);
+                }
+
+                {
+                    ReplicationLogEvent e = logEvents[1];
+                    Assert.AreEqual(TestNamespace, e.Namespace);
+                    Assert.AreEqual(TestBucket, e.Bucket);
+                    Assert.AreEqual("secondObject", e.Key.ToString());
+                    Assert.AreEqual(objectHash, e.Blob);
+                }
+
+                {
+                    ReplicationLogEvent e = logEvents[2];
+                    Assert.AreEqual(TestNamespace, e.Namespace);
+                    Assert.AreEqual(TestBucket, e.Bucket);
+                    Assert.AreEqual("thirdObject", e.Key.ToString());
+                    Assert.AreEqual(objectHash, e.Blob);
+                }
+
+                {
+                    ReplicationLogEvent e = logEvents[3];
+                    Assert.AreEqual(TestNamespace, e.Namespace);
+                    Assert.AreEqual(TestBucket, e.Bucket);
+                    Assert.AreEqual("fourthObject", e.Key.ToString());
+                    Assert.AreEqual(objectHash, e.Blob);
+                }
+            }
+
+            // verify there are no previous snapshots
+            Assert.AreEqual(0, await _replicationLog.GetSnapshots(TestNamespace).Count().FirstAsync());
+
+            // create a snapshot
+            ReplicationLogSnapshotBuilder snapshotBuilder = ActivatorUtilities.CreateInstance<ReplicationLogSnapshotBuilder>(_server!.Services);
+            Assert.IsNotNull(snapshotBuilder);
+            BlobIdentifier snapshotBlobId = await snapshotBuilder.BuildSnapshot(TestNamespace, SnapshotNamespace);
+            Assert.IsTrue(await _blobStore.Exists(SnapshotNamespace, snapshotBlobId));
+
+            SnapshotInfo? snapshotInfo = await _replicationLog.GetSnapshots(TestNamespace).FirstAsync();
+            Assert.IsNotNull(snapshotInfo);
+            Assert.AreEqual(snapshotBlobId, snapshotInfo.SnapshotBlob);
+
+            BlobContents blobContents = await _blobStore.GetObject(SnapshotNamespace, snapshotBlobId);
+            ReplicationLogSnapshot snapshot = await ReplicationLogSnapshot.DeserializeSnapshot(blobContents.Stream);
+
+            Assert.AreEqual(lastEventId, snapshot.LastEvent);
+            Assert.AreEqual(lastEventBucket, snapshot.LastBucket);
+            
+            Assert.IsTrue(snapshot.LiveObjects.Any(o => o.Bucket == TestBucket && o.Key == new KeyId("firstObject")));
+            Assert.IsTrue(snapshot.LiveObjects.Any(o => o.Bucket == TestBucket && o.Key == new KeyId("secondObject")));
+            Assert.IsTrue(snapshot.LiveObjects.Any(o => o.Bucket == TestBucket && o.Key == new KeyId("thirdObject")));
+            Assert.IsTrue(snapshot.LiveObjects.Any(o => o.Bucket == TestBucket && o.Key == new KeyId("fourthObject")));
+        }
+
+
+        [TestMethod]
+        public async Task ReplicationLogSnapshotQuerying()
+        {
+            CompactBinaryWriter writer = new CompactBinaryWriter();
+            writer.BeginObject();
+            writer.AddString("thisIsAField", "stringField");
+            writer.EndObject();
+
+            byte[] objectData = writer.Save();
+            BlobIdentifier objectHash = BlobIdentifier.FromBlob(objectData);
+
+            DateTime oldestTimestamp = DateTime.Now.AddDays(-1);
+            await _replicationLog.InsertAddEvent(TestNamespace, TestBucket, new KeyId("firstObject"), objectHash, oldestTimestamp);
+            await _replicationLog.InsertAddEvent(TestNamespace, TestBucket, new KeyId("secondObject"), objectHash, oldestTimestamp.AddHours(1));
+            await _replicationLog.InsertAddEvent(TestNamespace, TestBucket, new KeyId("thirdObject"), objectHash, oldestTimestamp.AddHours(1.5));
+            (string lastEventBucket, Guid lastEventId) = await _replicationLog.InsertAddEvent(TestNamespace, TestBucket, new KeyId("fourthObject"), objectHash, oldestTimestamp.AddDays(0.9));
+
+            // verify the objects were added
+            List<ReplicationLogEvent> logEvents = await _replicationLog.Get(TestNamespace, null, null).ToListAsync();
+            Assert.AreEqual(4, logEvents.Count);
+
+            // verify there are no previous snapshots
+            Assert.AreEqual(0, await _replicationLog.GetSnapshots(TestNamespace).Count().FirstAsync());
+
+            // create a snapshot
+            ReplicationLogSnapshotBuilder snapshotBuilder = ActivatorUtilities.CreateInstance<ReplicationLogSnapshotBuilder>(_server!.Services);
+            Assert.IsNotNull(snapshotBuilder);
+            BlobIdentifier snapshotBlobId = await snapshotBuilder.BuildSnapshot(TestNamespace, SnapshotNamespace);
+            Assert.IsTrue(await _blobStore.Exists(SnapshotNamespace, snapshotBlobId));
+
+            SnapshotInfo? snapshotInfo = await _replicationLog.GetSnapshots(TestNamespace).FirstAsync();
+            Assert.AreEqual(snapshotBlobId, snapshotInfo.SnapshotBlob);
+
+            // make sure the snapshot is returned by the rest api
+            {
+                HttpResponseMessage result = await _httpClient!.GetAsync(requestUri: $"api/v1/replication-log/snapshots/{TestNamespace}");
+                result.EnsureSuccessStatusCode();
+                string s = await result.Content.ReadAsStringAsync();
+                
+                Assert.AreEqual(result!.Content.Headers.ContentType!.MediaType, MediaTypeNames.Application.Json);
+
+                ReplicationLogSnapshots snapshots = await result.Content.ReadAsAsync<ReplicationLogSnapshots>();
+                Assert.IsNotNull(snapshots);
+                Assert.AreEqual(1, snapshots.Snapshots.Count);
+
+                SnapshotInfo foundSnapshot = snapshots.Snapshots[0];
+
+                Assert.AreEqual(snapshotBlobId, foundSnapshot.SnapshotBlob);
+                BlobContents blobContents = await _blobStore.GetObject(SnapshotNamespace, snapshotBlobId);
+                ReplicationLogSnapshot snapshot = await ReplicationLogSnapshot.DeserializeSnapshot(blobContents.Stream);
+
+                Assert.AreEqual(lastEventBucket, snapshot.LastBucket);
+                Assert.AreEqual(lastEventId, snapshot.LastEvent);
+            }
+        }
+
+        // builds a snapshot and make sure we can resume iterating after it
+        [TestMethod]
+        public async Task ReplicationLogSnapshotResume()
+        {
+            CompactBinaryWriter writer = new CompactBinaryWriter();
+            writer.BeginObject();
+            writer.AddString("thisIsAField", "stringField");
+            writer.EndObject();
+
+            byte[] objectData = writer.Save();
+            BlobIdentifier objectHash = BlobIdentifier.FromBlob(objectData);
+
+            DateTime oldestTimestamp = DateTime.Now.AddDays(-1);
+            await _replicationLog.InsertAddEvent(TestNamespace, TestBucket, new KeyId("firstObject"), objectHash, oldestTimestamp);
+            await _replicationLog.InsertAddEvent(TestNamespace, TestBucket, new KeyId("secondObject"), objectHash, oldestTimestamp.AddHours(1));
+            await _replicationLog.InsertAddEvent(TestNamespace, TestBucket, new KeyId("thirdObject"), objectHash, oldestTimestamp.AddHours(1.5));
+            await _replicationLog.InsertAddEvent(TestNamespace, TestBucket, new KeyId("fourthObject"), objectHash, oldestTimestamp.AddDays(0.9));
+
+            // verify the objects were added
+            List<ReplicationLogEvent> logEvents = await _replicationLog.Get(TestNamespace, null, null).ToListAsync();
+            Assert.AreEqual(4, logEvents.Count);
+
+            // verify there are no previous snapshots
+            Assert.AreEqual(0, await _replicationLog.GetSnapshots(TestNamespace).Count().FirstAsync());
+
+            // create a snapshot
+            ReplicationLogSnapshotBuilder snapshotBuilder = ActivatorUtilities.CreateInstance<ReplicationLogSnapshotBuilder>(_server!.Services);
+            Assert.IsNotNull(snapshotBuilder);
+            BlobIdentifier snapshotBlobId = await snapshotBuilder.BuildSnapshot(TestNamespace, SnapshotNamespace);
+            Assert.IsTrue(await _blobStore.Exists(SnapshotNamespace, snapshotBlobId));
+
+            SnapshotInfo? snapshotInfo = await _replicationLog.GetSnapshots(TestNamespace).FirstAsync();
+            Assert.AreEqual(snapshotBlobId, snapshotInfo.SnapshotBlob);
+
+            BlobContents blobContents = await _blobStore.GetObject(SnapshotNamespace, snapshotBlobId);
+            ReplicationLogSnapshot snapshot = await ReplicationLogSnapshot.DeserializeSnapshot(blobContents.Stream);
+
+            // insert more events
+            await _replicationLog.InsertAddEvent(TestNamespace, TestBucket, new KeyId("fifthObject"), objectHash, oldestTimestamp.AddDays(0.91));
+            await _replicationLog.InsertAddEvent(TestNamespace, TestBucket, new KeyId("sixthObject"), objectHash, oldestTimestamp.AddDays(0.92));
+
+            // verify the new events can be found when resuming from the snapshot
+            {
+                HttpResponseMessage result = await _httpClient!.GetAsync(requestUri: $"api/v1/replication-log/incremental/{TestNamespace}?lastBucket={snapshot.LastBucket}&lastEvent={snapshot.LastEvent}");
+                result.EnsureSuccessStatusCode();
+                
+                Assert.AreEqual(result!.Content.Headers.ContentType!.MediaType, MediaTypeNames.Application.Json);
+
+                ReplicationLogEvents? events = await result.Content.ReadAsAsync<ReplicationLogEvents>();
+                Assert.IsNotNull(events);
+                Assert.AreEqual(2, events!.Events.Count);
+
+                // parse the events returned, make sure they are in the right order
+                {
+                    ReplicationLogEvent e = events.Events[0];
+                    Assert.AreEqual(TestNamespace, e.Namespace);
+                    Assert.AreEqual(TestBucket, e.Bucket);
+                    Assert.AreEqual("fifthObject", e.Key.ToString());
+                    Assert.AreEqual(objectHash, e.Blob);
+                }
+
+                {
+                    ReplicationLogEvent e = events.Events[1];
+                    Assert.AreEqual(TestNamespace, e.Namespace);
+                    Assert.AreEqual(TestBucket, e.Bucket);
+                    Assert.AreEqual("sixthObject", e.Key.ToString());
+                    Assert.AreEqual(objectHash, e.Blob);
+                }
+            }
+        }
+
+        [TestMethod]
+        public async Task ReplicationLogSnapshotCleanup()
+        {
+            const int maxCountOfSnapshots = 10;
+            int countOfSnapshotsToCreate = maxCountOfSnapshots + 2;
+
+            CompactBinaryWriter writer = new CompactBinaryWriter();
+            writer.BeginObject();
+            writer.AddString("thisIsAField", "stringField");
+            writer.EndObject();
+
+            byte[] objectData = writer.Save();
+            BlobIdentifier objectHash = BlobIdentifier.FromBlob(objectData);
+
+            List<BlobIdentifier> createdSnapshots = new List<BlobIdentifier>();
+            for (int i = 0; i < countOfSnapshotsToCreate ; i++)
+            {
+                await _replicationLog.InsertAddEvent(TestNamespace, TestBucket, new KeyId($"object {i}"), objectHash);
+
+                ReplicationLogSnapshotBuilder snapshotBuilder = ActivatorUtilities.CreateInstance<ReplicationLogSnapshotBuilder>(_server!.Services);
+                Assert.IsNotNull(snapshotBuilder);
+                BlobIdentifier snapshotBlobId = await snapshotBuilder.BuildSnapshot(TestNamespace, SnapshotNamespace);
+                Assert.IsTrue(await _blobStore.Exists(SnapshotNamespace, snapshotBlobId));
+                createdSnapshots.Add(snapshotBlobId);
+            }
+
+            List<BlobIdentifier> snapshots = await _replicationLog.GetSnapshots(TestNamespace).Select(info => info.SnapshotBlob).ToListAsync();
+            // snapshots are returned newest first so we inverse this order
+            snapshots.Reverse();
+
+            // verify we hit the max number of snapshots
+            Assert.AreEqual(maxCountOfSnapshots, snapshots.Count);
+
+            // the first two snapshots should have been removed
+            createdSnapshots.RemoveAt(0);
+            createdSnapshots.RemoveAt(0);
+
+            CollectionAssert.AreEqual(createdSnapshots, snapshots);
+        }
+    }
+}
diff --git a/Engine/Source/Programs/Horde/HordeStorage/Horde.Storage/Tests/Functional/RefTests.cs b/Engine/Source/Programs/Horde/HordeStorage/Horde.Storage/Tests/Functional/RefTests.cs
new file mode 100644
--- /dev/null
+++ b/Engine/Source/Programs/Horde/HordeStorage/Horde.Storage/Tests/Functional/RefTests.cs
@@ -0,0 +1,812 @@
+// Copyright Epic Games, Inc. All Rights Reserved.
+
+using System;
+using System.Collections.Generic;
+using System.IO;
+using System.Linq;
+using System.Net;
+using System.Net.Http;
+using System.Net.Http.Headers;
+using System.Text;
+using System.Threading;
+using System.Threading.Tasks;
+using Amazon.DynamoDBv2;
+using Amazon.DynamoDBv2.Model;
+using Dasync.Collections;
+using Horde.Storage.Controllers;
+using Horde.Storage.Implementation;
+using Jupiter;
+using Jupiter.Implementation;
+using Microsoft.AspNetCore.TestHost;
+using Microsoft.AspNetCore.Hosting;
+using Microsoft.Extensions.Configuration;
+using Microsoft.Extensions.DependencyInjection;
+using Microsoft.Extensions.Options;
+using Microsoft.VisualStudio.TestTools.UnitTesting;
+using MongoDB.Driver;
+using Newtonsoft.Json;
+using Newtonsoft.Json.Linq;
+using Serilog;
+using Serilog.Core;
+
+namespace Horde.Storage.FunctionalTests.Ref
+{
+
+    [TestClass]
+    public class MongoRefTests : RefTests
+    {
+        private static readonly string MongoTestDatabase = "Europa_TestNamespace";
+        private static MongoSettings? _mongoSettings;
+        private MongoRefsStore? _refsStore;
+
+        private static MongoClient GetMongoClient()
+        {
+            return new MongoClient(_mongoSettings!.ConnectionString);
+        }
+
+        protected override async Task SeedDb(IServiceProvider provider)
+        {
+            _mongoSettings = provider.GetService<IOptionsMonitor<MongoSettings>>()!.CurrentValue;
+
+            MongoClient client = GetMongoClient();
+            if (client.GetDatabase(MongoTestDatabase) != null)
+            {
+                await client.DropDatabaseAsync(MongoTestDatabase);
+            }
+
+            
+            //verify we are using the expected refs store
+            IRefsStore? refStore = provider.GetService<IRefsStore>();
+            Assert.IsTrue(RefStoreIs(refStore, typeof(MongoRefsStore)));
+
+            _refsStore = new MongoRefsStore(provider.GetService<IOptionsMonitor<MongoSettings>>()!);
+            IBlobStore blobStore = (IBlobStore)provider.GetService(typeof(IBlobStore))!;
+
+            await SeedRefTestData(_refsStore, blobStore);
+
+        }
+
+        protected override Task TeardownDb(IServiceProvider provider)
+        {
+            return Task.CompletedTask;
+        }
+
+        protected override async Task<RefRecord> GetTestRecord(BucketId bucket, KeyId name, IRefsStore.ExtraFieldsFlag fields)
+        {
+            var record = await _refsStore!.Get(TestNamespace, bucket, name, fields);
+            if (record == null)
+                throw new Exception("Unable to find record");
+
+            return record;
+        }
+        protected override IEnumerable<KeyValuePair<string, string>> GetSettings()
+        {
+            return new[] {new KeyValuePair<string, string>("Horde.Storage:RefDbImplementation", HordeStorageSettings.RefDbImplementations.Mongo.ToString())};
+        }
+    }
+
+
+    [TestClass]
+    public class DynamoRefTests : RefTests
+    {
+        private DynamoDbRefsStore? _refStore;
+        private readonly string[] _tablesToDelete = new[] {"Europa_Cache", "Europa_Namespace"};
+
+        protected override IEnumerable<KeyValuePair<string, string>> GetSettings()
+        {
+            return new[] {new KeyValuePair<string, string>("Horde.Storage:RefDbImplementation", HordeStorageSettings.RefDbImplementations.DynamoDb.ToString())};
+        }
+
+        protected override async Task SeedDb(IServiceProvider provider)
+        {
+            IOptionsMonitor<DynamoDbSettings> dynamoSettings = provider.GetService<IOptionsMonitor<DynamoDbSettings>>()!;
+            IAmazonDynamoDB client = provider.GetService<IAmazonDynamoDB>()!;
+
+            //verify we are using the expected refs store
+            var cacheStore = provider.GetService<IRefsStore>();
+            Assert.IsTrue(RefStoreIs(cacheStore, typeof(DynamoDbRefsStore)));
+
+            _refStore = new DynamoDbRefsStore(dynamoSettings, provider.GetService<IAmazonDynamoDB>()!, provider.GetService<DynamoNamespaceStore>()!);
+            IBlobStore blobStore = (IBlobStore)provider.GetService(typeof(IBlobStore))!;
+
+            await SeedRefTestData(_refStore, blobStore);
+        }
+
+        protected override async Task TeardownDb(IServiceProvider provider)
+        {
+            // delete tables after use so that we always run from a clean slate
+            // you can remove this if you want to look at the resulting state after a run
+            IAmazonDynamoDB client = provider.GetService<IAmazonDynamoDB>()!;
+            ListTablesResponse tables = await client.ListTablesAsync();
+            foreach (string t in _tablesToDelete)
+            {
+                if (tables.TableNames.Contains(t))
+                {
+                    await client.DeleteTableAsync(t);
+                }
+            }
+        }
+
+        protected override async Task<RefRecord> GetTestRecord(BucketId bucket, KeyId name, IRefsStore.ExtraFieldsFlag fields)
+        {
+            var record = await _refStore!.Get(TestNamespace, bucket, name, fields);
+            if (record == null)
+                throw new Exception("Unable to find record");
+
+            return record;
+        }
+
+        // removing batch test using dynamo as it hangs the test runner for unknown reasons when run with a lot of other tests
+        [TestMethod]
+        public new Task Batch()
+        {
+            return Task.CompletedTask;
+        }
+    }
+    
+    public abstract class RefTests
+    {
+        private static TestServer? _server;
+        private static HttpClient? _httpClient;
+        private static ILastAccessCache<RefRecord>? _lastAccessCache;
+
+        protected static readonly string TestObjectContents = "Test Object Data";
+        protected static readonly byte[] TestObjectData = Encoding.ASCII.GetBytes(TestObjectContents);
+        protected static readonly ContentHash TestObjectHash = ContentHash.FromBlob(TestObjectData);
+        protected static readonly BlobIdentifier TestObjectBlobHash = BlobIdentifier.FromContentHash(TestObjectHash);
+        
+        protected readonly NamespaceId TestNamespace = new NamespaceId("test-namespace");
+
+        [TestInitialize]
+        public async Task Setup()
+
+        {
+            IConfigurationRoot configuration = new ConfigurationBuilder()
+                // we are not reading the base appSettings here as we want exact control over what runs in the tests
+                .AddJsonFile("appsettings.Testing.json", true)
+                .AddInMemoryCollection(GetSettings())
+                .AddEnvironmentVariables()
+                .Build();
+
+            Logger logger = new LoggerConfiguration()
+                .ReadFrom.Configuration(configuration)
+                .CreateLogger();
+
+            TestServer server = new TestServer(new WebHostBuilder()
+                .UseConfiguration(configuration)
+                .UseEnvironment("Testing")
+                .UseSerilog(logger)
+                .UseStartup<HordeStorageStartup>()
+            );
+            _httpClient = server.CreateClient();
+            _server = server;
+
+            _lastAccessCache = _server.Services.GetService<ILastAccessCache<RefRecord>>();
+
+            await SeedDb(server.Services);
+        }
+
+        protected async Task SeedRefTestData(IRefsStore refsStore, IBlobStore blobStore)
+        {
+            // add the objects as if they were last accessed a day ago
+            DateTime lastAccessTime = DateTime.Now.AddDays(-1);
+
+            BucketId bucket = new BucketId("bucket");
+
+            await refsStore.Add(new RefRecord(TestNamespace, bucket, new KeyId("testObject"), blobs: new[] { TestObjectBlobHash },
+                lastAccessTime: lastAccessTime, contentHash: TestObjectHash, metadata: null));
+            await blobStore.PutObject(TestNamespace, TestObjectData, TestObjectBlobHash);
+
+            // a object with metadata
+            await refsStore.Add(new RefRecord(TestNamespace, bucket, new KeyId("testObjectWithMetadata"), blobs: new[] { TestObjectBlobHash },
+                lastAccessTime: lastAccessTime, contentHash: TestObjectHash, metadata: new Dictionary<string, object>() {{"Foo", "Bar"}}));
+            await blobStore.PutObject(TestNamespace, TestObjectData, TestObjectBlobHash);
+
+            // the deletableObject should be slightly newer then testObject to control sorting for last access
+            await refsStore.Add(new RefRecord(TestNamespace, bucket, new KeyId("deletableObject"),
+                blobs: new[] { TestObjectBlobHash }, contentHash: TestObjectHash,lastAccessTime: lastAccessTime.AddMinutes(2), metadata: null));
+            await blobStore.PutObject(TestNamespace, TestObjectData, TestObjectBlobHash);
+
+            byte[] contents = Encoding.ASCII.GetBytes("This content will only be hashed");
+            BlobIdentifier newContentIdentifier = BlobIdentifier.FromBlob(contents);
+            await refsStore.Add(new RefRecord(TestNamespace, bucket, new KeyId("notUploadedToIo"), blobs: new[] { newContentIdentifier },
+                lastAccessTime: lastAccessTime, contentHash: newContentIdentifier, metadata: null));
+            // this object should not be uploaded to IO, and is considered recently uploaded for last access filtering
+
+        }
+
+
+        [TestCleanup]
+        public async Task Teardown()
+        {
+            if (_server != null) 
+                await TeardownDb(_server.Services);
+        }
+
+        protected bool RefStoreIs(IRefsStore? refStore, Type refStoreType)
+        {
+            if (refStore == null)
+                return false;
+
+            if (refStore.GetType() == refStoreType)
+            {
+                return true;
+            }
+
+            if (refStore is CachedRefStore store)
+            {
+                if (store.BackingStore.GetType() == refStoreType)
+                {
+                    return true;
+                }
+            }
+
+            return false;
+        }
+
+        protected abstract IEnumerable<KeyValuePair<string, string>> GetSettings();
+
+        protected abstract Task SeedDb(IServiceProvider provider);
+        protected abstract Task TeardownDb(IServiceProvider provider);
+
+        [TestMethod]
+        public async Task ListNamespaces()
+        {
+            HttpResponseMessage response = await _httpClient!.GetAsync("api/v1/c/ddc");
+            response.EnsureSuccessStatusCode();
+            GetNamespacesResponse content = await response.Content.ReadAsAsync<GetNamespacesResponse>();
+
+            // the namespaces here may not just be the namespaces we have created depending on what has been done in the dbs before
+            // as such we can only verify that TestNamespace exists
+
+            Assert.IsTrue(content.Namespaces.Contains(TestNamespace));
+        }
+
+        [TestMethod]
+        public async Task GetRefObject()
+        {
+            // fetch refs records once first so that any previous runs of tests are ignored
+            List<(RefRecord, DateTime)> _ = await _lastAccessCache!.GetLastAccessedRecords();
+
+            Task<HttpResponseMessage> response = _httpClient!.GetAsync($"api/v1/c/ddc/{TestNamespace}/bucket/testObject");
+            await response;
+            HttpResponseMessage result = response.Result;
+
+            result.EnsureSuccessStatusCode();
+            RefResponse content = await result.Content.ReadAsAsync<RefResponse>();
+            Assert.AreEqual(expected: $"{TestNamespace}.bucket.testObject", content.Name);
+            Assert.IsNull(content.Metadata);
+            Assert.IsNotNull(content.Blob);
+            CollectionAssert.AreEqual(TestObjectData, content.Blob);
+
+            // give the last access cache some time to get updated as this is not guaranteed to have finished when the Get endpoints returns
+            // Disabled these checks for last access cache as there is no guarantees when it is updated compared to our get, so these tests will just intermittently fail
+            //List<(RefRecord, DateTime)> accessedRecords = await _lastAccessCache!.GetLastAccessedRecords();
+            //Assert.AreEqual(1, accessedRecords.Count);
+            //(RefRecord record, DateTime _) = accessedRecords.First();
+            //Assert.AreEqual("testObject", record.RefName);
+        }
+
+        [TestMethod]
+        public async Task GetFullRefObjectAsCompactBinary()
+        {
+            HttpRequestMessage request = new HttpRequestMessage(HttpMethod.Get, $"api/v1/c/ddc/{TestNamespace}/bucket/testObjectWithMetadata.uecb");
+            HttpResponseMessage response = await _httpClient!.SendAsync(request);
+
+            Assert.AreEqual(CustomMediaTypeNames.UnrealCompactBinary, response.Content.Headers.ContentType!.MediaType!);
+            response.EnsureSuccessStatusCode();
+
+            byte[] buffer = await response.Content.ReadAsByteArrayAsync();
+            ReadOnlyMemory<byte> memory = new ReadOnlyMemory<byte>(buffer);
+            CompactBinaryObject o = CompactBinaryObject.Load(ref memory);
+            List<CompactBinaryField> fields = o.GetFields().ToList();
+            Assert.AreEqual(4, fields.Count);
+            Assert.AreEqual($"{TestNamespace}.bucket.testObjectWithMetadata", o["name"]!.AsString());
+            // we do not support serializing random dictionaries to compact binary
+            Assert.IsNull(o["metadata"]);
+
+            Assert.IsNotNull(o["lastAccessTime"]);
+
+            Assert.IsNotNull(o["blob"]);
+            CollectionAssert.AreEqual(TestObjectData, o["blob"]!.AsBinary());
+        }
+
+        [TestMethod]
+        public async Task GetRefWithMetadata()
+        {
+            Task<HttpResponseMessage> response = _httpClient!.GetAsync($"api/v1/c/ddc/{TestNamespace}/bucket/testObjectWithMetadata");
+            await response;
+            HttpResponseMessage result = response.Result;
+
+            result.EnsureSuccessStatusCode();
+            RefResponse content = await result.Content.ReadAsAsync<RefResponse>();
+            Assert.AreEqual(expected: $"{TestNamespace}.bucket.testObjectWithMetadata", content.Name);
+            Assert.IsNotNull(content.Metadata);
+            Assert.IsTrue(content.Metadata!.ContainsKey("Foo"));
+            Assert.AreEqual("Bar", content!.Metadata["Foo"]);
+
+            Assert.IsNotNull(content.Blob);
+            CollectionAssert.AreEqual(TestObjectData, content.Blob);
+        }
+
+        [TestMethod]
+        public async Task GetRefRawObject()
+        {
+            Task<HttpResponseMessage> response = _httpClient!.GetAsync($"api/v1/c/ddc/{TestNamespace}/bucket/testObject.raw");
+            await response;
+            HttpResponseMessage result = response.Result;
+
+            result.EnsureSuccessStatusCode();
+            byte[] content = await result.Content.ReadAsByteArrayAsync();
+            CollectionAssert.AreEqual(TestObjectData, content);
+        }
+
+        [TestMethod]
+        public async Task GetRefObjectFieldFilteringAwayBlob()
+        {
+            // fetch refs records once first so that any previous runs of tests are ignored
+            List<(RefRecord, DateTime)> _ = await _lastAccessCache!.GetLastAccessedRecords();
+
+            Task<HttpResponseMessage> response = _httpClient!.GetAsync($"api/v1/c/ddc/{TestNamespace}/bucket/testObject?fields=name&fields=blobIdentifiers");
+            HttpResponseMessage result = await response;
+
+            result.EnsureSuccessStatusCode();
+
+            string content = await result.Content.ReadAsStringAsync();
+            JToken? jsonResult = JsonConvert.DeserializeObject(content) as JToken;
+
+            Assert.IsNotNull(jsonResult);
+
+            Assert.AreEqual($"{TestNamespace}.bucket.testObject", jsonResult!["name"]);
+            Assert.AreEqual(TestObjectBlobHash, new BlobIdentifier(jsonResult!["blobIdentifiers"]!.Value<JArray>()![0]!.Value<string>()!));
+            Assert.IsTrue(jsonResult["blob"] == null, "Expecting blob field to not be present");
+            Assert.IsTrue(jsonResult["contentHash"] == null, "Expecting contentHash field to not be present");
+        }
+
+        
+        [TestMethod]
+        public async Task GetRefObjectEmptyFields()
+        {
+            // fetch refs records once first so that any previous runs of tests are ignored
+            List<(RefRecord, DateTime)> _ = await _lastAccessCache!.GetLastAccessedRecords();
+
+            Task<HttpResponseMessage> response = _httpClient!.GetAsync($"api/v1/c/ddc/{TestNamespace}/bucket/testObject?fields=");
+            HttpResponseMessage result = await response;
+
+            result.EnsureSuccessStatusCode();
+
+            RefResponse content = await result.Content.ReadAsAsync<RefResponse>();
+            Assert.AreEqual(expected: $"{TestNamespace}.bucket.testObject", content.Name);
+            Assert.IsNull(content.Metadata);
+            Assert.IsNotNull(content.Blob);
+            CollectionAssert.AreEqual(TestObjectData, content.Blob);
+        }
+
+
+        [TestMethod]
+        public async Task GetNonExistentIoObject()
+        {
+            HttpResponseMessage response = await _httpClient!.GetAsync($"api/v1/c/ddc/{TestNamespace}/bucket/notUploadedToIo");
+
+            Assert.AreEqual(HttpStatusCode.BadRequest, response.StatusCode);
+
+            string content = await response.Content.ReadAsStringAsync();
+            JToken? jsonResult = JsonConvert.DeserializeObject(content) as JToken;
+            Assert.AreEqual("Object EBA4BE7385F365E6AC2354096A6CA2C5AF412A0F in test-namespace not found", jsonResult!["title"]!.ToString());
+        }
+
+
+        [TestMethod]
+        public async Task PutRefObject()
+        {
+            byte[] payload = Encoding.ASCII.GetBytes("I have some new content");
+            ContentHash payloadHash = ContentHash.FromBlob(payload);
+            using ByteArrayContent requestContent = new ByteArrayContent(payload);
+            requestContent.Headers.ContentType = new MediaTypeHeaderValue("application/octet-stream");
+            requestContent.Headers.ContentLength = payload.Length;
+            requestContent.Headers.Add(CommonHeaders.HashHeaderName, payloadHash.ToString());
+
+            HttpResponseMessage result =
+                await _httpClient!.PutAsync(requestUri: $"api/v1/c/ddc/{TestNamespace}/bucket/newObject", requestContent);
+
+            result.EnsureSuccessStatusCode();
+        }
+
+        [TestMethod]
+        public async Task PutGetStructuredRefObject()
+        {
+            string jsonObject = $"{{\"metadata\": null, \"blobReferences\": [\"{TestObjectHash}\"], \"contentHash\" : \"{TestObjectHash}\"}}";
+            using StringContent requestContent = new StringContent(jsonObject);
+            requestContent.Headers.ContentType = new MediaTypeHeaderValue("application/json");
+            requestContent.Headers.ContentLength = jsonObject.Length;
+            requestContent.Headers.Add(CommonHeaders.HashHeaderName, TestObjectHash.ToString());
+
+            HttpResponseMessage result = await _httpClient!.PutAsync(requestUri: $"api/v1/c/ddc/{TestNamespace}/bucket/newStructuredObject", requestContent);
+            result.EnsureSuccessStatusCode();
+
+            HttpResponseMessage getResponse = await _httpClient.GetAsync($"api/v1/c/ddc/{TestNamespace}/bucket/newStructuredObject.raw");
+            getResponse.EnsureSuccessStatusCode();
+            await using MemoryStream ms = new MemoryStream();
+            await getResponse.Content.CopyToAsync(ms);
+
+            byte[] roundTrippedBuffer = ms.ToArray();
+            string roundTrippedPayload = Encoding.ASCII.GetString(roundTrippedBuffer);
+
+            Assert.AreEqual(TestObjectContents, roundTrippedPayload);
+            CollectionAssert.AreEqual(TestObjectData, roundTrippedBuffer);
+        }
+
+
+        [TestMethod]
+        public async Task PutGetLargeObject()
+        {
+            // we chunk at 1 kb to lets generate a 3kb string
+            string content = string.Concat(Enumerable.Repeat("Duplicate string", 200));
+            byte[] payload = Encoding.ASCII.GetBytes(content);
+            ContentHash payloadHash = ContentHash.FromBlob(payload);
+            using ByteArrayContent requestContent = new ByteArrayContent(payload);
+            requestContent.Headers.ContentType = new MediaTypeHeaderValue("application/octet-stream");
+            requestContent.Headers.ContentLength = payload.Length;
+            requestContent.Headers.Add(CommonHeaders.HashHeaderName, payloadHash.ToString());
+
+            HttpResponseMessage result =
+                await _httpClient!.PutAsync(requestUri: $"api/v1/c/ddc/{TestNamespace}/bucket/largeObject", requestContent);
+
+            result.EnsureSuccessStatusCode();
+
+
+            HttpResponseMessage getResponse = await _httpClient.GetAsync($"api/v1/c/ddc/{TestNamespace}/bucket/largeObject.raw");
+            getResponse.EnsureSuccessStatusCode();
+            await using MemoryStream ms = new MemoryStream();
+            await getResponse.Content.CopyToAsync(ms);
+
+            byte[] roundTrippedBuffer = ms.ToArray();
+            string roundTrippedPayload = Encoding.ASCII.GetString(roundTrippedBuffer);
+
+            Assert.AreEqual(content.Length, roundTrippedPayload.Length, "Expected the same object, length was different");
+
+            Assert.AreEqual(content, roundTrippedPayload);
+
+            CollectionAssert.AreEqual(payload, roundTrippedBuffer);
+        }
+
+
+        [TestMethod]
+        public async Task DeleteRefObject()
+        {
+            HttpResponseMessage deleteResponse =
+                await _httpClient!.DeleteAsync($"api/v1/c/ddc/{TestNamespace}/bucket/deletableObject");
+            deleteResponse.EnsureSuccessStatusCode();
+        }
+
+        [TestMethod]
+        public async Task DeleteNonExistentObject()
+        {
+            HttpResponseMessage deleteResponse =
+                await _httpClient!.DeleteAsync($"api/v1/c/ddc/{TestNamespace}/bucket/thisDoesNotExistForDelete");
+            Assert.AreEqual(System.Net.HttpStatusCode.BadRequest, deleteResponse.StatusCode);
+        }
+
+        [TestMethod]
+        public async Task FullFlow()
+        {
+            byte[] payload = Encoding.ASCII.GetBytes("Foo bar");
+            ContentHash payloadHash = ContentHash.FromBlob(payload);
+            using ByteArrayContent requestContent = new ByteArrayContent(payload);
+            requestContent.Headers.ContentType = new MediaTypeHeaderValue("application/octet-stream");
+            requestContent.Headers.ContentLength = payload.Length;
+            requestContent.Headers.Add(CommonHeaders.HashHeaderName, payloadHash.ToString());
+
+            HttpResponseMessage putResponse =
+                await _httpClient!.PutAsync(requestUri: $"api/v1/c/ddc/{TestNamespace}/bucket/fullFlowObject",
+                    requestContent);
+            putResponse.EnsureSuccessStatusCode();
+            await putResponse.Content.ReadAsStringAsync();
+
+            HttpResponseMessage getResponse =
+                await _httpClient.GetAsync($"api/v1/c/ddc/{TestNamespace}/bucket/fullFlowObject");
+            getResponse.EnsureSuccessStatusCode();
+            RefResponse content = await getResponse.Content.ReadAsAsync<RefResponse>();
+            Assert.AreEqual(expected: $"{TestNamespace}.bucket.fullFlowObject", content.Name);
+            Assert.IsNull(content.Metadata);
+            Assert.IsNotNull(content.Blob);
+            CollectionAssert.AreEqual(payload, content.Blob);
+
+            HttpResponseMessage deleteResponse =
+                await _httpClient.DeleteAsync($"api/v1/c/ddc/{TestNamespace}/bucket/fullFlowObject");
+            deleteResponse.EnsureSuccessStatusCode();
+            Assert.AreEqual(System.Net.HttpStatusCode.NoContent, deleteResponse.StatusCode);
+        }
+
+        [TestMethod]
+        public async Task DeleteNamespace()
+        {
+            NamespaceId deleteNamespace = new NamespaceId("delete-namespace");
+
+            byte[] payload = Encoding.ASCII.GetBytes("Foo bar");
+            ContentHash payloadHash = ContentHash.FromBlob(payload);
+            using ByteArrayContent requestContent = new ByteArrayContent(payload);
+            requestContent.Headers.ContentType = new MediaTypeHeaderValue("application/octet-stream");
+            requestContent.Headers.ContentLength = payload.Length;
+            requestContent.Headers.Add(CommonHeaders.HashHeaderName, payloadHash.ToString());
+
+            HttpResponseMessage putResponse = await _httpClient!.PutAsync(requestUri: $"api/v1/c/ddc/{deleteNamespace}/bucket/deleteNamespaceObject", requestContent);
+            putResponse.EnsureSuccessStatusCode();
+            await putResponse.Content.ReadAsStringAsync();
+
+            HttpResponseMessage deleteResponse = await _httpClient.DeleteAsync($"api/v1/c/ddc/{deleteNamespace}");
+            deleteResponse.EnsureSuccessStatusCode();
+            Assert.AreEqual(System.Net.HttpStatusCode.NoContent, deleteResponse.StatusCode);
+
+            HttpResponseMessage getResponse = await _httpClient.GetAsync($"api/v1/c/ddc/{TestNamespace}/bucket/deleteNamespaceObject");
+            Assert.AreEqual(System.Net.HttpStatusCode.BadRequest, getResponse.StatusCode);
+        }
+
+        protected abstract Task<RefRecord> GetTestRecord(BucketId bucket, KeyId name, IRefsStore.ExtraFieldsFlag fields);
+
+        [TestMethod]
+        public async Task LastAccessRollup()
+        {
+            BucketId bucket = new BucketId("bucket");
+            KeyId key = new KeyId("lastAccessObject");
+            byte[] payload = Encoding.ASCII.GetBytes("This payload does not matter");
+            ContentHash payloadHash = ContentHash.FromBlob(payload);
+            using ByteArrayContent requestContent = new ByteArrayContent(payload);
+            requestContent.Headers.ContentType = new MediaTypeHeaderValue("application/octet-stream");
+            requestContent.Headers.ContentLength = payload.Length;
+            requestContent.Headers.Add(CommonHeaders.HashHeaderName, payloadHash.ToString());
+
+            HttpResponseMessage putResponse =
+                await _httpClient!.PutAsync(requestUri: $"api/v1/c/ddc/{TestNamespace}/{bucket}/{key}",
+                    requestContent);
+            putResponse.EnsureSuccessStatusCode();
+            await putResponse.Content.ReadAsStringAsync();
+
+            RefRecord initialLastAccessObject = await GetTestRecord(bucket, key, IRefsStore.ExtraFieldsFlag.LastAccess);
+            Assert.AreEqual(key, initialLastAccessObject.RefName);
+            Assert.IsTrue(initialLastAccessObject.LastAccessTime.HasValue);
+            DateTime creationTime = initialLastAccessObject.LastAccessTime!.Value;
+
+            HttpResponseMessage getResponse =
+                await _httpClient.GetAsync($"api/v1/c/ddc/{TestNamespace}/{bucket}/{key}");
+            getResponse.EnsureSuccessStatusCode();
+            RefResponse content = await getResponse.Content.ReadAsAsync<RefResponse>();
+            Assert.AreEqual(expected: $"{TestNamespace}.bucket.lastAccessObject", content.Name);
+
+            StringContent jsonContent = new StringContent("", Encoding.UTF8, "application/json");
+            HttpResponseMessage lastAccessRollupResponse =
+                await _httpClient.PostAsync(requestUri: $"api/v1/admin/startLastAccessRollup", jsonContent);
+            lastAccessRollupResponse.EnsureSuccessStatusCode();
+
+            RefRecord doc = await GetTestRecord(bucket, key, IRefsStore.ExtraFieldsFlag.LastAccess);
+            Assert.AreEqual(key, doc.RefName);
+            Assert.IsTrue(doc.LastAccessTime > creationTime);
+
+            // make sure the refs is empty
+            List<(RefRecord, DateTime)> records = await _lastAccessCache!.GetLastAccessedRecords();
+            Assert.AreEqual(0, records.Count);
+        }
+
+
+        [TestMethod]
+        public async Task LastAccess()
+        {
+            IRefsStore refsStore = _server!.Services.GetService<IRefsStore>()!;
+
+            OldRecord[] emptyRecords = await refsStore.GetOldRecords(TestNamespace, TimeSpan.FromDays(7)).ToArrayAsync();
+            // no content older then 7 days as we just insert it
+            Assert.AreEqual(0, emptyRecords.Length);
+
+            // check content inserted in the last hour, which should return the 3 objects set as inserted a day ago, but not the one that just got inserted
+            OldRecord[] oldRecords = await refsStore.GetOldRecords(TestNamespace, TimeSpan.FromHours(1)).ToArrayAsync();
+
+            List<KeyId> validRefs = new List<KeyId>()
+            {
+                new KeyId("deletableObject"),
+                new KeyId("testObject"),
+                new KeyId("testObjectWithMetadata"),
+                new KeyId("notUploadedToIo")
+            };
+            foreach(OldRecord record in oldRecords)
+            {
+                Assert.IsTrue(validRefs.Contains(record.RefName));
+            }
+
+            // we should have our 4 pre-seeded objects
+            Assert.AreEqual(4, oldRecords.Length);
+        }
+
+        public class LastAccessResponse
+        {
+            public class LastAccessRefRecord
+            {
+                public string Name { get; set; } = null!;
+                public string Bucket { get; set; } = null!;
+                public string Namespace { get; set; } = null!;
+            }
+
+            public LastAccessRefRecord[] Records { get; set; } = null!;
+        }
+
+        [TestMethod]
+        public async Task Batch()
+        {
+            // we chunk at 1 kb to lets generate a 3kb string
+            string content = string.Concat(Enumerable.Repeat("Duplicate string", 200));
+            byte[] payload = Encoding.ASCII.GetBytes(content);
+            ContentHash payloadHash = ContentHash.FromBlob(payload);
+
+            var ops = new
+            {
+                Operations = new object[]
+                {
+                    new
+                    {
+                        Op = "GET",
+                        Namespace = TestNamespace,
+                        Bucket = "bucket",
+                        Id = "testObject",
+                    },
+                    new
+                    {
+                        Namespace = TestNamespace,
+                        Op = "GET",
+                        Bucket = "bucket",
+                        Id = "thisDoesNotExist",
+                    },
+                    new
+                    {
+                        Namespace = TestNamespace,
+                        Op = "PUT",
+                        Bucket = "bucket",
+                        Id = "largeBatchObject",
+                        Content = payload,
+                        ContentHash = payloadHash
+                    }
+                }
+            };
+
+            HttpResponseMessage response = await _httpClient.PostAsJsonAsync("api/v1/c/ddc-rpc", ops);
+            response.EnsureSuccessStatusCode();
+
+            string stringResult = await response.Content.ReadAsStringAsync();
+            object? result = JsonConvert.DeserializeObject(stringResult);
+            JArray? array = result as JArray;
+            Assert.IsNotNull(array);
+            Assert.AreEqual(3, array!.Count);
+
+            Assert.AreEqual($"{TestNamespace}.bucket.testObject", array![0]["name"]);
+            Assert.AreEqual(TestObjectHash, new ContentHash(array![0]["contentHash"]!.Value<string>()!));
+            Assert.AreEqual(TestObjectBlobHash, new BlobIdentifier(array![0]["blobIdentifiers"]![0]!.Value<string>()!));
+            Assert.AreEqual("Test Object Data", Encoding.ASCII.GetString(Convert.FromBase64String(array![0]["blob"]!.Value<string>()!)));
+
+            Assert.AreEqual("Exception thrown", array![1]["title"]);
+            Assert.AreEqual(0L, array![2]["transactionId"]!.Value<long>());
+        }
+
+        [TestMethod]
+        public async Task BatchGet()
+        {
+            var ops = new
+            {
+                Namespace = TestNamespace,
+                Operations = new object[]
+                {
+                    new
+                    {
+                        Bucket = "bucket",
+                        Key = "testObject",
+                    },
+                    new
+                    {
+                        Bucket = "bucket",
+                        Key = "thisDoesNotExist",
+                    },
+                    new
+                    {
+                        Bucket = "bucket",
+                        Key = "headCheckObjectNotExist",
+                        Verb = "HEAD"
+                    },
+                    new
+                    {
+                        Bucket = "bucket",
+                        Key = "testObjectWithMetadata",
+                        Verb = "HEAD"
+                    },
+                    new
+                    {
+                        Bucket = "bucket",
+                        Key = "notUploadedToIo",
+                    }
+                }
+            };
+
+            HttpResponseMessage response = await _httpClient.PostAsJsonAsync("api/v1/c/ddc-rpc/batchGet", ops);
+            response.EnsureSuccessStatusCode();
+
+            Stream stream = await response.Content.ReadAsStreamAsync();
+            await using MemoryStream ms = new MemoryStream();
+            // read the entire stream into a buffer before we start processing it
+            await stream.CopyToAsync(ms);
+            byte[] buffer = ms.ToArray();
+            ReadOnlyMemory<byte> span = new ReadOnlyMemory<byte>(buffer);
+            // first 4 bytes is a magic header
+            Assert.AreEqual("JPTR", Encoding.ASCII.GetString(span.Slice(0, 4).Span));
+            span = span.Slice(4);
+            const uint ExpectedCountOfObjects = 5u;
+            // next we expect a 32 bit unsigned integer of the number of results
+            Assert.AreEqual(ExpectedCountOfObjects, BitConverter.ToUInt32(span.Slice(0, 4).Span));
+            span = span.Slice(4);
+
+            for (int i = 0; i < ExpectedCountOfObjects; i++)
+            {
+                // first 4 bytes is another magic header
+                Assert.AreEqual("JPEE", Encoding.ASCII.GetString(span.Slice(0, 4).Span));
+                span = span.Slice(4);
+
+                // count length of string
+                int count = 0;
+                ReadOnlyMemory<byte> stringMemory = span;
+                while (stringMemory.Slice(0, 1).ToArray()[0] != '\0')
+                {
+                    count++;
+                    stringMemory = stringMemory.Slice(1);
+                }
+
+                // next we should find the full key to the object
+                string key = Encoding.ASCII.GetString(span.Slice(0, count).Span);
+                span = span.Slice(count+1); // +1 because we want to also remove the null terminator
+
+                byte state = span.Slice(0, 1).ToArray()[0];
+                span = span.Slice(1);
+
+                if (key == $"{TestNamespace}.bucket.testObject")
+                {
+                    Assert.AreEqual(0, state);
+
+                    // next is a hash (20 bytes)
+                    ContentHash hash = new ContentHash(span.Slice(0, 20).ToArray());
+                    Assert.AreEqual(TestObjectHash, hash);
+                    span = span.Slice(20);
+
+                    // next we expect a 64 bit uint of the number of results
+                    ulong lengthOfBlob = BitConverter.ToUInt64(span.Slice(0, 8).Span);
+                    Assert.AreEqual((uint)TestObjectData.Length, lengthOfBlob);
+                    span = span.Slice(8);
+
+                    // lastly we read the object
+                    byte[] blob = span.Slice(0, (int)lengthOfBlob).ToArray();
+                    CollectionAssert.AreEqual(TestObjectData, blob);
+                    span = span.Slice((int)lengthOfBlob);
+
+                } 
+                else if (key == $"{TestNamespace}.bucket.thisDoesNotExist")
+                {
+                    // this object was not found, thus its state will be NotFound but the test is empty
+                    Assert.AreEqual(2, state);
+                }
+                else if (key == $"{TestNamespace}.bucket.notUploadedToIo")
+                {
+                    // this object does exist but a payload is missing from the blob store
+                    Assert.AreEqual(2, state);
+                }
+                else if (key == $"{TestNamespace}.bucket.headCheckObjectNotExist")
+                {
+                    // this object was not found, thus its state will be NotFound but the test is empty
+                    Assert.AreEqual(2, state);
+                }
+                else if (key == $"{TestNamespace}.bucket.testObjectWithMetadata")
+                {
+                    // this object should exist but as its a head check will only return the state
+                    Assert.AreEqual(3, state);
+                }
+                else
+                    throw new NotImplementedException();
+            }
+
+            // we have reached the end of the result thus the stream should be empty
+            Assert.IsTrue(span.IsEmpty);
+        }
+    }
+}
