Moved Horde.Storage from git to p4

diff --git a/Engine/Source/Programs/Horde/HordeStorage/Horde.Storage/Tests/Functional/StorageTests.cs b/Engine/Source/Programs/Horde/HordeStorage/Horde.Storage/Tests/Functional/StorageTests.cs
new file mode 100644
--- /dev/null
+++ b/Engine/Source/Programs/Horde/HordeStorage/Horde.Storage/Tests/Functional/StorageTests.cs
@@ -0,0 +1,802 @@
+// Copyright Epic Games, Inc. All Rights Reserved.
+
+using System;
+using System.Collections.Generic;
+using System.IO;
+using System.Linq;
+using System.Net;
+using System.Net.Http;
+using System.Net.Http.Headers;
+using System.Net.Http.Json;
+using System.Net.Mime;
+using System.Text;
+using System.Threading;
+using System.Threading.Tasks;
+using Amazon.S3;
+using Amazon.S3.Model;
+using Azure.Storage.Blobs;
+using Dasync.Collections;
+using Horde.Storage.Controllers;
+using Horde.Storage.Implementation;
+using Jupiter;
+using Jupiter.Implementation;
+using Microsoft.AspNetCore.TestHost;
+using Microsoft.AspNetCore.Hosting;
+using Microsoft.AspNetCore.Mvc;
+using Microsoft.AspNetCore.Mvc.Formatters;
+using Microsoft.Extensions.Configuration;
+using Microsoft.Extensions.DependencyInjection;
+using Microsoft.Extensions.Options;
+using Microsoft.VisualStudio.TestTools.UnitTesting;
+using Newtonsoft.Json;
+using Newtonsoft.Json.Linq;
+using Serilog;
+using Serilog.Core;
+
+namespace Horde.Storage.FunctionalTests.Storage
+{
+    [TestClass]
+    public class S3StorageTests : StorageTests
+    {
+        private IAmazonS3? _s3;
+        protected override IEnumerable<KeyValuePair<string, string>> GetSettings()
+        {
+            return new[]
+            {
+                new KeyValuePair<string, string>("Horde.Storage:StorageImplementations:0", HordeStorageSettings.StorageBackendImplementations.S3.ToString()),
+                new KeyValuePair<string, string>("S3:BucketName", $"tests-{TestNamespaceName}")
+            };
+        }
+
+        protected override async Task Seed(IServiceProvider provider)
+        {
+            string S3BucketName = $"tests-{TestNamespaceName}";
+
+            _s3 = provider.GetService<IAmazonS3>();
+            Assert.IsNotNull(_s3);
+            if (await _s3!.DoesS3BucketExistAsync(S3BucketName))
+            {
+                // if we have failed to run the cleanup for some reason we run it now
+                await Teardown();
+            }
+
+            await _s3.PutBucketAsync(S3BucketName);
+            await _s3.PutObjectAsync(new PutObjectRequest { BucketName = S3BucketName, Key = _smallFileHash.AsS3Key(), ContentBody = SmallFileContents });
+            await _s3.PutObjectAsync(new PutObjectRequest { BucketName = S3BucketName, Key = _anotherFileHash.AsS3Key(), ContentBody = AnotherFileContents });
+            await _s3.PutObjectAsync(new PutObjectRequest { BucketName = S3BucketName, Key = _deleteFileHash.AsS3Key(), ContentBody = DeletableFileContents });
+            await _s3.PutObjectAsync(new PutObjectRequest {BucketName = S3BucketName, Key = _oldBlobFileHash.AsS3Key(), ContentBody = OldFileContents});
+        }
+
+        protected override async Task Teardown()
+        {
+            string S3BucketName = $"tests-{TestNamespaceName}";
+            ListObjectsResponse response = await _s3!.ListObjectsAsync(S3BucketName);
+            List<KeyVersion> objectKeys = response.S3Objects.Select(o => new KeyVersion { Key = o.Key }).ToList();
+            await _s3.DeleteObjectsAsync(new DeleteObjectsRequest { BucketName = S3BucketName, Objects = objectKeys });
+
+            await _s3.DeleteBucketAsync(S3BucketName);
+        }
+    }
+
+    [TestClass]
+    public class AzureStorageTests : StorageTests
+    {
+        private AzureSettings? _settings;
+
+        protected override IEnumerable<KeyValuePair<string, string>> GetSettings()
+        {
+            return new[] {new KeyValuePair<string, string>("Horde.Storage:StorageImplementations:0", HordeStorageSettings.StorageBackendImplementations.Azure.ToString())};
+        }
+
+        protected override async Task Seed(IServiceProvider provider)
+        {
+            _settings = provider.GetService<IOptionsMonitor<AzureSettings>>()!.CurrentValue;
+
+            string connectionString = _settings.ConnectionString;
+            BlobContainerClient container = new BlobContainerClient(connectionString, TestNamespaceName.ToString());
+
+            if (await container.ExistsAsync())
+            {
+                // if we have failed to run the cleanup for some reason we run it now
+                await Teardown();
+            }
+
+            await container.CreateAsync();
+            BlobClient smallBlob = container.GetBlobClient(_smallFileHash.ToString());
+            byte[] smallContents = Encoding.ASCII.GetBytes(SmallFileContents);
+            await using MemoryStream smallBlobStream = new MemoryStream(smallContents);
+            await smallBlob.UploadAsync(smallBlobStream);
+
+            BlobClient anotherBlob = container.GetBlobClient(_anotherFileHash.ToString());
+            byte[] anotherContents = Encoding.ASCII.GetBytes(AnotherFileContents);
+            await using MemoryStream anotherContentsStream = new MemoryStream(anotherContents);
+            await anotherBlob.UploadAsync(anotherContentsStream);
+
+            BlobClient deleteBlob = container.GetBlobClient(_deleteFileHash.ToString());
+            byte[] deleteContents = Encoding.ASCII.GetBytes(DeletableFileContents);
+            await using MemoryStream deleteContentsStream = new MemoryStream(deleteContents);
+            await deleteBlob.UploadAsync(deleteContentsStream);
+
+            BlobClient oldBlob = container.GetBlobClient(_oldBlobFileHash.ToString());
+            byte[] oldBlobContents = Encoding.ASCII.GetBytes(OldFileContents);
+            await using MemoryStream oldBlobContentsSteam = new MemoryStream(oldBlobContents);
+            await oldBlob.UploadAsync(oldBlobContentsSteam);
+        }
+
+        protected override async Task Teardown()
+        {
+            string connectionString = _settings!.ConnectionString;
+            BlobContainerClient container = new BlobContainerClient(connectionString, TestNamespaceName.ToString());
+
+            await container.DeleteAsync();
+        }
+    }
+
+    
+    [TestClass]
+    public class FileSystemStoreTests : StorageTests
+    {
+        private string _localTestDir;
+        private readonly NamespaceId FooNamespace = new NamespaceId("foo");
+
+        public FileSystemStoreTests()
+        {
+            _localTestDir = Path.Combine(Path.GetTempPath(), "IoFileSystemTests", Path.GetRandomFileName());
+        }
+        protected override IEnumerable<KeyValuePair<string, string>> GetSettings()
+        {
+            return new[] { 
+                new KeyValuePair<string, string>("Horde.Storage:StorageImplementations:0", HordeStorageSettings.StorageBackendImplementations.FileSystem.ToString()),
+                new KeyValuePair<string, string>("Filesystem:RootDir", _localTestDir)
+            };
+        }
+
+
+        protected override async Task Seed(IServiceProvider provider)
+        {
+            NamespaceId folderName = TestNamespaceName;
+            
+            Directory.CreateDirectory(_localTestDir);
+
+            FileInfo smallFileInfo = FileSystemStore.GetFilesystemPath(_localTestDir, folderName, _smallFileHash);
+            smallFileInfo.Directory?.Create();
+            await File.WriteAllBytesAsync(
+                smallFileInfo.FullName,
+                Encoding.ASCII.GetBytes(SmallFileContents)
+            );
+
+            FileInfo anotherFileInfo = FileSystemStore.GetFilesystemPath(_localTestDir, folderName, _anotherFileHash);
+            anotherFileInfo.Directory?.Create();
+            await File.WriteAllBytesAsync(
+                anotherFileInfo.FullName,
+                Encoding.ASCII.GetBytes(AnotherFileContents)
+            );
+            
+            FileInfo deleteFileInfo = FileSystemStore.GetFilesystemPath(_localTestDir, folderName, _deleteFileHash);
+            deleteFileInfo.Directory?.Create();
+            await File.WriteAllBytesAsync(
+                deleteFileInfo.FullName,
+                Encoding.ASCII.GetBytes(DeletableFileContents)
+            );
+
+            // a old file used to verify cutoff filtering
+            FileInfo oldFileInfo = FileSystemStore.GetFilesystemPath(_localTestDir, folderName, _oldBlobFileHash);
+            oldFileInfo.Directory?.Create();
+            await File.WriteAllBytesAsync(
+                oldFileInfo.FullName,
+                Encoding.ASCII.GetBytes(OldFileContents)
+            );
+            File.SetLastWriteTimeUtc(oldFileInfo.FullName, DateTime.Now.AddDays(-7));
+        }
+
+        protected override Task Teardown()
+        {
+            Directory.Delete(_localTestDir, true);
+
+            return Task.CompletedTask;
+        }
+
+        // only a file system allows us to update the last modified time of the object to actually execute this test
+        [TestMethod]
+        public async Task ListOldBlobs()
+        {
+            IBlobStore blobStore = _server!.Services.GetService<IBlobStore>()!;
+
+            // we fetch all objects that are more then a day old
+            // as the old blob is set to be a week old this should be the only object returned
+            DateTime cutoff = DateTime.Now.AddDays(-1);
+            {
+                BlobIdentifier[] blobs = await blobStore.ListOldObjects(TestNamespaceName, cutoff).ToArrayAsync();
+                Assert.AreEqual(_oldBlobFileHash, blobs[0]);
+            }
+        }
+
+        [TestMethod]
+        public async Task ListNamespaces()
+        {
+            IBlobStore blobStore = _server!.Services.GetService<IBlobStore>()!;
+            FileSystemStore fsStore = (blobStore as FileSystemStore)!;
+
+            List<NamespaceId> namespaces = await fsStore.ListNamespaces().ToListAsync();
+            Assert.AreEqual(1, namespaces.Count);
+            
+            await fsStore.PutObject(FooNamespace, Encoding.ASCII.GetBytes(SmallFileContents), _smallFileHash);
+            
+            namespaces = await fsStore.ListNamespaces().ToListAsync();
+            Assert.AreEqual(2, namespaces.Count);
+        }
+
+        [TestMethod]
+        public async Task GarbageCollect()
+        {
+            // Remove data added from test seeding
+            Directory.Delete(_localTestDir, true);
+            Directory.CreateDirectory(_localTestDir);
+            
+            FilesystemSettings fsSettings = _server!.Services.GetService<IOptionsMonitor<FilesystemSettings>>()!.CurrentValue;
+            fsSettings.MaxSizeBytes = 600;
+            IBlobStore blobStore = _server!.Services.GetService<IBlobStore>()!;
+            FileSystemStore fsStore = (blobStore as FileSystemStore)!;
+
+            CancellationTokenSource cts = new CancellationTokenSource();
+            Assert.IsTrue((await fsStore.CleanupInternal(cts.Token, batchSize: 2)).Count == 0); // No garbage to collect, should return false
+            
+            FileInfo[] fooFiles = CreateFilesInNamespace(FooNamespace, 10);
+
+            Assert.AreEqual(10 * 100, await fsStore.CalculateDiskSpaceUsed());
+
+            Assert.IsTrue((await fsStore.CleanupInternal(cts.Token, batchSize: 2)).Count > 0);
+
+            Assert.AreEqual(5 * 100, await fsStore.CalculateDiskSpaceUsed());
+
+            fooFiles.ToList().ForEach(x => x.Refresh());
+            Assert.IsTrue(fooFiles[0].Exists); // Most recently accessed/modified
+            Assert.IsTrue(fooFiles[1].Exists);
+            Assert.IsTrue(fooFiles[2].Exists);
+            Assert.IsTrue(fooFiles[3].Exists);
+            Assert.IsTrue(fooFiles[4].Exists);
+            Assert.IsFalse(fooFiles[5].Exists);
+            Assert.IsFalse(fooFiles[6].Exists);
+            Assert.IsFalse(fooFiles[7].Exists);
+            Assert.IsFalse(fooFiles[8].Exists);
+            Assert.IsFalse(fooFiles[9].Exists); // Least recently accessed/modified
+        }
+        
+        [TestMethod]
+        public async Task GetLeastRecentlyAccessedObjects()
+        {
+            FileInfo[] fooFiles = CreateFilesInNamespace(FooNamespace, 10);
+            CreateFilesInNamespace(new NamespaceId("bar"), 10);
+
+            IBlobStore blobStore = _server!.Services.GetService<IBlobStore>()!;
+            FileSystemStore fsStore = (blobStore as FileSystemStore)!;
+
+            Assert.AreEqual(10, (await fsStore.GetLeastRecentlyAccessedObjects(FooNamespace)).Length);
+
+            FileInfo[] results = await fsStore.GetLeastRecentlyAccessedObjects(FooNamespace, 3);
+            Assert.AreEqual(3, results.Length);
+            Assert.AreEqual(fooFiles[7].LastAccessTime, results[2].LastAccessTime);
+            Assert.AreEqual(fooFiles[8].LastAccessTime, results[1].LastAccessTime);
+            Assert.AreEqual(fooFiles[9].LastAccessTime, results[0].LastAccessTime);
+        }
+        
+        [TestMethod]
+        public async Task CalculateUsedDiskSpace()
+        {
+            IBlobStore blobStore = _server!.Services.GetService<IBlobStore>()!;
+            FileSystemStore fsStore = (blobStore as FileSystemStore)!;
+            await fsStore.PutObject(FooNamespace, Encoding.ASCII.GetBytes(SmallFileContents), _smallFileHash);
+            await fsStore.PutObject(FooNamespace, Encoding.ASCII.GetBytes(AnotherFileContents), _anotherFileHash);
+            
+            Assert.AreEqual(SmallFileContents.Length + AnotherFileContents.Length, await fsStore.CalculateDiskSpaceUsed(FooNamespace));
+        }
+        
+        private FileInfo[] CreateFilesInNamespace(NamespaceId ns, int numFiles)
+        {
+            FileInfo[] files = new FileInfo[numFiles];
+            for (int i = 0; i < numFiles; i++)
+            {
+                byte[] content = Encoding.ASCII.GetBytes(i + 1000 + new string('j', 96));
+                BlobIdentifier bi = BlobIdentifier.FromBlob(content);
+                FileInfo fi = FileSystemStore.GetFilesystemPath(_localTestDir, ns, bi);
+                fi.Directory?.Create();
+                File.WriteAllBytes(fi.FullName, content);
+                fi.LastWriteTime = DateTime.UtcNow.AddDays(-i);
+                fi.Refresh();
+                files[i] = fi;
+            }
+
+            return files;
+        }
+    }
+
+    public abstract class StorageTests
+    {
+        protected TestServer? _server;
+        protected HttpClient? _httpClient;
+
+        protected readonly NamespaceId TestNamespaceName = new NamespaceId("testbucket");
+        protected const string SmallFileContents = "Small file contents";
+        protected const string AnotherFileContents = "Another file with contents";
+        protected const string DeletableFileContents = "Delete Me";
+        protected const string OldFileContents = "a old blob used for testing cutoff filtering";
+
+        protected readonly BlobIdentifier _smallFileHash = BlobIdentifier.FromBlob( Encoding.ASCII.GetBytes(SmallFileContents));
+        protected readonly BlobIdentifier _anotherFileHash = BlobIdentifier.FromBlob(Encoding.ASCII.GetBytes(AnotherFileContents));
+        protected readonly BlobIdentifier _deleteFileHash = BlobIdentifier.FromBlob(Encoding.ASCII.GetBytes(DeletableFileContents));
+        protected readonly BlobIdentifier _oldBlobFileHash = BlobIdentifier.FromBlob(Encoding.ASCII.GetBytes(OldFileContents));
+
+        [TestInitialize]
+        public async Task Setup()
+        {
+            IConfigurationRoot configuration = new ConfigurationBuilder()
+                // we are not reading the base appSettings here as we want exact control over what runs in the tests
+                .AddJsonFile("appsettings.Testing.json", true)
+                .AddInMemoryCollection(GetSettings())
+                .AddEnvironmentVariables()
+                .Build();
+
+            Logger logger = new LoggerConfiguration()
+                .ReadFrom.Configuration(configuration)
+                .CreateLogger();
+
+            TestServer server = new TestServer(new WebHostBuilder()
+                .UseConfiguration(configuration)
+                .UseEnvironment("Testing")
+                .UseSerilog(logger)
+                .UseStartup<HordeStorageStartup>()
+            );
+            _httpClient = server.CreateClient();
+            _server = server;
+
+            // Seed storage
+            await Seed(_server.Services);
+        }
+
+        protected abstract IEnumerable<KeyValuePair<string, string>> GetSettings();
+
+        protected abstract Task Seed(IServiceProvider serverServices);
+        protected abstract Task Teardown();
+
+
+        [TestCleanup]
+        public async Task MyTeardown()
+        {
+            await Teardown();
+        }
+
+
+        [TestMethod]
+        public async Task GetSmallFile()
+        {
+            HttpResponseMessage result = await _httpClient!.GetAsync($"api/v1/s/{TestNamespaceName}/{_smallFileHash}");
+            result.EnsureSuccessStatusCode();
+            string content = await result.Content.ReadAsStringAsync();
+            Assert.AreEqual(SmallFileContents, content);
+        }
+
+
+
+        [TestMethod]
+        public async Task GetNotExistentFile()
+        {
+            byte[] payload = Encoding.ASCII.GetBytes("This content does not exist");
+            ContentHash contentHash = ContentHash.FromBlob(payload);
+            HttpResponseMessage result = await _httpClient!.GetAsync($"api/v1/s/{TestNamespaceName}/{contentHash}");
+
+            Assert.AreEqual(HttpStatusCode.NotFound, result.StatusCode);
+        }
+
+        [TestMethod]
+        public async Task GetInvalidHash()
+        {
+            Task<HttpResponseMessage> response = _httpClient!.GetAsync($"api/v1/s/{TestNamespaceName}/smallFile");
+            await response;
+            HttpResponseMessage result = response.Result;
+
+            Assert.AreEqual(HttpStatusCode.BadRequest, result.StatusCode);
+        }
+
+        [TestMethod]
+        public async Task PutSmallBlob()
+        {
+            byte[] payload = Encoding.ASCII.GetBytes("I am a small blob");
+            using ByteArrayContent requestContent = new ByteArrayContent(payload);
+            requestContent.Headers.ContentType = new MediaTypeHeaderValue("application/octet-stream");
+            requestContent.Headers.ContentLength = payload.Length;
+            BlobIdentifier contentHash = BlobIdentifier.FromBlob(payload);
+            Task<HttpResponseMessage> response =
+                _httpClient!.PutAsync(requestUri: $"api/v1/s/{TestNamespaceName}/{contentHash}", requestContent);
+            await response;
+            HttpResponseMessage result = response.Result;
+
+            result.EnsureSuccessStatusCode();
+            InsertResponse content = await result.Content.ReadAsAsync<InsertResponse>();
+            Assert.AreEqual(contentHash, content.Identifier);
+        }
+
+        [TestMethod]
+        public async Task DeleteBlob()
+        {
+            HttpResponseMessage result = await  _httpClient!.DeleteAsync($"api/v1/s/{TestNamespaceName}/{_deleteFileHash}");
+            result.EnsureSuccessStatusCode();
+
+            Assert.AreEqual(HttpStatusCode.NoContent, result.StatusCode);
+        }
+
+        [TestMethod]
+        public async Task ListBlobs()
+        {
+            List<BlobIdentifier> validBlobHashes = new List<BlobIdentifier>
+            {
+                _smallFileHash,
+                _anotherFileHash,
+                _deleteFileHash,
+                _oldBlobFileHash
+            };
+
+            IBlobStore blobStore = _server!.Services.GetService<IBlobStore>()!;
+            // fetch all objects so set cutoff to a time in the future
+            DateTime cutoff = DateTime.Now.AddMinutes(5);
+
+            BlobIdentifier[] oldObjects = await blobStore.ListOldObjects(TestNamespaceName, cutoff).ToArrayAsync();
+
+            Assert.AreEqual(4, oldObjects.Length);
+            Assert.IsTrue(validBlobHashes.Contains(oldObjects[0]));
+            Assert.IsTrue(validBlobHashes.Contains(oldObjects[1]));
+            Assert.IsTrue(validBlobHashes.Contains(oldObjects[2]));
+            Assert.IsTrue(validBlobHashes.Contains(oldObjects[3]));
+        }
+        
+        [TestMethod]
+        public async Task BlobExists()
+        {
+            {
+                HttpResponseMessage result = await _httpClient!.SendAsync(new HttpRequestMessage(HttpMethod.Head,
+                    requestUri: $"api/v1/s/{TestNamespaceName}/{_smallFileHash}"));
+                Assert.AreEqual(HttpStatusCode.OK, result.StatusCode);
+            }
+
+            {
+                ContentHash newContent = ContentHash.FromBlob(Encoding.ASCII.GetBytes("this content has never been submitted"));
+                HttpResponseMessage resultNew = await _httpClient!.SendAsync(new HttpRequestMessage(HttpMethod.Head,
+                    requestUri: $"api/v1/s/{TestNamespaceName}/{newContent}"));
+                Assert.AreEqual(HttpStatusCode.NotFound, resultNew.StatusCode);
+                string content = await resultNew.Content.ReadAsStringAsync();
+                ValidationProblemDetails result = JsonConvert.DeserializeObject<ValidationProblemDetails>(content)!;
+                Assert.AreEqual($"Blob {newContent} not found", result.Title);
+            }
+        }
+
+
+        [TestMethod]
+        public async Task BlobExistsBatch()
+        {
+            BlobIdentifier newContent = BlobIdentifier.FromBlob(Encoding.ASCII.GetBytes("this content has never been submitted"));
+
+            var ops = new
+            {
+                Operations = new object[]
+                {
+                    new
+                    {
+                        Op = "HEAD",
+                        Namespace = TestNamespaceName,
+                        Id = _smallFileHash,
+                    },
+                    new
+                    {
+                        Op = "HEAD",
+                        Namespace = TestNamespaceName,
+                        Id = newContent,
+                    }
+                }
+            };
+
+            HttpResponseMessage response = await _httpClient.PostAsJsonAsync("api/v1/s", ops);
+            response.EnsureSuccessStatusCode();
+            string content = await response.Content.ReadAsStringAsync();
+            object result = JsonConvert.DeserializeObject(content)!;
+            JArray? array = result as JArray;
+            Assert.IsNotNull(array);
+            Assert.AreEqual(2, array!.Count);
+            Assert.IsNull(array[0].Value<string>());
+            BlobIdentifier id = new BlobIdentifier(array[1].Value<string>()!);
+            Assert.AreEqual(newContent, id);
+        }
+
+
+        [TestMethod]
+        public async Task BatchOp()
+        {
+            var ops = new
+            {
+                Operations = new object[]
+                {
+                    new
+                    {
+                        Op = "GET",
+                        Namespace = TestNamespaceName,
+                        Id = _smallFileHash,
+                    },
+                    new
+                    {
+                        Op = "PUT",
+                        Namespace = TestNamespaceName,
+                        Id = _deleteFileHash,
+                        Content = Encoding.ASCII.GetBytes(DeletableFileContents)
+                    }
+                }
+            };
+
+            HttpResponseMessage response = await _httpClient.PostAsJsonAsync("api/v1/s", ops);
+            response.EnsureSuccessStatusCode();
+            string content = await response.Content.ReadAsStringAsync();
+            object result = JsonConvert.DeserializeObject(content)!;
+            Assert.IsTrue(result != null);
+            JArray? array = result as JArray;
+            Assert.IsNotNull(array);
+            Assert.AreEqual(2, array!.Count);
+            string base64Content = array[0].Value<string>()!;
+            string convertedContent = Encoding.ASCII.GetString(Convert.FromBase64String(base64Content));
+            Assert.AreEqual(SmallFileContents, convertedContent);
+            BlobIdentifier identifier = new BlobIdentifier(array[1].Value<string>()!);
+            Assert.AreEqual(_deleteFileHash, identifier);
+        }
+
+        [TestMethod]
+        public async Task BatchOpBadRequest()
+        {
+            var ops = new
+            {
+                Operations = new object[]
+                {
+                    new
+                    {
+                        Op = "GET",
+                        Namespace = TestNamespaceName,
+                        Id =  BlobIdentifier.FromBlob(Encoding.ASCII.GetBytes("foo"))
+                    },
+                    new
+                    {
+                        Op = "PUT",
+                        Namespace = TestNamespaceName,
+                        // no content
+                    },
+                    new
+                    {
+                        Op = "notAValidEnumValue"
+                    }
+                }
+            };
+
+            HttpResponseMessage response = await _httpClient.PostAsJsonAsync("api/v1/s", ops);
+            Assert.AreEqual(HttpStatusCode.BadRequest, response.StatusCode);
+
+            string content = await response.Content.ReadAsStringAsync();
+            SerializableError? result = JsonConvert.DeserializeObject<SerializableError>(content);
+            Assert.IsNotNull(result);
+
+            Assert.AreEqual(2, result.Keys.Count);
+        }
+
+        [TestMethod]
+        public async Task BatchOpBlobNotPresent()
+        {
+            var ops = new
+            {
+                Operations = new object[]
+                {
+                    new
+                    {
+                        Op = "HEAD",
+                        Namespace = TestNamespaceName,
+                        Id = "1DEE7232FE05FEFE623D2119626F103E05D3EE98", // random hash that does not exist
+                    },
+                }
+            };
+
+            HttpResponseMessage response = await _httpClient.PostAsJsonAsync("api/v1/s", ops);
+            Assert.AreEqual(HttpStatusCode.OK, response.StatusCode);
+
+            string content = await response.Content.ReadAsStringAsync();
+            string[]? result = JsonConvert.DeserializeObject<string[]>(content);
+            Assert.IsNotNull(result);
+
+            Assert.AreEqual(1, result.Length);
+            Assert.AreEqual("1DEE7232FE05FEFE623D2119626F103E05D3EE98", result[0]);
+        }
+
+        [TestMethod]
+        public async Task MultipleBlobChecks()
+        {
+            BlobIdentifier newContent = BlobIdentifier.FromBlob(Encoding.ASCII.GetBytes("this content has never been submitted"));
+            HttpResponseMessage response = await _httpClient!.SendAsync(new HttpRequestMessage(HttpMethod.Post,
+                requestUri: $"api/v1/s/{TestNamespaceName}/exists?id={_smallFileHash}&id={newContent}"));
+            Assert.AreEqual(HttpStatusCode.OK, response.StatusCode);
+            Assert.AreEqual("application/json", response.Content.Headers.ContentType!.MediaType);
+
+            string s = await response.Content.ReadAsStringAsync();
+            HeadMultipleResponse? result = JsonConvert.DeserializeObject<HeadMultipleResponse>(s);
+
+            Assert.IsNotNull(result);
+            Assert.AreEqual(1, result.Needs.Length);
+            Assert.AreEqual(newContent, result.Needs[0]);
+        }
+
+        [TestMethod]
+        public async Task MultipleBlobChecksBodyCB()
+        {
+            BlobIdentifier newContent = BlobIdentifier.FromBlob(Encoding.ASCII.GetBytes("this content has never been submitted"));
+            HttpRequestMessage request = new HttpRequestMessage(HttpMethod.Post, requestUri: $"api/v1/s/{TestNamespaceName}/exist");
+            CompactBinaryWriter writer = new CompactBinaryWriter();
+            writer.AsArray(new BlobIdentifier[] { _smallFileHash, newContent}, CompactBinaryFieldType.Hash);
+            byte[] buf = writer.Save();
+            request.Content = new ByteArrayContent(buf);
+            request.Content.Headers.ContentType = new MediaTypeHeaderValue(CustomMediaTypeNames.UnrealCompactBinary);
+
+            HttpResponseMessage response = await _httpClient!.SendAsync(request);
+            Assert.AreEqual(HttpStatusCode.OK, response.StatusCode);
+            Assert.AreEqual("application/json", response.Content.Headers.ContentType!.MediaType);
+
+            string s = await response.Content.ReadAsStringAsync();
+            HeadMultipleResponse? result = JsonConvert.DeserializeObject<HeadMultipleResponse>(s);
+
+            Assert.IsNotNull(result);
+            Assert.AreEqual(1, result.Needs.Length);
+            Assert.AreEqual(newContent, result.Needs[0]);
+        }
+
+        [TestMethod]
+        public async Task MultipleBlobChecksBodyJson()
+        {
+            BlobIdentifier newContent = BlobIdentifier.FromBlob(Encoding.ASCII.GetBytes("this content has never been submitted"));
+            HttpRequestMessage request = new HttpRequestMessage(HttpMethod.Post, requestUri: $"api/v1/s/{TestNamespaceName}/exist");
+            string jsonBody = JsonConvert.SerializeObject(new BlobIdentifier[] { _smallFileHash, newContent });
+            request.Content = new StringContent(jsonBody, Encoding.UTF8, MediaTypeNames.Application.Json);
+
+            HttpResponseMessage response = await _httpClient!.SendAsync(request);
+            Assert.AreEqual(HttpStatusCode.OK, response.StatusCode);
+            Assert.AreEqual("application/json", response.Content.Headers.ContentType!.MediaType);
+
+            string s = await response.Content.ReadAsStringAsync();
+            HeadMultipleResponse? result = JsonConvert.DeserializeObject<HeadMultipleResponse>(s);
+
+            Assert.IsNotNull(result);
+            Assert.AreEqual(1, result.Needs.Length);
+            Assert.AreEqual(newContent, result.Needs[0]);
+        }
+
+        [TestMethod]
+        public async Task MultipleBlobCompactBinaryResponse()
+        {
+            BlobIdentifier newContent = BlobIdentifier.FromBlob(Encoding.ASCII.GetBytes("this content has never been submitted"));
+            HttpRequestMessage request = new HttpRequestMessage(HttpMethod.Post, requestUri: $"api/v1/s/{TestNamespaceName}/exists?id={_smallFileHash}&id={newContent}");
+            request.Headers.Accept.Add(new MediaTypeWithQualityHeaderValue(CustomMediaTypeNames.UnrealCompactBinary));
+            HttpResponseMessage response = await _httpClient!.SendAsync(request);
+            Assert.AreEqual(HttpStatusCode.OK, response.StatusCode);
+            Assert.AreEqual(CustomMediaTypeNames.UnrealCompactBinary, response.Content.Headers.ContentType!.MediaType!);
+            byte[] data = await response.Content.ReadAsByteArrayAsync();
+            ReadOnlyMemory<byte> memory = new ReadOnlyMemory<byte>(data);
+            CompactBinaryObject cb = CompactBinaryObject.Load(ref memory);
+            
+            Assert.AreEqual(1, cb.GetFields().Count());
+            CompactBinaryField? needs = cb["needs"];
+            Assert.IsNotNull(needs);
+            BlobIdentifier?[] neededBlobs = needs!.AsArray().Select(field => field.AsHash()).ToArray();
+
+            Assert.AreEqual(1, neededBlobs.Length);
+            Assert.AreEqual(newContent, neededBlobs[0]);
+        }
+
+        [TestMethod]
+        public async Task FullFlow()
+        {
+            byte[] payload = Encoding.ASCII.GetBytes("Foo bar");
+            using ByteArrayContent requestContent = new ByteArrayContent(payload);
+            requestContent.Headers.ContentType = new MediaTypeHeaderValue("application/octet-stream");
+            requestContent.Headers.ContentLength = payload.Length;
+            BlobIdentifier contentHash = BlobIdentifier.FromBlob(payload);
+            HttpResponseMessage putResponse =
+                await _httpClient!.PutAsync(requestUri: $"api/v1/s/{TestNamespaceName}/{contentHash}", requestContent);
+            putResponse.EnsureSuccessStatusCode();
+            InsertResponse response = await putResponse.Content.ReadAsAsync<InsertResponse>();
+            Assert.AreEqual(contentHash, response.Identifier);
+
+            HttpResponseMessage getResponse = await _httpClient.GetAsync($"api/v1/s/{TestNamespaceName}/{contentHash}");
+            getResponse.EnsureSuccessStatusCode();
+            CollectionAssert.AreEqual(payload, actual: await getResponse.Content.ReadAsByteArrayAsync());
+
+            HttpResponseMessage deleteResponse = await _httpClient.DeleteAsync($"api/v1/s/{TestNamespaceName}/{contentHash}");
+            deleteResponse.EnsureSuccessStatusCode();
+            Assert.AreEqual(HttpStatusCode.NoContent, deleteResponse.StatusCode);
+        }
+
+        /// <summary>
+        /// write a large file (bigger then that C# can have in memory)
+        /// </summary>
+        
+        [TestMethod]
+        public async Task PutGetLargePayload()
+        {
+            // we submit a blob so large that it can not fit using the memory blob store
+            IBlobStore? blobStore = _server?.Services.GetService<IBlobStore>();
+            Assert.IsFalse(blobStore is MemoryBlobStore);
+
+            if (blobStore is AzureBlobStore)
+                Assert.Inconclusive("Azure blob store gets internal server errors when receiving large blobs");
+            FileInfo fi = new FileInfo(Path.GetTempFileName());
+
+            FileInfo tempOutputFile = new FileInfo(Path.GetTempFileName());
+
+            try
+            {
+                {
+                    await using FileStream fs = fi.OpenWrite();
+
+                    byte[] block = new byte[1024 * 1024];
+                    Array.Fill(block, (byte)'a');
+                    // we want a file larger then 2GB, each block is 1 MB
+                    int countOfBlocks = 2100;
+                    for (int i = 0; i < countOfBlocks; i++)
+                    {
+                        await fs.WriteAsync(block, 0, block.Length);
+                    }
+                }
+
+                BlobIdentifier blobIdentifier;
+                {
+                    await using FileStream fs = fi.OpenRead();
+                    blobIdentifier = await BlobIdentifier.FromStream(fs);
+                }
+
+                {
+                    await using FileStream fs = fi.OpenRead();
+                    StreamContent content = new StreamContent(fs);
+                    content.Headers.ContentType = new MediaTypeHeaderValue(MediaTypeNames.Application.Octet);
+                    HttpResponseMessage result = await _httpClient!.PutAsync($"api/v1/blobs/{TestNamespaceName}/{blobIdentifier}", content);
+                    result.EnsureSuccessStatusCode();
+
+                    InsertResponse response = await result.Content.ReadAsAsync<InsertResponse>();
+                    Assert.AreEqual(blobIdentifier, response.Identifier);
+                }
+                
+                {
+                    // verify we can fetch the blob again
+                    HttpResponseMessage result = await _httpClient!.GetAsync($"api/v1/blobs/{TestNamespaceName}/{blobIdentifier}", HttpCompletionOption.ResponseHeadersRead);
+                    result.EnsureSuccessStatusCode();
+                    Stream s = await result.Content.ReadAsStreamAsync();
+
+                    {
+                        // stream this to disk so we have something to look at in case there is an error
+                        await using FileStream fs = tempOutputFile.OpenWrite();
+                        await s.CopyToAsync(fs);
+                        fs.Close();
+                        s.Close();
+
+                        s = tempOutputFile.OpenRead();
+                    }
+
+                    BlobIdentifier downloadedBlobIdentifier = await BlobIdentifier.FromStream(s);
+                    Assert.AreEqual(blobIdentifier, downloadedBlobIdentifier);
+                    s.Close();
+                }
+
+            }
+            finally
+            {
+                if (fi.Exists)
+                    fi.Delete();
+
+                if (tempOutputFile.Exists)
+                    tempOutputFile.Delete();
+            }
+        }
+    }
+
+
+    public class InsertResponse
+    {
+        public BlobIdentifier? Identifier;
+    }
+}
diff --git a/Engine/Source/Programs/Horde/HordeStorage/Horde.Storage/Tests/UnitTests/BlobCleanupServiceTest.cs b/Engine/Source/Programs/Horde/HordeStorage/Horde.Storage/Tests/UnitTests/BlobCleanupServiceTest.cs
new file mode 100644
--- /dev/null
+++ b/Engine/Source/Programs/Horde/HordeStorage/Horde.Storage/Tests/UnitTests/BlobCleanupServiceTest.cs
@@ -0,0 +1,35 @@
+﻿// Copyright Epic Games, Inc. All Rights Reserved.
+
+using System.Collections.Generic;
+using System.Threading;
+using System.Threading.Tasks;
+using Dasync.Collections;
+using Horde.Storage.Implementation;
+using Microsoft.VisualStudio.TestTools.UnitTesting;
+using Moq;
+
+namespace Horde.Storage.UnitTests
+{
+    [TestClass]
+    public class BlobCleanupServiceTest
+    {
+        [TestMethod]
+        public async Task CleanupOnPoll()
+        {
+            GCSettings gcSettings = new GCSettings();
+            TestOptionsMonitor<GCSettings> gcSettingsMon = new TestOptionsMonitor<GCSettings>(gcSettings);
+            BlobCleanupService blobCleanupService = new BlobCleanupService(gcSettingsMon);
+
+            Mock<IBlobCleanup> store1 = new Mock<IBlobCleanup>();
+            Mock<IBlobCleanup> store2 = new Mock<IBlobCleanup>();
+            
+            blobCleanupService.RegisterCleanup(store1.Object);
+            blobCleanupService.RegisterCleanup(store2.Object);
+
+            await blobCleanupService.OnPoll(blobCleanupService.State, new CancellationTokenSource().Token);
+            
+            store1.Verify(m => m.Cleanup(It.IsAny<CancellationToken>()), Times.Once);
+            store2.Verify(m => m.Cleanup(It.IsAny<CancellationToken>()), Times.Once);
+        }
+    }
+}
diff --git a/Engine/Source/Programs/Horde/HordeStorage/Horde.Storage/Tests/UnitTests/CompactBinaryTests.cs b/Engine/Source/Programs/Horde/HordeStorage/Horde.Storage/Tests/UnitTests/CompactBinaryTests.cs
new file mode 100644
--- /dev/null
+++ b/Engine/Source/Programs/Horde/HordeStorage/Horde.Storage/Tests/UnitTests/CompactBinaryTests.cs
@@ -0,0 +1,120 @@
+﻿using System;
+using System.Collections.Generic;
+using System.IO;
+using System.Linq;
+using Jupiter.Implementation;
+using Microsoft.VisualStudio.TestTools.UnitTesting;
+
+namespace EuropaUnit
+{
+    [TestClass]
+    public class CompactBinaryTests
+    {
+
+        [TestMethod]
+        public void BuildObject()
+        {
+            byte[] bytes = File.ReadAllBytes("CompactBinaryObjects/build");
+
+            ReadOnlyMemory<byte> memory = new ReadOnlyMemory<byte>(bytes);
+            CompactBinaryObject o = CompactBinaryObject.Load(ref memory);
+
+            Assert.AreEqual("BuildAction", o.Name);
+            List<CompactBinaryField> buildActionFields = o.GetFields().ToList();
+            Assert.AreEqual(3, buildActionFields.Count);
+            Assert.AreEqual("Function", buildActionFields[0].Name);
+            Assert.AreEqual("Constants", buildActionFields[1].Name);
+            Assert.AreEqual("Inputs", buildActionFields[2].Name);
+
+            List<CompactBinaryField>  constantsFields = buildActionFields[1].GetFields().ToList();
+            Assert.AreEqual(3, constantsFields.Count);
+            Assert.AreEqual("TextureBuildSettings", constantsFields[0].Name);
+            Assert.AreEqual("TextureOutputSettings", constantsFields[1].Name);
+            Assert.AreEqual("TextureSource", constantsFields[2].Name);
+
+            List<CompactBinaryField>  inputsFields = buildActionFields[2].GetFields().ToList();
+            Assert.AreEqual(1, inputsFields.Count);
+            Assert.AreEqual("7587B323422942733DDD048A91709FDE", inputsFields[0].Name);
+            Assert.IsTrue(inputsFields[0].IsBinaryAttachment());
+            Assert.IsTrue(inputsFields[0].IsAttachment());
+            Assert.IsFalse(inputsFields[0].IsCompactBinaryAttachment());
+            Assert.AreEqual(new BlobIdentifier("f855382171a0b1e5a1c653aa6c5121a05cbf4ba0"), inputsFields[0].AsHash());
+        }
+
+        
+        [TestMethod]
+        public void ReferenceOutput()
+        {
+            byte[] bytes = File.ReadAllBytes("CompactBinaryObjects/ReferenceOutput");
+
+            ReadOnlyMemory<byte> memory = new ReadOnlyMemory<byte>(bytes);
+            CompactBinaryObject o = CompactBinaryObject.Load(ref memory);
+
+            Assert.AreEqual("BuildOutput", o.Name);
+            List<CompactBinaryField> buildActionFields = o.GetFields().ToList();
+            Assert.AreEqual(1, buildActionFields.Count);
+            CompactBinaryField payloads = buildActionFields[0];
+            List<CompactBinaryField> payloadFields = payloads.GetFields().ToList();
+            Assert.AreEqual(3, payloadFields.Count);
+
+            Assert.AreEqual(new BlobIdentifier("5d8a6dc277c968f0d027c98f879c955c1905c293"), payloadFields[0]["RawHash"]!.AsHash());
+            Assert.AreEqual(new BlobIdentifier("313f0d0d334100d83aeb1ee2c42794fd087cb0ae"), payloadFields[1]["RawHash"]!.AsHash());
+            Assert.AreEqual(new BlobIdentifier("c7a03f83c08cdca882110ecf2b5654ee3b09b11e"), payloadFields[2]["RawHash"]!.AsHash());
+        }
+
+
+        [TestMethod]
+        public void WriteArray()
+        {
+            var hash1 = new BlobIdentifier("5d8a6dc277c968f0d027c98f879c955c1905c293");
+            var hash2 = new BlobIdentifier("313f0d0d334100d83aeb1ee2c42794fd087cb0ae");
+
+            CompactBinaryWriter writer = new CompactBinaryWriter();
+            writer.BeginObject();
+            writer.AddUniformArray(new BlobIdentifier[] { hash1, hash2}, CompactBinaryFieldType.Hash, "needs");
+            writer.EndObject();
+
+            byte[] objectData = writer.Save();
+            ReadOnlyMemory<byte> memory = new ReadOnlyMemory<byte>(objectData);
+            CompactBinaryObject o = CompactBinaryObject.Load(ref memory);
+
+            // the top object has no name
+            Assert.AreEqual("", o.Name);
+            List<CompactBinaryField> fields = o.GetFields().ToList();
+            Assert.AreEqual(1, fields.Count);
+            CompactBinaryField? needs = o["needs"];
+            List<CompactBinaryField> blobList = needs!.AsArray().ToList();
+            BlobIdentifier?[] blobs = blobList.Select(field => field!.AsHash()).ToArray();
+            CollectionAssert.AreEqual(new BlobIdentifier[] {hash1, hash2}, blobs);
+        }
+
+
+        [TestMethod]
+        public void WriteObject()
+        {
+            var hash1 = new BlobIdentifier("5d8a6dc277c968f0d027c98f879c955c1905c293");
+
+            CompactBinaryWriter writer = new CompactBinaryWriter();
+            writer.BeginObject();
+            writer.AddString("test", "string");
+            writer.AddBinaryAttachment(hash1, "hash");
+            writer.EndObject();
+
+            byte[] objectData = writer.Save();
+            ReadOnlyMemory<byte> memory = new ReadOnlyMemory<byte>(objectData);
+            CompactBinaryObject o = CompactBinaryObject.Load(ref memory);
+
+            // the object has no name and 2 fields
+            Assert.AreEqual("", o.Name);
+            List<CompactBinaryField> fields = o.GetFields().ToList();
+            Assert.AreEqual(2, fields.Count);
+
+            CompactBinaryField? stringField = o["string"];
+            Assert.AreEqual("test", stringField!.AsString());
+
+            CompactBinaryField? hashField = o["hash"];
+            Assert.AreEqual(hash1, hashField!.AsAttachment());
+        }
+    }
+
+}
diff --git a/Engine/Source/Programs/Horde/HordeStorage/Horde.Storage/Tests/UnitTests/EuropaUnit.csproj b/Engine/Source/Programs/Horde/HordeStorage/Horde.Storage/Tests/UnitTests/EuropaUnit.csproj
new file mode 100644
--- /dev/null
+++ b/Engine/Source/Programs/Horde/HordeStorage/Horde.Storage/Tests/UnitTests/EuropaUnit.csproj
@@ -0,0 +1,38 @@
+﻿<Project Sdk="Microsoft.NET.Sdk">
+
+  <PropertyGroup>
+    <TargetFramework>net5.0</TargetFramework>
+
+    <IsPackable>false</IsPackable>
+
+    <AssemblyName>Horde.Storage.Tests.Unit</AssemblyName>
+    <Nullable>enable</Nullable>
+
+    <TreatWarningsAsErrors>true</TreatWarningsAsErrors>
+    <WarningsAsErrors />
+    <LangVersion>10</LangVersion>
+  </PropertyGroup>
+
+  <ItemGroup>
+    <PackageReference Include="Microsoft.AspNetCore.TestHost" Version="5.0.2" />
+    <PackageReference Include="Microsoft.NET.Test.Sdk" Version="16.8.3" />
+    <PackageReference Include="Moq" Version="4.16.0" />
+    <PackageReference Include="MSTest.TestAdapter" Version="2.1.2" />
+    <PackageReference Include="MSTest.TestFramework" Version="2.1.2" />
+    <PackageReference Include="SerilogAnalyzer" Version="0.15.0" />
+  </ItemGroup>
+
+  <ItemGroup>
+    <ProjectReference Include="..\..\Horde.Storage.csproj" />
+  </ItemGroup>
+
+  <ItemGroup>
+    <None Update="CompactBinaryObjects\build">
+      <CopyToOutputDirectory>PreserveNewest</CopyToOutputDirectory>
+    </None>
+    <None Update="CompactBinaryObjects\ReferenceOutput">
+      <CopyToOutputDirectory>PreserveNewest</CopyToOutputDirectory>
+    </None>
+  </ItemGroup>
+
+</Project>
diff --git a/Engine/Source/Programs/Horde/HordeStorage/Horde.Storage/Tests/UnitTests/HierarchicalBlobStorageTest.cs b/Engine/Source/Programs/Horde/HordeStorage/Horde.Storage/Tests/UnitTests/HierarchicalBlobStorageTest.cs
new file mode 100644
--- /dev/null
+++ b/Engine/Source/Programs/Horde/HordeStorage/Horde.Storage/Tests/UnitTests/HierarchicalBlobStorageTest.cs
@@ -0,0 +1,168 @@
+﻿// Copyright Epic Games, Inc. All Rights Reserved.
+
+using System;
+using System.Collections.Generic;
+using System.IO;
+using System.Linq;
+using System.Text;
+using System.Threading.Tasks;
+using Dasync.Collections;
+using Horde.Storage.Implementation;
+using Jupiter.Implementation;
+using Microsoft.VisualStudio.TestTools.UnitTesting;
+
+namespace Horde.Storage.UnitTests
+{
+    [TestClass]
+    public class HierarchicalBlobStorageTest
+    {
+        private readonly NamespaceId Ns = new NamespaceId("my-namespace");
+        private readonly NamespaceId NsnonExistingNs = new NamespaceId("non-existing-ns");
+        private readonly NamespaceId NsOnlyFirst = new NamespaceId("ns-only-in-first");
+        private readonly NamespaceId NsOnlySecond = new NamespaceId("ns-only-in-second");
+        private readonly MemoryBlobStore _first = new MemoryBlobStore(throwOnOverwrite: true);
+        private readonly MemoryBlobStore _second = new MemoryBlobStore(throwOnOverwrite: true);
+        private readonly MemoryBlobStore _third = new MemoryBlobStore(throwOnOverwrite: true);
+        private readonly IBlobStore _chained;
+        
+        private readonly BlobIdentifier _onlyFirstId = new BlobIdentifier(new string('1', 40));
+        private readonly BlobIdentifier _onlySecondId = new BlobIdentifier(new string('2', 40));
+        private readonly BlobIdentifier _onlyThirdId = new BlobIdentifier(new string('3', 40));
+        private readonly BlobIdentifier _allId = new BlobIdentifier(new string('4', 40));
+        private readonly BlobIdentifier _onlyFirstUniqueNsId = new BlobIdentifier(new string('5', 40));
+        private readonly BlobIdentifier _onlySecondUniqueNsId = new BlobIdentifier(new string('6', 40));
+        private readonly BlobIdentifier _nonExisting = new BlobIdentifier(new string('0', 40));
+
+        public HierarchicalBlobStorageTest()
+        {
+            _chained = new HierarchicalBlobStore(new []{_first, _second, _third});
+        }
+
+        [TestInitialize]
+        public async Task Setup()
+        {
+            await _first.PutObject(Ns, Encoding.ASCII.GetBytes("onlyFirstContent"), _onlyFirstId);
+            await _second.PutObject(Ns, Encoding.ASCII.GetBytes("onlySecondContent"), _onlySecondId);
+            await _third.PutObject(Ns, Encoding.ASCII.GetBytes("onlyThirdContent"), _onlyThirdId);
+            
+            await _first.PutObject(Ns, Encoding.ASCII.GetBytes("allContent"), _allId);
+            await _second.PutObject(Ns, Encoding.ASCII.GetBytes("allContent"), _allId);
+            await _third.PutObject(Ns, Encoding.ASCII.GetBytes("allContent"), _allId);
+            await _second.PutObject(NsOnlyFirst, Encoding.ASCII.GetBytes("onlyFirstUniqueNs"), _onlyFirstUniqueNsId);
+            await _second.PutObject(NsOnlySecond, Encoding.ASCII.GetBytes("onlySecondUniqueNs"), _onlySecondUniqueNsId);
+        }
+        
+        [TestMethod]
+        public async Task PutObject()
+        {
+            BlobIdentifier new1 = new BlobIdentifier("1000000000000000000000000000000000000000");
+            Assert.IsFalse(await _chained.Exists(Ns, new1));
+            await _chained.PutObject(Ns, Encoding.ASCII.GetBytes("new1"), new1);
+            Assert.IsTrue(await _chained.Exists(Ns, new1));
+            Assert.IsTrue(await _first.Exists(Ns, new1));
+            Assert.IsTrue(await _second.Exists(Ns, new1));
+        }
+        
+        [TestMethod]
+        public async Task GetObject()
+        {
+            Assert.AreEqual("onlyFirstContent", BlobToString(await _chained.GetObject(Ns, _onlyFirstId)));
+            Assert.AreEqual("onlySecondContent", BlobToString(await _chained.GetObject(Ns, _onlySecondId)));
+            Assert.AreEqual("onlyFirstUniqueNs", BlobToString(await _chained.GetObject(NsOnlyFirst, _onlyFirstUniqueNsId)));
+            Assert.AreEqual("onlySecondUniqueNs", BlobToString(await _chained.GetObject(NsOnlySecond, _onlySecondUniqueNsId)));
+            Assert.AreEqual("allContent", BlobToString(await _chained.GetObject(Ns, _allId)));
+            await Assert.ThrowsExceptionAsync<BlobNotFoundException>(() => _chained.GetObject(Ns, _nonExisting));
+            await Assert.ThrowsExceptionAsync<NamespaceNotFoundException>(() => _chained.GetObject(new NamespaceId("non-existing-ns"), _nonExisting));
+
+            // verify that the objects have propagated
+            Assert.AreEqual(3, _first.GetIdentifiers(Ns).Count());
+            Assert.AreEqual(2, _second.GetIdentifiers(Ns).Count());
+            Assert.AreEqual(2, _third.GetIdentifiers(Ns).Count());
+        }
+        
+        [TestMethod]
+        public async Task PopulateHierarchy()
+        {
+            Assert.IsFalse(await _first.Exists(Ns, _onlyThirdId));
+            Assert.IsFalse(await _second.Exists(Ns, _onlyThirdId));
+            Assert.IsTrue(await _third.Exists(Ns, _onlyThirdId));
+            
+            // Should populate 'first' and 'second' as they are higher up in the hierarchy
+            Assert.AreEqual("onlyThirdContent", BlobToString(await _chained.GetObject(Ns, _onlyThirdId)));
+            
+            Assert.AreEqual("onlyThirdContent", BlobToString(await _first.GetObject(Ns, _onlyThirdId)));
+            Assert.AreEqual("onlyThirdContent", BlobToString(await _second.GetObject(Ns, _onlyThirdId)));
+            Assert.AreEqual("onlyThirdContent", BlobToString(await _third.GetObject(Ns, _onlyThirdId)));
+        }
+        
+        [TestMethod]
+        public async Task Exists()
+        {
+            Assert.IsTrue(await _chained.Exists(Ns, _onlyFirstId));
+            Assert.IsTrue(await _chained.Exists(Ns, _onlySecondId));
+            Assert.IsTrue(await _chained.Exists(Ns, _allId));
+            Assert.IsFalse(await _chained.Exists(Ns, _nonExisting));
+        }
+        
+        [TestMethod]
+        public async Task DeleteObject()
+        {
+            await Assert.ThrowsExceptionAsync<NamespaceNotFoundException>(() => _chained.DeleteObject(NsnonExistingNs, _nonExisting));
+            await Assert.ThrowsExceptionAsync<BlobNotFoundException>(() => _chained.DeleteObject(NsOnlyFirst, _nonExisting));
+            await Assert.ThrowsExceptionAsync<BlobNotFoundException>(() => _chained.DeleteObject(NsOnlySecond, _nonExisting));
+            
+            Assert.IsTrue(await _first.Exists(Ns, _onlyFirstId));
+            await _chained.DeleteObject(Ns, _onlyFirstId);
+            Assert.IsFalse(await _first.Exists(Ns, _onlyFirstId));
+            
+            Assert.IsTrue(await _second.Exists(NsOnlySecond, _onlySecondUniqueNsId));
+            await _chained.DeleteObject(NsOnlySecond, _onlySecondUniqueNsId);
+            Assert.IsFalse(await _second.Exists(NsOnlySecond, _onlySecondUniqueNsId));
+            
+            Assert.IsTrue(await _first.Exists(Ns, _allId));
+            Assert.IsTrue(await _second.Exists(Ns, _allId));
+            await _chained.DeleteObject(Ns, _allId);
+            Assert.IsFalse(await _first.Exists(Ns, _allId));
+            Assert.IsFalse(await _second.Exists(Ns, _allId));
+        }
+        
+        [TestMethod]
+        public async Task DeleteNamespace()
+        {
+            await Assert.ThrowsExceptionAsync<NamespaceNotFoundException>(() => _chained.DeleteNamespace(NsnonExistingNs));
+
+            Assert.IsTrue(await _first.Exists(Ns, _allId));
+            Assert.IsTrue(await _second.Exists(Ns, _allId));
+            await _chained.DeleteNamespace(Ns);
+            await Assert.ThrowsExceptionAsync<NamespaceNotFoundException>(() => _first.DeleteNamespace(Ns));
+            await Assert.ThrowsExceptionAsync<NamespaceNotFoundException>(() => _second.DeleteNamespace(Ns));
+        }
+        
+        [TestMethod]
+        public async Task ListOldObjects()
+        {
+            List<BlobIdentifier> blobIds = await _chained.ListOldObjects(Ns, DateTime.Now.AddYears(10)).GetAsyncEnumerator().ToListAsync();
+            Assert.IsTrue(blobIds.Exists(x => x.Equals(_onlyFirstId)));
+            Assert.IsTrue(blobIds.Exists(x => x.Equals(_onlySecondId)));
+            Assert.IsTrue(blobIds.Exists(x => x.Equals(_onlyThirdId)));
+            Assert.IsTrue(blobIds.Exists(x => x.Equals(_allId)));
+            Assert.AreEqual(4, blobIds.Count);
+            
+            blobIds = await _chained.ListOldObjects(NsOnlyFirst, DateTime.Now.AddYears(10)).GetAsyncEnumerator().ToListAsync();
+            Assert.IsTrue(blobIds.Exists(x => x.Equals(_onlyFirstUniqueNsId)));
+            Assert.AreEqual(1, blobIds.Count);
+            
+            blobIds = await _chained.ListOldObjects(NsOnlySecond, DateTime.Now.AddYears(10)).GetAsyncEnumerator().ToListAsync();
+            Assert.IsTrue(blobIds.Exists(x => x.Equals(_onlySecondUniqueNsId)));
+            Assert.AreEqual(1, blobIds.Count);
+            
+            await Assert.ThrowsExceptionAsync<NamespaceNotFoundException>(() => _chained.ListOldObjects(NsnonExistingNs, DateTime.Now.AddYears(10)).GetAsyncEnumerator().ToListAsync());
+        }
+
+        private string BlobToString(BlobContents contents)
+        {
+            using StreamReader reader = new StreamReader(contents.Stream);
+            return reader.ReadToEnd();
+        }
+    }
+}
diff --git a/Engine/Source/Programs/Horde/HordeStorage/Horde.Storage/Tests/UnitTests/Horde.Storage.Tests.Unit.csproj b/Engine/Source/Programs/Horde/HordeStorage/Horde.Storage/Tests/UnitTests/Horde.Storage.Tests.Unit.csproj
new file mode 100644
--- /dev/null
+++ b/Engine/Source/Programs/Horde/HordeStorage/Horde.Storage/Tests/UnitTests/Horde.Storage.Tests.Unit.csproj
@@ -0,0 +1,38 @@
+﻿<Project Sdk="Microsoft.NET.Sdk">
+
+  <PropertyGroup>
+    <TargetFramework>net6.0</TargetFramework>
+
+    <IsPackable>false</IsPackable>
+
+    <AssemblyName>Horde.Storage.Tests.Unit</AssemblyName>
+    <Nullable>enable</Nullable>
+
+    <TreatWarningsAsErrors>true</TreatWarningsAsErrors>
+    <WarningsAsErrors />
+    <LangVersion>10</LangVersion>
+  </PropertyGroup>
+
+  <ItemGroup>
+    <PackageReference Include="Microsoft.AspNetCore.TestHost" Version="6.0.0" />
+    <PackageReference Include="Microsoft.NET.Test.Sdk" Version="17.0.0" />
+    <PackageReference Include="Moq" Version="4.16.1" />
+    <PackageReference Include="MSTest.TestAdapter" Version="2.2.7" />
+    <PackageReference Include="MSTest.TestFramework" Version="2.2.7" />
+    <PackageReference Include="SerilogAnalyzer" Version="0.15.0" />
+  </ItemGroup>
+
+  <ItemGroup>
+    <ProjectReference Include="..\..\Horde.Storage.csproj" />
+  </ItemGroup>
+
+  <ItemGroup>
+    <None Update="CompactBinaryObjects\build">
+      <CopyToOutputDirectory>PreserveNewest</CopyToOutputDirectory>
+    </None>
+    <None Update="CompactBinaryObjects\ReferenceOutput">
+      <CopyToOutputDirectory>PreserveNewest</CopyToOutputDirectory>
+    </None>
+  </ItemGroup>
+
+</Project>
diff --git a/Engine/Source/Programs/Horde/HordeStorage/Horde.Storage/Tests/UnitTests/HordeStorageStartupTest.cs b/Engine/Source/Programs/Horde/HordeStorage/Horde.Storage/Tests/UnitTests/HordeStorageStartupTest.cs
new file mode 100644
--- /dev/null
+++ b/Engine/Source/Programs/Horde/HordeStorage/Horde.Storage/Tests/UnitTests/HordeStorageStartupTest.cs
@@ -0,0 +1,69 @@
+﻿// Copyright Epic Games, Inc. All Rights Reserved.
+
+using System.Collections.Generic;
+using System.IO;
+using System.Linq;
+using Horde.Storage.Implementation;
+using Microsoft.AspNetCore.TestHost;
+using Microsoft.AspNetCore.Hosting;
+using Microsoft.Extensions.Configuration;
+using Microsoft.Extensions.DependencyInjection;
+using Microsoft.VisualStudio.TestTools.UnitTesting;
+using Serilog;
+
+namespace Horde.Storage.UnitTests
+{
+    [TestClass]
+    public class HordeStorageStartupTest
+    {
+        private string localTestDir;
+        
+        public HordeStorageStartupTest()
+        {
+            localTestDir = Path.Combine(Path.GetTempPath(), "HordeStorageStartupTest", Path.GetRandomFileName());
+        }
+
+        [TestMethod]
+        public void StorageBackendImplConfiguration()
+        {
+            // No configuration set
+            IBlobStore blobStore = GetBlobStoreForConfig(new Dictionary<string, string>());
+            Assert.IsTrue(blobStore is MemoryCacheBlobStore);
+            
+            // A single blob store configuration should yield the store itself without a hierarchical wrapper store
+            blobStore = GetBlobStoreForConfig(new Dictionary<string, string> {{"Horde.Storage:StorageImplementations:0", "FileSystem"}});
+            Assert.IsTrue(blobStore is FileSystemStore);
+            
+            // Should not be case-sensitive
+            blobStore = GetBlobStoreForConfig(new Dictionary<string, string> {{"Horde.Storage:StorageImplementations:0", "FiLeSYsTEm"}});
+            Assert.IsTrue(blobStore is FileSystemStore);
+            
+            // Two or more blob stores returns a hierarchical store
+            blobStore = GetBlobStoreForConfig(new Dictionary<string, string>
+            {
+                {"Horde.Storage:StorageImplementations:0", "FileSystem"},
+                {"Horde.Storage:StorageImplementations:1", "Memory"},
+            });
+            Assert.IsTrue(blobStore is HierarchicalBlobStore);
+            Assert.IsTrue((blobStore as HierarchicalBlobStore)!.BlobStores.ToList()[0] is FileSystemStore);
+            Assert.IsTrue((blobStore as HierarchicalBlobStore)!.BlobStores.ToList()[1] is MemoryCacheBlobStore);
+        }
+
+        private IBlobStore GetBlobStoreForConfig(Dictionary<string, string> configDict)
+        {
+            IConfigurationRoot configuration = new ConfigurationBuilder()
+                .AddJsonFile("appsettings.Testing.json", true)
+                .AddInMemoryCollection(new Dictionary<string, string> {{"Filesystem:RootDir", localTestDir}})
+                .AddInMemoryCollection(configDict)
+                .Build();
+            TestServer server = new TestServer(new WebHostBuilder()
+                .UseConfiguration(configuration)
+                .UseEnvironment("Testing")
+                .UseSerilog(new LoggerConfiguration().ReadFrom.Configuration(configuration).CreateLogger())
+                .UseStartup<HordeStorageStartup>()
+            );
+
+            return server.Services.GetService<IBlobStore>()!;
+        }
+    }
+}
diff --git a/Engine/Source/Programs/Horde/HordeStorage/Horde.Storage/Tests/UnitTests/KestrelDispatcherTest.cs b/Engine/Source/Programs/Horde/HordeStorage/Horde.Storage/Tests/UnitTests/KestrelDispatcherTest.cs
new file mode 100644
--- /dev/null
+++ b/Engine/Source/Programs/Horde/HordeStorage/Horde.Storage/Tests/UnitTests/KestrelDispatcherTest.cs
@@ -0,0 +1,256 @@
+﻿// Copyright Epic Games, Inc. All Rights Reserved.
+
+using System;
+using System.Collections.Generic;
+using System.IO;
+using System.IO.Pipelines;
+using System.Text;
+using System.Threading.Tasks;
+using Horde.Storage.Implementation;
+using Horde.Storage.Implementation.Kestrel;
+using Jupiter.Implementation;
+using Microsoft.VisualStudio.TestTools.UnitTesting;
+using Moq;
+
+namespace Horde.Storage.UnitTests
+{
+    [TestClass]
+    public class KestrelDispatcherTest
+    {
+        private class HttpResponse
+        {
+            public readonly int StatusCode;
+            public readonly Dictionary<string, string> Headers = new ();
+            public readonly byte[] Body;
+
+            public HttpResponse(int statusCode, byte[] body)
+            {
+                StatusCode = statusCode;
+                Body = body;
+            }
+
+            public string GetBodyAsString()
+            {
+                return Encoding.ASCII.GetString(Body);
+            }
+        }
+
+        public Mock<IDDCRefService> GetDdcRefServiceMock()
+        {
+            Mock<IDDCRefService> mock = new () { DefaultValue = DefaultValue.Empty };
+
+            NamespaceId ns = new NamespaceId("test.ddc");
+            BucketId bucket = new BucketId("default");
+            ContentHash contentHash1 = new ("BBBBBBBBBBBBBBBBBBBBBBBBBBBBBBBBBBBBBBBB");
+            RefResponse refResponse = new RefResponse("bogusName", DateTime.Now, contentHash1, new BlobIdentifier[0], null);
+            byte[] data = new AsciiString("testing");
+            MemoryStream ms = new (data);
+            BlobContents blobContents = new BlobContents(ms, data.Length);
+            var returnGetSuccess = (refResponse, blobContents);
+            mock.Setup(s => s.Get(ns, bucket, new KeyId("somekey"), It.IsAny<string[]>())).ReturnsAsync(returnGetSuccess).Verifiable();
+
+            ContentHash contentHash2 = new ("AAAAAAAAAAAAAAAAAAAAAAAAAAAAAAAAAAAAAAAA");
+            RefRecord returnExistsSuccess = new (ns, bucket, new KeyId("BLARGH"), new BlobIdentifier[0], null, contentHash2, null);
+            mock.Setup(s => s.Exists(ns, bucket, new KeyId("somekey"))).ReturnsAsync(returnExistsSuccess).Verifiable();
+
+            return mock;
+        }
+        
+        [TestMethod]
+        public async Task Ping()
+        {
+            IDDCRefService ddcRefService = GetDdcRefServiceMock().Object;
+            HttpResponse res = await SendHttpRequest(ddcRefService,"GET", "/ping");
+            Assert.AreEqual(200, res.StatusCode);
+            Assert.AreEqual("pong", res.GetBodyAsString());
+            Assert.AreEqual("text/plain", res.Headers["Content-Type"]);
+        }
+        
+        [TestMethod]
+        public async Task HealthLive()
+        {
+            IDDCRefService ddcRefService = GetDdcRefServiceMock().Object;
+            HttpResponse res = await SendHttpRequest(ddcRefService,"GET", "/health/live");
+            Assert.AreEqual(200, res.StatusCode);
+            Assert.AreEqual("Healthy", res.GetBodyAsString());
+            Assert.AreEqual("text/plain", res.Headers["Content-Type"]);
+        }
+        
+        [TestMethod]
+        public async Task HealthReady()
+        {
+            IDDCRefService ddcRefService = GetDdcRefServiceMock().Object;
+            HttpResponse res = await SendHttpRequest(ddcRefService,"GET", "/health/ready");
+            Assert.AreEqual(200, res.StatusCode);
+            Assert.AreEqual("Healthy", res.GetBodyAsString());
+            Assert.AreEqual("text/plain", res.Headers["Content-Type"]);
+        }
+        
+        [TestMethod]
+        public async Task DdcGet()
+        {
+            IDDCRefService ddcRefService = GetDdcRefServiceMock().Object;
+            HttpResponse res = await SendHttpRequest(ddcRefService,"GET", "/api/v1/c/ddc/test.ddc/default/somekey.raw");
+            Assert.AreEqual(200, res.StatusCode);
+            Assert.AreEqual("application/octet-stream", res.Headers["Content-Type"]);
+            Assert.AreEqual("BBBBBBBBBBBBBBBBBBBBBBBBBBBBBBBBBBBBBBBB", res.Headers[KestrelDispatcher.HeaderIoHash.ToString()]);
+            Assert.AreEqual("testing", res.GetBodyAsString());
+        }
+        
+        [TestMethod]
+        public async Task DdcGetFail()
+        {
+            IDDCRefService ddcRefService = GetDdcRefServiceMock().Object;
+            HttpResponse res = await SendHttpRequest(ddcRefService,"GET", "/api/v1/c/ddc/test.ddc/default/foo/bar/baz");
+            Assert.AreEqual(400, res.StatusCode);
+            Assert.AreEqual("application/json", res.Headers["Content-Type"]);
+
+            string expectedBody = "Bad request path: /api/v1/c/ddc/test.ddc/default/foo/bar/baz";
+            Assert.AreEqual(expectedBody.Length, Convert.ToInt32(res.Headers["Content-Length"]));
+            Assert.AreEqual(expectedBody, res.GetBodyAsString());
+        }
+        
+        [TestMethod]
+        public async Task DdcHead()
+        {
+            IDDCRefService ddcRefService = GetDdcRefServiceMock().Object;
+            HttpResponse res = await SendHttpRequest(ddcRefService,"HEAD", "/api/v1/c/ddc/test.ddc/default/somekey");
+            Assert.AreEqual(204, res.StatusCode);
+            Assert.AreEqual("AAAAAAAAAAAAAAAAAAAAAAAAAAAAAAAAAAAAAAAA", res.Headers["X-Jupiter-IoHash"]);
+            Assert.AreEqual("", res.GetBodyAsString());
+        }
+        
+        [TestMethod]
+        public async Task DdcHeadFail()
+        {
+            IDDCRefService ddcRefService = GetDdcRefServiceMock().Object;
+            HttpResponse res = await SendHttpRequest(ddcRefService,"HEAD", "/api/v1/c/ddc/test.ddc/default/foo/bar/baz");
+            Assert.AreEqual(400, res.StatusCode);
+            Assert.IsTrue(res.GetBodyAsString().Contains("Bad request path"));
+        }
+        
+        [TestMethod]
+        public async Task DdcPut()
+        {
+            IDDCRefService ddcRefService = GetDdcRefServiceMock().Object;
+            HttpResponse res = await SendHttpRequest(ddcRefService,"PUT", "/api/v1/c/ddc/test.ddc/default/somekey", "mypayload");
+            Assert.AreEqual(200, res.StatusCode);
+            Assert.AreEqual("{\"TransactionId\": 1}", res.GetBodyAsString());
+        }
+        
+        [TestMethod]
+        public async Task DdcPutLarge()
+        {
+            IDDCRefService ddcRefService = GetDdcRefServiceMock().Object;
+            byte[] largeData = new byte[10 * 1024 * 1024];
+            Array.Fill(largeData, (byte) 'j');
+            HttpResponse res = await SendHttpRequest(ddcRefService,"PUT", "/api/v1/c/ddc/test.ddc/default/somekey", new AsciiString(largeData));
+            Assert.AreEqual(200, res.StatusCode);
+            Assert.AreEqual("{\"TransactionId\": 1}", res.GetBodyAsString());
+            
+            res = await SendHttpRequest(ddcRefService,"GET", "/api/v1/c/ddc/test.ddc/default/somekey.raw");
+            Assert.AreEqual(200, res.StatusCode);
+            Assert.AreEqual("application/octet-stream", res.Headers["Content-Type"]);
+            Assert.AreEqual("testing", res.GetBodyAsString());
+        }
+        
+        [TestMethod]
+        public async Task NotFound()
+        {
+            IDDCRefService ddcRefService = GetDdcRefServiceMock().Object;
+            HttpResponse res = await SendHttpRequest(ddcRefService, "GET", "/invalid-path");
+            Assert.AreEqual(404, res.StatusCode);
+            Assert.AreEqual("404 Not Found", res.GetBodyAsString());
+            Assert.AreEqual("text/plain", res.Headers["Content-Type"]);
+        }
+
+        private static async Task<HttpResponse> SendHttpRequest(IDDCRefService ddcRefService, string method, string path, AsciiString? requestBody = null)
+        {
+            Pipe requestPipe = new ();
+            Pipe responsePipe = new ();
+
+            KestrelDispatcher dispatcher = new (ddcRefService, requestPipe.Reader, responsePipe.Writer, null!);
+            
+            AsciiString crlf = "\r\n";
+            AsciiString eoh = "\r\n\r\n"; // End Of Headers
+            AsciiString requestBytes = method + " " + path + " HTTP/1.1" + crlf + "Host: localhost";
+
+            if (requestBody != null)
+            {
+                requestBytes += crlf + "Content-Length: " + Convert.ToString(requestBody.Value.Length);
+            }
+            
+            requestBytes += eoh;
+
+            await requestPipe.Writer.WriteAsync(requestBytes.Data);
+            
+            if (requestBody != null)
+            {
+                // Do two separate writes for splitting the buffers on the receiving (hopefully exercising more edge cases in pipeline reader)
+                int size = requestBody.Value.Data.Length / 2;
+                Memory<byte> firstSlice = requestBody.Value.Data.AsMemory().Slice(0, size);
+                Memory<byte> secondSlice = requestBody.Value.Data.AsMemory().Slice(size);
+                var t1 = requestPipe.Writer.WriteAsync(firstSlice);
+                var t2 = requestPipe.Writer.WriteAsync(secondSlice);
+            }
+            
+            await requestPipe.Writer.CompleteAsync();
+            
+            await dispatcher.ExecuteAsync(null!);
+
+            byte[] temp = new byte[10000];
+            MemoryStream ms = new (temp);
+            await responsePipe.Reader.AsStream().CopyToAsync(ms);
+
+            int bodyOffset = SearchBytes(temp, eoh);
+            
+            string headersStr = Encoding.ASCII.GetString(temp.AsSpan(0, bodyOffset));
+            string[] headers = headersStr.Split("\r\n");
+            string[] statusLineTokens = headers[0].Split(' ');
+            Assert.AreEqual("HTTP/1.1", statusLineTokens[0]);
+            int statusCode = Convert.ToInt32(statusLineTokens[1]);
+            
+            byte[] body = PruneZeroBytes(temp.AsSpan(bodyOffset + eoh.Length).ToArray());
+            HttpResponse response = new (statusCode, body);
+            
+            foreach (string str in headers[1..])
+            {
+                string[] headerTokens = str.Split(": ");
+                response.Headers[headerTokens[0]] = headerTokens[1];
+            }
+            
+            return response;
+        }
+        
+        private static int SearchBytes(byte[] haystack, byte[] needle)
+        {
+            int len = needle.Length;
+            int limit = haystack.Length - len;
+            for (int i = 0;  i <= limit;  i++ )
+            {
+                int k = 0;
+                for (; k < len; k++)
+                {
+                    if (needle[k] != haystack[i + k])
+                        break;
+                }
+                
+                if( k == len )
+                    return i;
+            }
+            return -1;
+        }
+        
+        private static byte[] PruneZeroBytes(byte[] data) {
+            if (data.Length == 0) return data;
+            int i = data.Length - 1;
+            while (data[i] == 0) {
+                i--;
+                if (i == 0) return Array.Empty<byte>();
+            }
+            byte[] copy = new byte[i + 1];
+            Array.Copy(data, copy, i + 1);
+            return copy;
+        }
+    }
+}
diff --git a/Engine/Source/Programs/Horde/HordeStorage/Horde.Storage/Tests/UnitTests/MemoryCacheBlobStoreTest.cs b/Engine/Source/Programs/Horde/HordeStorage/Horde.Storage/Tests/UnitTests/MemoryCacheBlobStoreTest.cs
new file mode 100644
--- /dev/null
+++ b/Engine/Source/Programs/Horde/HordeStorage/Horde.Storage/Tests/UnitTests/MemoryCacheBlobStoreTest.cs
@@ -0,0 +1,119 @@
+﻿// Copyright Epic Games, Inc. All Rights Reserved.
+
+using System;
+using System.IO;
+using System.Text;
+using System.Threading.Tasks;
+using Horde.Storage.Implementation;
+using Jupiter.Implementation;
+using Microsoft.Extensions.Options;
+using Microsoft.VisualStudio.TestTools.UnitTesting;
+
+namespace Horde.Storage.UnitTests
+{
+    public class TestOptionsMonitor<T> : IOptionsMonitor<T>
+        where T : class, new()
+    {
+        public TestOptionsMonitor(T currentValue)
+        {
+            CurrentValue = currentValue;
+        }
+
+        public T Get(string name)
+        {
+            return CurrentValue;
+        }
+
+        public IDisposable OnChange(Action<T, string> listener)
+        {
+            throw new NotImplementedException();
+        }
+
+        public T CurrentValue { get; }
+    }
+    
+    [TestClass]
+    public class MemoryCacheBlobStoreTest
+    {
+        private readonly NamespaceId Ns1 = new NamespaceId("my-namespace-1");
+        private readonly NamespaceId Ns2 = new NamespaceId("my-namespace-2");
+        private readonly BlobIdentifier _id1 = new BlobIdentifier(new string('1', 40));
+        private readonly BlobIdentifier _id2 = new BlobIdentifier(new string('2', 40));
+        private readonly BlobIdentifier _id3 = new BlobIdentifier(new string('3', 40));
+
+        private readonly MemoryCacheBlobStore _store;
+
+        public MemoryCacheBlobStoreTest()
+        {
+            MemoryCacheBlobSettings settings = new MemoryCacheBlobSettings();
+            _store = new MemoryCacheBlobStore(new TestOptionsMonitor<MemoryCacheBlobSettings>(settings));
+            _store.PutObject(Ns1, Encoding.ASCII.GetBytes("first content"), _id1);
+            _store.PutObject(Ns1, Encoding.ASCII.GetBytes("second content"), _id2);
+            _store.PutObject(Ns2, Encoding.ASCII.GetBytes("third content"), _id3);
+        }
+        
+        [TestMethod]
+        public async Task PutObject()
+        {
+            BlobIdentifier new1 = new BlobIdentifier(new string('x', 40));
+            byte[] data = Encoding.ASCII.GetBytes("some new content");
+            await _store.PutObject(Ns1, data, new1);
+            BlobContents contents = await _store.GetObject(Ns1, new1);
+            await using MemoryStream tempStream = new MemoryStream();
+            await contents.Stream.CopyToAsync(tempStream);
+            
+            Assert.AreEqual("some new content", Encoding.ASCII.GetString(tempStream.ToArray()));
+        }
+
+        [TestMethod]
+        public async Task GetObject()
+        {
+            BlobContents contents = await _store.GetObject(Ns1, _id1);
+            await using MemoryStream tempStream = new MemoryStream();
+            await contents.Stream.CopyToAsync(tempStream);
+            
+            Assert.AreEqual("first content", Encoding.ASCII.GetString(tempStream.ToArray()));
+            Assert.AreEqual(13, contents.Length);
+        }
+        
+        [TestMethod]
+        public async Task DeleteObject()
+        {
+            Assert.IsTrue(await _store.Exists(Ns1, _id1));
+            await _store.DeleteObject(Ns1, _id1);
+            Assert.IsFalse(await _store.Exists(Ns1, _id1));
+        }
+        
+        [TestMethod]
+        public async Task Exists()
+        {
+            Assert.IsTrue(await _store.Exists(Ns1, _id1));
+            Assert.IsFalse(await _store.Exists(Ns1, new BlobIdentifier(new string('o', 40))));
+        }
+        
+        [TestMethod]
+        [Ignore("Always returns 0")]
+        public void GetMemoryCacheSize()
+        {
+            Assert.AreEqual(2, _store.GetMemoryCacheSize());
+        }
+        
+        // Ignored as they are not implemented
+        [TestMethod]
+        [Ignore("Not implemented due to limitations with MS MemoryCache")]
+        public async Task DeleteNamespace()
+        {
+            Assert.IsTrue(await _store.Exists(Ns1, _id1));
+            Assert.IsTrue(await _store.Exists(Ns2, _id3));
+            await _store.DeleteNamespace(Ns2);
+            Assert.IsTrue(await _store.Exists(Ns1, _id1));
+            Assert.IsFalse(await _store.Exists(Ns2, _id3));
+        }
+
+        [TestMethod]
+        [Ignore("Not implemented due to limitations with MS MemoryCache")]
+        public void ListOldObjects()
+        {
+        }
+    }
+}
diff --git a/Engine/Source/Programs/Horde/HordeStorage/Horde.Storage/Tests/UnitTests/ReplicatorTests.cs b/Engine/Source/Programs/Horde/HordeStorage/Horde.Storage/Tests/UnitTests/ReplicatorTests.cs
new file mode 100644
--- /dev/null
+++ b/Engine/Source/Programs/Horde/HordeStorage/Horde.Storage/Tests/UnitTests/ReplicatorTests.cs
@@ -0,0 +1,336 @@
+﻿// Copyright Epic Games, Inc. All Rights Reserved.
+
+using System;
+using System.Collections.Generic;
+using System.IO;
+using System.Linq;
+using System.Net;
+using System.Text;
+using System.Threading;
+using Microsoft.VisualStudio.TestTools.UnitTesting;
+using System.Threading.Tasks;
+using Horde.Storage.Implementation;
+using Jupiter.Implementation;
+using Microsoft.Extensions.Options;
+using Moq;
+using RestSharp;
+
+namespace Horde.Storage.UnitTests
+{
+    [TestClass]
+    public class ReplicatorTests
+    {
+        private readonly DirectoryInfo _tempPath = new DirectoryInfo(Path.GetTempPath());
+        private readonly Guid _currentLogGeneration = new Guid();
+        private const string CurrentSite = "Test";
+        private const string ReplicatorNameV1 = "TestReplicatorV1";
+
+        private readonly NamespaceId NamespaceV1 = new NamespaceId("test-namespace-v1");
+
+        [TestInitialize]
+        public void CleanupTempDir()
+        {
+            {
+                string stateFile = Path.Combine(_tempPath.FullName, path2: $"{ReplicatorNameV1}.json");
+                if (File.Exists(stateFile))
+                {
+                    File.Delete(stateFile);
+                }
+            }
+        }
+
+        [TestMethod]
+        public async Task ReplicatorAddTransactionReplicatorV1()
+        {
+            ReplicatorSettings replicatorSettings = new ReplicatorSettings
+            {
+                NamespaceToReplicate = NamespaceV1, ReplicatorName = ReplicatorNameV1, Version = ReplicatorVersion.V1
+            };
+
+            byte[] contents = Encoding.ASCII.GetBytes("test string");
+            BlobIdentifier blobToReplication = BlobIdentifier.FromBlob(contents);
+
+            Mock<IBlobStore> blobStoreMock = new Mock<IBlobStore>();
+            IServiceCredentials serviceCredentials = Mock.Of<IServiceCredentials>();
+            Mock<ITransactionLogWriter> transactionLogWriter = new Mock<ITransactionLogWriter>();
+            Mock<IRestClient> remoteClientMock = new Mock<IRestClient> { DefaultValue = DefaultValue.Empty };
+
+            CallistoReader.CallistoGetResponse[] responses = new[]
+            {
+                new CallistoReader.CallistoGetResponse(
+                    generation: _currentLogGeneration,
+                    currentOffset: 0,
+                    events: new List<TransactionEvent>(new[]
+                    {
+                        new AddTransactionEvent("", "", new[] {blobToReplication}, identifier:0, nextIdentifier: 100)
+                    })
+                ),
+                new CallistoReader.CallistoGetResponse(
+                    generation: _currentLogGeneration,
+                    currentOffset: 100,
+                    events: new List<TransactionEvent>()
+                )
+            };
+
+            Mock<IRestResponse<CallistoReader.CallistoGetResponse>>[] mockedCallistoResponse = responses.Select(response => CreateMockResponse(HttpStatusCode.OK, response)).ToArray();
+            // mock of callisto calls
+            remoteClientMock
+                .SetupSequence(x => x.ExecuteGetAsync<CallistoReader.CallistoGetResponse>(It.IsAny<IRestRequest>(), It.IsAny<CancellationToken>()))
+                // first we will query for a callisto event
+                .ReturnsAsync(mockedCallistoResponse[0].Object)
+                // then it will iterate again and should now reach the stop event
+                .ReturnsAsync(mockedCallistoResponse[0].Object)
+                .ReturnsAsync(mockedCallistoResponse[1].Object);
+
+            // mock of io calls
+            Mock<IRestResponse> mockedIoResponse = CreateMockResponseRaw(HttpStatusCode.OK, contents);
+            remoteClientMock
+                .Setup(x => x.ExecuteGetAsync(It.IsAny<IRestRequest>(), It.IsAny<CancellationToken>()))
+                .ReturnsAsync(mockedIoResponse.Object);
+
+            ReplicationSettings replicationSettings = new ReplicationSettings {CurrentSite = CurrentSite, StateRoot = _tempPath.FullName};
+            IOptionsMonitor<ReplicationSettings> replicationSettingsMonitor = Mock.Of<IOptionsMonitor<ReplicationSettings>>(_ => _.CurrentValue == replicationSettings);
+
+            using IReplicator replicator = new ReplicatorV1(replicatorSettings, replicationSettingsMonitor, blobStoreMock.Object, transactionLogWriter.Object, remoteClientMock.Object);
+
+            Assert.IsNull(replicator.State.ReplicatorOffset,"Expected state to have been reset during test initialize");
+
+            bool ran = await replicator.TriggerNewReplications();
+            Assert.IsTrue(ran);
+
+            NamespaceId ns = replicatorSettings.NamespaceToReplicate;
+
+            // we should have checked the local io once for if we had the blob
+            blobStoreMock.Verify(blobStore => blobStore.Exists(ns, blobToReplication), Times.Once);
+
+            // we should have attempted to store the replicated object in our callisto
+            remoteClientMock.Verify(client => client.ExecuteGetAsync<CallistoReader.CallistoGetResponse>(It.IsAny<IRestRequest>(), It.IsAny<CancellationToken>()), Times.Exactly(3));
+            transactionLogWriter.Verify(writer => writer.Add(ns, It.Is<AddTransactionEvent>(e => e.Blobs[0].Equals(blobToReplication))), Times.Once);
+
+            // as this was a add operation we should have transferred the blob from the remote blob store to the local
+            remoteClientMock.Verify(client => client.ExecuteGetAsync(It.IsAny<IRestRequest>(), It.IsAny<CancellationToken>()), Times.Once);
+            blobStoreMock.Verify(blobStore => blobStore.PutObject(ns, It.IsAny<byte[]>(), blobToReplication), Times.Once);
+
+            Assert.AreEqual(_currentLogGeneration, replicator.State.ReplicatingGeneration);
+
+            remoteClientMock.VerifyNoOtherCalls();
+        }
+
+
+        [TestMethod]
+        public async Task ReplicatorSkipLogReplicatorV1()
+        {
+            ReplicatorSettings replicatorSettings = new ReplicatorSettings
+            {
+                NamespaceToReplicate = NamespaceV1, ReplicatorName = ReplicatorNameV1, Version = ReplicatorVersion.V1
+            };
+
+            Mock<IRestClient> remoteClientMock = new Mock<IRestClient> { DefaultValue = DefaultValue.Empty };
+
+            CallistoReader.CallistoGetResponse[] responses = new[]
+            {
+                new CallistoReader.CallistoGetResponse(
+                    generation: _currentLogGeneration,
+                    currentOffset: 100000,
+                    events: new List<TransactionEvent>()
+                ),
+            };
+            Mock<IRestResponse<CallistoReader.CallistoGetResponse>>[] mockedCallistoResponse = responses.Select(response => CreateMockResponse(HttpStatusCode.OK, response)).ToArray();
+            // mock of callisto calls
+            remoteClientMock
+                .SetupSequence(x => x.ExecuteGetAsync<CallistoReader.CallistoGetResponse>(It.IsAny<IRestRequest>(), It.IsAny<CancellationToken>()))
+                .ReturnsAsync(mockedCallistoResponse[0].Object)
+                .ReturnsAsync(mockedCallistoResponse[0].Object);
+
+            Mock<IBlobStore> blobStoreMock = new Mock<IBlobStore>();
+            Mock<ITransactionLogWriter> transactionLogWriter = new Mock<ITransactionLogWriter>();
+            ReplicationSettings replicationSettings = new ReplicationSettings {CurrentSite = CurrentSite, StateRoot = _tempPath.FullName};
+            IOptionsMonitor<ReplicationSettings> replicationSettingsMonitor = Mock.Of<IOptionsMonitor<ReplicationSettings>>(_ => _.CurrentValue == replicationSettings);
+
+            using IReplicator replicator = new ReplicatorV1(replicatorSettings, replicationSettingsMonitor, blobStoreMock.Object, transactionLogWriter.Object, remoteClientMock.Object);
+
+            Assert.IsNull(replicator.State.ReplicatorOffset,"Expected state to have been reset during test initialize");
+
+            bool replicatedEvents = await replicator.TriggerNewReplications();
+            Assert.IsFalse(replicatedEvents);
+
+            blobStoreMock.VerifyNoOtherCalls();
+            transactionLogWriter.VerifyNoOtherCalls();
+
+            Assert.AreEqual(100000, replicator.State.ReplicatorOffset);
+            Assert.AreEqual(_currentLogGeneration, replicator.State.ReplicatingGeneration);
+
+            // we should called the transaction log 2
+            remoteClientMock.Verify(client => client.ExecuteGetAsync<CallistoReader.CallistoGetResponse>(It.IsAny<IRestRequest>(), It.IsAny<CancellationToken>()), Times.Exactly(2));
+
+            remoteClientMock.VerifyNoOtherCalls();
+        }
+        
+        [TestMethod]
+        public async Task ReplicatorSendingIncorrectPayloadV1()
+        {
+            ReplicatorSettings replicatorSettings = new ReplicatorSettings
+            {
+                NamespaceToReplicate = NamespaceV1, ReplicatorName = ReplicatorNameV1, Version = ReplicatorVersion.V1
+            };
+
+            byte[] contents = Encoding.ASCII.GetBytes("test string");
+            BlobIdentifier blobToReplication = BlobIdentifier.FromBlob(contents);
+
+            Mock<IBlobStore> blobStoreMock = new Mock<IBlobStore>();
+            IServiceCredentials serviceCredentials = Mock.Of<IServiceCredentials>();
+            Mock<ITransactionLogWriter> transactionLogWriter = new Mock<ITransactionLogWriter>();
+            Mock<IRestClient> remoteClientMock = new Mock<IRestClient> { DefaultValue = DefaultValue.Empty };
+
+            CallistoReader.CallistoGetResponse[] responses = new[]
+            {
+                new CallistoReader.CallistoGetResponse(
+                    generation: _currentLogGeneration,
+                    currentOffset: 0,
+                    events: new List<TransactionEvent>(new[]
+                    {
+                        new AddTransactionEvent("", "", new[] {blobToReplication}, identifier:0, nextIdentifier: 100)
+                    })
+                ),
+                new CallistoReader.CallistoGetResponse(
+                    generation: _currentLogGeneration,
+                    currentOffset: 100,
+                    events: new List<TransactionEvent>()
+                )
+            };
+
+            Mock<IRestResponse<CallistoReader.CallistoGetResponse>>[] mockedCallistoResponse = responses.Select(response => CreateMockResponse(HttpStatusCode.OK, response)).ToArray();
+            // mock of callisto calls
+            remoteClientMock
+                .SetupSequence(x => x.ExecuteGetAsync<CallistoReader.CallistoGetResponse>(It.IsAny<IRestRequest>(), It.IsAny<CancellationToken>()))
+                // first we will query for a callisto event
+                .ReturnsAsync(mockedCallistoResponse[0].Object)
+                // then it will iterate again and should now reach the stop event
+                .ReturnsAsync(mockedCallistoResponse[0].Object)
+                .ReturnsAsync(mockedCallistoResponse[1].Object);
+
+            // mock of io calls
+            // return a empty array (0 byte response) from the remote IO and on retry returns a valid object
+            Mock<IRestResponse> mockedIoResponse = CreateMockResponseRaw(HttpStatusCode.OK, Array.Empty<byte>());
+            Mock<IRestResponse> mockedIoResponse2 = CreateMockResponseRaw(HttpStatusCode.OK, contents);
+            remoteClientMock
+                .SetupSequence(x => x.ExecuteGetAsync(It.IsAny<IRestRequest>(), It.IsAny<CancellationToken>()))
+                .ReturnsAsync(mockedIoResponse.Object)
+                .ReturnsAsync(mockedIoResponse2.Object);
+
+            ReplicationSettings replicationSettings = new ReplicationSettings {CurrentSite = CurrentSite, StateRoot = _tempPath.FullName};
+            IOptionsMonitor<ReplicationSettings> replicationSettingsMonitor = Mock.Of<IOptionsMonitor<ReplicationSettings>>(_ => _.CurrentValue == replicationSettings);
+
+            using IReplicator replicator = new ReplicatorV1(replicatorSettings, replicationSettingsMonitor, blobStoreMock.Object, transactionLogWriter.Object, remoteClientMock.Object);
+
+            Assert.IsNull(replicator.State.ReplicatorOffset,"Expected state to have been reset during test initialize");
+
+            bool ran = await replicator.TriggerNewReplications();
+            Assert.IsTrue(ran);
+
+            NamespaceId ns = replicatorSettings.NamespaceToReplicate;
+
+            // we should have checked the local io once for if we had the blob
+            blobStoreMock.Verify(blobStore => blobStore.Exists(ns, blobToReplication), Times.Once);
+
+            // we should have attempted to store the replicated object in our callisto
+            remoteClientMock.Verify(client => client.ExecuteGetAsync<CallistoReader.CallistoGetResponse>(It.IsAny<IRestRequest>(), It.IsAny<CancellationToken>()), Times.Exactly(3));
+            transactionLogWriter.Verify(writer => writer.Add(ns, It.Is<AddTransactionEvent>(e => e.Blobs[0].Equals(blobToReplication))), Times.Once);
+
+            // as this was a add operation we should have transferred the blob from the remote blob store to the local
+            remoteClientMock.Verify(client => client.ExecuteGetAsync(It.IsAny<IRestRequest>(), It.IsAny<CancellationToken>()), Times.Exactly(2));
+            blobStoreMock.Verify(blobStore => blobStore.PutObject(ns, It.IsAny<byte[]>(), blobToReplication), Times.Once);
+
+            Assert.AreEqual(_currentLogGeneration, replicator.State.ReplicatingGeneration);
+
+            remoteClientMock.VerifyNoOtherCalls();
+        }
+
+
+        [TestMethod]
+        public async Task ReplicatorDeleteTransactionV1()
+        {
+            ReplicatorSettings replicatorSettings = new ReplicatorSettings
+            {
+                NamespaceToReplicate = NamespaceV1, ReplicatorName = ReplicatorNameV1, Version = ReplicatorVersion.V1
+            };
+
+            NamespaceId ns = replicatorSettings.NamespaceToReplicate;
+
+            Mock<IBlobStore> blobStoreMock = new Mock<IBlobStore>();
+            IServiceCredentials serviceCredentials = Mock.Of<IServiceCredentials>();
+            Mock<ITransactionLogWriter> transactionLogWriter = new Mock<ITransactionLogWriter>();
+            Mock<IRestClient> remoteClientMock = new Mock<IRestClient> { DefaultValue = DefaultValue.Empty };
+
+            CallistoReader.CallistoGetResponse[] responses = new[]
+            {
+                new CallistoReader.CallistoGetResponse(
+                    generation: _currentLogGeneration,
+                    currentOffset: 0,
+                    events: new List<TransactionEvent>(new[]
+                    {
+                        new RemoveTransactionEvent("", "", 0, 100)
+                    })
+                ),
+                new CallistoReader.CallistoGetResponse(
+                    generation: _currentLogGeneration,
+                    currentOffset: 100,
+                    events: new List<TransactionEvent>()
+                )
+            };
+
+            Mock<IRestResponse<CallistoReader.CallistoGetResponse>>[] mockedCallistoResponse = responses.Select(response => CreateMockResponse(HttpStatusCode.OK, response)).ToArray();
+            // mock of callisto calls
+            remoteClientMock
+                .SetupSequence(x => x.ExecuteGetAsync<CallistoReader.CallistoGetResponse>(It.IsAny<IRestRequest>(), It.IsAny<CancellationToken>()))
+                // first we will query for a callisto event
+                .ReturnsAsync(mockedCallistoResponse[0].Object)
+                // then it will iterate again and should now reach the stop event
+                .ReturnsAsync(mockedCallistoResponse[0].Object)
+                .ReturnsAsync(mockedCallistoResponse[1].Object);
+
+            ReplicationSettings replicationSettings = new ReplicationSettings {CurrentSite = CurrentSite, StateRoot = _tempPath.FullName};
+            IOptionsMonitor<ReplicationSettings> replicationSettingsMonitor = Mock.Of<IOptionsMonitor<ReplicationSettings>>(_ => _.CurrentValue == replicationSettings);
+
+            using IReplicator replicator = new ReplicatorV1(replicatorSettings, replicationSettingsMonitor, blobStoreMock.Object, transactionLogWriter.Object, remoteClientMock.Object);
+
+            Assert.IsNull(replicator.State.ReplicatorOffset, "Expected state to have been reset during test initialize");
+
+            bool ran = await replicator.TriggerNewReplications();
+            Assert.IsTrue(ran);
+
+            // we should have fetched the replicated object
+            remoteClientMock.Verify(client => client.ExecuteGetAsync<CallistoReader.CallistoGetResponse>(It.IsAny<IRestRequest>(), It.IsAny<CancellationToken>()), Times.Exactly(3));
+            transactionLogWriter.Verify(writer => writer.Add(ns, It.IsAny<RemoveTransactionEvent>()), Times.Once);
+            // other then copying the remove event we should have written nothing else
+            transactionLogWriter.VerifyNoOtherCalls();
+            // this was a delete so should not have stored anything in the local store
+            blobStoreMock.VerifyNoOtherCalls();
+
+            Assert.AreEqual(_currentLogGeneration, replicator.State.ReplicatingGeneration);
+
+            remoteClientMock.VerifyNoOtherCalls();
+        }
+
+        private Mock<IRestResponse<T>> CreateMockResponse<T>(HttpStatusCode httpStatusCode, T data)
+        {
+            Mock<IRestResponse<T>> response = new Mock<IRestResponse<T>>();
+            response.Setup(_ => _.StatusCode).Returns(httpStatusCode);
+            response.Setup(_ => _.IsSuccessful).Returns(httpStatusCode == HttpStatusCode.OK);
+            response.Setup(_ => _.Data).Returns(data);
+
+            return response;
+        }
+        
+        public static Mock<IRestResponse> CreateMockResponseRaw(HttpStatusCode httpStatusCode, byte[] data)
+        {
+            Mock<IRestResponse> response = new Mock<IRestResponse>();
+            response.Setup(_ => _.StatusCode).Returns(httpStatusCode);
+            response.Setup(_ => _.IsSuccessful).Returns(httpStatusCode == HttpStatusCode.OK);
+            response.Setup(_ => _.RawBytes).Returns(data);
+
+            return response;
+        }
+
+    }
+}
diff --git a/Engine/Source/Programs/Horde/HordeStorage/Horde.Storage/Tests/UnitTests/S3StoreTests.cs b/Engine/Source/Programs/Horde/HordeStorage/Horde.Storage/Tests/UnitTests/S3StoreTests.cs
new file mode 100644
--- /dev/null
+++ b/Engine/Source/Programs/Horde/HordeStorage/Horde.Storage/Tests/UnitTests/S3StoreTests.cs
@@ -0,0 +1,99 @@
+﻿// Copyright Epic Games, Inc. All Rights Reserved.
+
+using System;
+using System.IO;
+using System.Text;
+using System.Threading.Tasks;
+using Amazon.S3;
+using Amazon.S3.Model;
+using Horde.Storage.Implementation;
+using Jupiter.Implementation;
+using Microsoft.Extensions.Options;
+using Microsoft.VisualStudio.TestTools.UnitTesting;
+using Moq;
+
+namespace Horde.Storage.UnitTests
+{
+    [TestClass]
+    public class S3StoreTests
+    {
+        private NamespaceId Namespace = new NamespaceId("foo");
+
+        [TestMethod]
+        public async Task PutBufferSuccess()
+        {
+            Mock<IAmazonS3> s3Mock = new Mock<IAmazonS3>();
+            S3Settings settings = new S3Settings{BucketName = "tests-foo"};
+            IOptionsMonitor<S3Settings> settingsMock = Mock.Of<IOptionsMonitor<S3Settings>>(_ => _.CurrentValue == settings);
+
+            AmazonS3Store store = new AmazonS3Store(s3Mock.Object, settingsMock);
+            byte[] content = Encoding.ASCII.GetBytes("test content");
+            BlobIdentifier blobIdentifier = BlobIdentifier.FromBlob(content);
+            Task task = store.PutObject(Namespace, content.AsMemory(), blobIdentifier);
+            await task;
+
+            s3Mock.Verify( s3 => s3.DoesS3BucketExistAsync("tests-foo"));
+            s3Mock.Verify( s3 => s3.PutBucketAsync(It.Is<PutBucketRequest>(p => p.BucketName =="tests-foo"), default));
+            s3Mock.Verify(s3 =>
+                s3.PutObjectAsync(It.Is<PutObjectRequest>(p => p.BucketName == "tests-foo" && p.Key == blobIdentifier.AsS3Key()), default));
+
+            s3Mock.Verify( s3 => s3.PutPublicAccessBlockAsync(It.IsAny<PutPublicAccessBlockRequest>(), default));
+            s3Mock.VerifyNoOtherCalls();
+        }
+
+        [TestMethod]
+        [ExpectedException(typeof(Exception))]
+        public async Task PutBufferError()
+        {
+            byte[] content = Encoding.ASCII.GetBytes("test content");
+            BlobIdentifier blob = BlobIdentifier.FromBlob(content);
+
+            Mock<IAmazonS3> s3Mock = new Mock<IAmazonS3>();
+            S3Settings settings = new S3Settings{BucketName = "tests-foo"};
+            IOptionsMonitor<S3Settings> settingsMock = Mock.Of<IOptionsMonitor<S3Settings>>(_ => _.CurrentValue == settings);
+
+            s3Mock.Setup(s3 => s3.PutObjectAsync(It.IsAny<PutObjectRequest>(), default)).Throws<Exception>();
+            AmazonS3Store store = new AmazonS3Store(s3Mock.Object, settingsMock);
+            Task task = store.PutObject(Namespace, content, blob);
+            await task;
+
+            s3Mock.Verify(s3 =>
+                s3.PutObjectAsync(
+                    It.Is<PutObjectRequest>(p => p.BucketName == "tests-foo" && p.FilePath == blob.AsS3Key() && p.InputStream == null),
+                    default));
+        }
+
+        [TestMethod]
+        public async Task GetSuccess()
+        {
+            byte[] content = Encoding.ASCII.GetBytes("test content");
+            BlobIdentifier blob = BlobIdentifier.FromBlob(content);
+            Mock<IAmazonS3> s3Mock = new Mock<IAmazonS3>();
+            s3Mock.Setup(s3 => s3.GetObjectAsync("tests-foo", blob.AsS3Key(), default)).ReturnsAsync(Mock.Of<GetObjectResponse>());
+
+            S3Settings settings = new S3Settings{BucketName = "tests-foo"};
+            IOptionsMonitor<S3Settings> settingsMock = Mock.Of<IOptionsMonitor<S3Settings>>(_ => _.CurrentValue == settings);
+
+            AmazonS3Store store = new AmazonS3Store(s3Mock.Object, settingsMock);
+            await using BlobContents blobContents = await store.GetObject(Namespace, blob);
+
+            s3Mock.Verify(s3 => s3.GetObjectAsync("tests-foo", blob.AsS3Key(), default));
+        }
+
+        [TestMethod]
+        public async Task DeleteSuccess()
+        {
+            byte[] content = Encoding.ASCII.GetBytes("test content");
+            BlobIdentifier blob = BlobIdentifier.FromBlob(content);
+            Mock<IAmazonS3> s3Mock = new Mock<IAmazonS3>();
+            S3Settings settings = new S3Settings{BucketName = "tests-foo"};
+            IOptionsMonitor<S3Settings> settingsMock = Mock.Of<IOptionsMonitor<S3Settings>>(_ => _.CurrentValue == settings);
+
+            AmazonS3Store store = new AmazonS3Store(s3Mock.Object, settingsMock);
+            Task task = store.DeleteObject(Namespace, blob);
+            await task;
+
+            s3Mock.Verify(s3 => s3.DeleteObjectAsync("tests-foo", blob.AsS3Key(), default));
+        }
+    }
+}
