Moved Horde.Storage from git to p4

diff --git a/Engine/Source/Programs/Horde/HordeStorage/Horde.Storage/appsettings.Development.json b/Engine/Source/Programs/Horde/HordeStorage/Horde.Storage/appsettings.Development.json
new file mode 100644
--- /dev/null
+++ b/Engine/Source/Programs/Horde/HordeStorage/Horde.Storage/appsettings.Development.json
@@ -0,0 +1,99 @@
+{
+  "Serilog": {
+    "Using": [ "Serilog.Sinks.Console" ],
+    "MinimumLevel": {
+      "Default": "Information",
+      "Override": {
+        "Microsoft": "Information",
+        "System": "Information"
+      }
+    },
+    "WriteTo": [
+      { "Name": "Console" }
+    ],
+    "Enrich": [ "FromLogContext", "WithMachineName", "WithThreadId" ],
+    "Properties": {
+      "Application": "Horde.Storage"
+    }
+  },
+
+  "Horde.Storage": {
+    "refDBImplementation": "mongo",
+
+    "TransactionLogWriterImplementation": "memory",
+
+    "treeStoreImplementation": "memory",
+
+    "replicationLogWriterImplementation": "memory",
+
+    "contentIdStoreImplementation": "memory"
+  },
+
+  "Callisto": {
+    "ConnectionString": "http://localhost"
+  },
+
+  "GC": {
+    "BlobCleanupServiceEnabled": false,
+    "CleanOldRefRecords": false,
+    "CleanOldBlobs": false
+  },
+
+  "Replication": {
+    "StateRoot": "./ganymede",
+    "ReplicationPollFrequencySeconds": 15,
+    "CurrentSite": "local-test",
+    "Replicators": [
+      {
+        "ReplicatorName": "Local test",
+        "NamespaceToReplicate": "ue4.ddc",
+
+        "ConnectionString": "https://jupiter.devtools-dev.epicgames.com"
+      }
+    ]
+  },
+
+  "Auth": {
+    "Method": "Disabled"
+  },
+
+  "Authorization": {
+    "NamespaceToClaim": {
+      "ue4.ddc": "*"
+    }
+  },
+
+  "Mongo": {
+    "ConnectionString": "mongodb://root:abc123@localhost:27017",
+    "RequireTls12": false
+  },
+
+  "DynamoDb": {
+    "ConnectionString": "localhost:8000",
+    "UseGlobalTables": false
+  },
+
+  "Jupiter": {
+    "DisableAuthOnPorts": [
+      57003
+    ]
+  },
+  "ConsistencyCheck": {
+    "Enabled": false
+  },
+
+  "Kestrel": {
+    "Endpoints": {
+      "Http": {
+        "Url": "http://localhost:57001"
+      },
+      "Http2": {
+        "Url": "http://localhost:57002",
+        "Protocols": "Http2"
+      },
+      "AuthDisabledHttp": {
+        "Url": "http://localhost:57003"
+      }
+    }
+  }
+}
diff --git a/Engine/Source/Programs/Horde/HordeStorage/Horde.Storage/appsettings.Production.json b/Engine/Source/Programs/Horde/HordeStorage/Horde.Storage/appsettings.Production.json
new file mode 100644
--- /dev/null
+++ b/Engine/Source/Programs/Horde/HordeStorage/Horde.Storage/appsettings.Production.json
@@ -0,0 +1,51 @@
+{
+  "Serilog": {
+    "Using": [ "Serilog.Sinks.Console" ],
+    "MinimumLevel": {
+      "Default": "Information",
+      "Override": {
+        "Microsoft": "Warning",
+        "System": "Information"
+      }
+    },
+    "WriteTo": [
+      {
+        "Name": "Console",
+        "Args": {
+          "formatter": "Serilog.Formatting.Json.JsonFormatter"
+        }
+      }
+    ],
+    "Enrich": [ "FromLogContext" ]
+  },
+
+  "Mongo": {
+    "RequireTls12": true
+  },
+
+  "DynamoDB": {
+    "AutoscaleCapacity": true
+  },
+
+  "Horde.Storage": {
+    "TransactionLogWriterImplementation": "callisto"
+  },
+
+  "Replication": {
+    "ReplicationPollFrequencySeconds": 15
+  },
+
+  "AllowedHosts": "*",
+
+  "Kestrel": {
+    "Endpoints": {
+      "Http": {
+        "Url": "http://0.0.0.0:80"
+      },
+      "Http2": {
+        "Url": "http://0.0.0.0:8080",
+        "Protocols": "Http2"
+      }
+    }
+  }
+}
diff --git a/Engine/Source/Programs/Horde/HordeStorage/Horde.Storage/appsettings.Testing.json b/Engine/Source/Programs/Horde/HordeStorage/Horde.Storage/appsettings.Testing.json
new file mode 100644
--- /dev/null
+++ b/Engine/Source/Programs/Horde/HordeStorage/Horde.Storage/appsettings.Testing.json
@@ -0,0 +1,106 @@
+{
+  "Serilog": {
+    "Using": [ "Serilog.Sinks.Console" ],
+    "MinimumLevel": {
+      "Default": "Information",
+      "Override": {
+        "Microsoft": "Warning",
+        "System": "Information"
+      }
+    },
+    "WriteTo": [
+      { "Name": "Console" }
+    ],
+    "Enrich": [ "FromLogContext", "WithMachineName", "WithThreadId" ],
+    "Properties": {
+      "Application": "Horde.Storage"
+    }
+  },
+
+  "Mongo": {
+    "ConnectionString": "mongodb://root:abc123@localhost:27017"
+  },
+
+  "Scylla": {
+    "LocalDatacenterName": "datacenter1",
+    "LocalKeyspaceSuffix": "test"
+  },
+
+  "DynamoDb": {
+    "ConnectionString": "http://localhost:8000",
+    "UseOndemandCapacityProvisioning": false,
+    "UseGlobalTables": false
+  },
+
+  "Horde.Storage": {
+    "RefDBImplementation": "mongo",
+    "StorageImplementations": [ "memory" ],
+    "TransactionLogWriterImplementation": "memory"
+  },
+
+  "Callisto": {
+    "ConnectionString": "http://localhost"
+  },
+
+  "Replication": {
+    "StateRoot": "./ganymede",
+    "ReplicationPollFrequencySeconds": 15,
+    "CurrentSite": "local-test",
+    "Enabled": false
+  },
+
+  "GC": {
+    "BlobCleanupServiceEnabled": false,
+    "CleanOldBlobs": false,
+    "CleanOldRefRecords": false
+  },
+
+  "S3": {
+    // use a local S3 emulation when running the tests
+    "ConnectionString": "http://127.0.0.1:9000",
+    "ForceAWSPathStyle": true, // required by minio,
+    "SetBucketPolicies": false, // minio does not support bucket access policies
+    "BucketPrefix": "tests"
+  },
+  "AWSCredentials": {
+    // minio and health checks expect access keys to have been set
+    "AwsAccessKey": "FSYPTVCPKC2IV5D1U8RE",
+    "AwsSecretKey": "WaUf4Zq+J6YMlvbCVvTRkQJIPc7DvZwTm5UP0tPJ"
+  },
+  "Azure": {
+    "ConnectionString": "DefaultEndpointsProtocol=https;BlobEndpoint=http://localhost:11002/blobaccount;AccountName=blobaccount;AccountKey=iU6uTvlF1ysppmft+NO5lAD0E3hwrAORr5Rb5xcBWUgEz/OicrSkFxwZYMNK5XL29/wXZKGOoOVSW040nAOfPg==",
+  },
+  "Auth": {
+    "Method": "Disabled"
+  },
+
+  "Authorization": {
+    "NamespaceToClaim": {
+      "ue4.ddc": "*",
+      "TestNamespace": "*"
+    }
+  },
+
+  "AllowedHosts": "*",
+
+  "Jupiter": {
+    "DisableAuthOnPorts": [
+      57003
+    ]
+  },
+
+  "Kestrel": {
+    "Endpoints": {
+      "Http": {
+        "Url": "http://localhost:57001"
+      },
+      "Http2": {
+        "Url": "http://localhost:57002",
+        "Protocols": "Http2"
+      },
+      "AuthDisabledHttp": {
+        "Url": "http://localhost:57003"
+      }
+    }
+  }
+}
diff --git a/Engine/Source/Programs/Horde/HordeStorage/Horde.Storage/appsettings.json b/Engine/Source/Programs/Horde/HordeStorage/Horde.Storage/appsettings.json
new file mode 100644
--- /dev/null
+++ b/Engine/Source/Programs/Horde/HordeStorage/Horde.Storage/appsettings.json
@@ -0,0 +1,18 @@
+{
+  "AllowedHosts": "*",
+
+  "Cache.Blob": {
+    // limit cache to 1 GB
+    "SizeLimit": "1073741824"
+  },
+  "Cache.Tree": {
+    // limit cache to 1 GB
+    "SizeLimit": "1073741824"
+  },
+
+  "GRPC": {
+    // bump send and receive size to 1gb
+    "MaxReceiveMessageSize": "1073741824",
+    "MaxSendMessageSize": "1073741824"
+  }
+}
diff --git a/Engine/Source/Programs/Horde/HordeStorage/Jupiter.Common/Authentication/NamespaceAuthorizationHandler.cs b/Engine/Source/Programs/Horde/HordeStorage/Jupiter.Common/Authentication/NamespaceAuthorizationHandler.cs
new file mode 100644
--- /dev/null
+++ b/Engine/Source/Programs/Horde/HordeStorage/Jupiter.Common/Authentication/NamespaceAuthorizationHandler.cs
@@ -0,0 +1,59 @@
+﻿// Copyright Epic Games, Inc. All Rights Reserved.
+
+using System.Collections.Generic;
+using System.ComponentModel.DataAnnotations;
+using System.Threading.Tasks;
+using Jupiter.Implementation;
+using Microsoft.AspNetCore.Authorization;
+using Microsoft.Extensions.Options;
+
+namespace Jupiter
+{
+    // verifies that you have access to a namespace by checking if you have a corresponding claim to that namespace
+    public class NamespaceAuthorizationHandler : AuthorizationHandler<NamespaceAccessRequirement, NamespaceId>
+    {
+        private readonly IOptionsMonitor<AuthorizationSettings> _authorizationSettings;
+
+        public NamespaceAuthorizationHandler(IOptionsMonitor<AuthorizationSettings> authorizationSettings)
+        {
+            _authorizationSettings = authorizationSettings;
+        }
+
+        protected override Task HandleRequirementAsync(AuthorizationHandlerContext context, NamespaceAccessRequirement requirement,
+            NamespaceId namespaceName)
+        {
+            if (context.User.HasClaim(claim => claim.Type == "AllNamespaces"))
+            {
+                context.Succeed(requirement);
+                return Task.CompletedTask;
+            }
+
+            if (_authorizationSettings.CurrentValue.NamespaceToClaim.TryGetValue(namespaceName.ToString(), out string? expectedClaim))
+            {
+                // if expected claim is * then everyone is allowed to use the namespace
+                if (expectedClaim == "*")
+                {
+                    context.Succeed(requirement);
+                }
+
+                if (context.User.HasClaim(claim => claim.Type == expectedClaim))
+                {
+                    context.Succeed(requirement);
+                }
+            }
+
+
+            return Task.CompletedTask;
+        }
+    }
+
+    public class AuthorizationSettings
+    {
+        [Required] public Dictionary<string, string> NamespaceToClaim { get; set; } = null!;
+    }
+
+    public class NamespaceAccessRequirement : IAuthorizationRequirement
+    {
+        public const string Name = "NamespaceAccess";
+    }
+}
diff --git a/Engine/Source/Programs/Horde/HordeStorage/Jupiter.Common/BaseProgram.cs b/Engine/Source/Programs/Horde/HordeStorage/Jupiter.Common/BaseProgram.cs
new file mode 100644
--- /dev/null
+++ b/Engine/Source/Programs/Horde/HordeStorage/Jupiter.Common/BaseProgram.cs
@@ -0,0 +1,115 @@
+// Copyright Epic Games, Inc. All Rights Reserved.
+
+using System;
+using System.Collections.Generic;
+using System.IO;
+using System.IO.Pipelines;
+using System.Net.Sockets;
+using System.Reflection;
+using System.Threading.Tasks;
+using Microsoft.AspNetCore.Connections;
+using Microsoft.AspNetCore.Hosting;
+using Microsoft.Extensions.Configuration;
+using Microsoft.Extensions.Hosting;
+using Serilog;
+
+namespace Jupiter
+{
+    public abstract class BaseHttpConnection
+    {
+        protected readonly IServiceProvider ServiceProvider;
+        protected readonly PipeReader Reader;
+        protected readonly PipeWriter Writer;
+        protected readonly Socket Socket;
+
+        public BaseHttpConnection(IServiceProvider serviceProvider, PipeReader reader, PipeWriter writer, Socket socket)
+        {
+            ServiceProvider = serviceProvider;
+            Reader = reader;
+            Writer = writer;
+            Socket = socket;
+        }
+
+        /// <summary>
+        /// A function that can process a connection.
+        /// </summary>
+        /// <param name="connection">A <see cref="ConnectionContext" /> representing the connection.</param>
+        /// <returns>A <see cref="Task"/> that represents the connection lifetime. When the task completes, the connection will be closed.</returns>
+        public abstract Task ExecuteAsync(ConnectionContext connection);
+    }
+    
+    // ReSharper disable once UnusedMember.Global
+    public static class BaseProgram<T> where T : BaseStartup
+    {
+        public delegate BaseHttpConnection HttpConnectionFactory(IServiceProvider sp, PipeReader reader, PipeWriter writer, Socket socket);
+
+        private static IConfiguration Configuration { get; } = GetConfigurationBuilder();
+
+        private static IConfiguration GetConfigurationBuilder()
+        {
+            string env = Environment.GetEnvironmentVariable("ASPNETCORE_ENVIRONMENT") ?? "Production";
+            string configRoot = "/config";
+            // set the config root to config under the current directory for windows
+            if (OperatingSystem.IsWindows())
+            {
+                configRoot = Path.Combine(Directory.GetCurrentDirectory(), "config");
+            }
+            return new ConfigurationBuilder()
+                .SetBasePath(Directory.GetCurrentDirectory())
+                .AddJsonFile("appsettings.json", false, false)
+                .AddJsonFile(
+                    path:
+                    $"appsettings.{env}.json",
+                    true)
+                .AddYamlFile(Path.Combine(configRoot, "appsettings.Local.yaml"), optional: true, reloadOnChange: true)
+                .AddEnvironmentVariables()
+                .Build();
+
+        }
+
+        public static int BaseMain(string[] args, HttpConnectionFactory? httpConnectionFactory = null)
+        {
+            Log.Logger = new LoggerConfiguration()
+                .ReadFrom.Configuration(Configuration)
+                .CreateLogger();
+
+            try
+            {
+                Log.Information("Creating ASPNET Host");
+                CreateHostBuilder(args, httpConnectionFactory).Build().Run();
+                return 0;
+            }
+            catch (Exception ex)
+            {
+                Log.Fatal(ex, "Host terminated unexpectedly");
+                return 1;
+            }
+            finally
+            {
+                Log.CloseAndFlush();
+            }
+        }
+
+        private static IHostBuilder CreateHostBuilder(string[] args, HttpConnectionFactory? httpConnectionFactory)
+        {
+            return Host.CreateDefaultBuilder(args)
+                .ConfigureWebHostDefaults(webBuilder =>
+                {
+                    webBuilder.UseStartup<T>();
+                    webBuilder.UseConfiguration(Configuration);
+                    webBuilder.UseSerilog();
+                    // configure microsoft.extensions.logging to configure log4net to allow us to set it in our appsettings
+                    webBuilder.ConfigureLogging((hostingContext, logging) =>
+                    {
+                        // configure log4net (used by aws sdk) to write to serilog so we get the logs in the system we want it in
+                        Log4net.Appender.Serilog.Configuration.Configure();
+                    });
+                    // remove the server header from kestrel
+                    webBuilder.ConfigureKestrel(options =>
+                    {
+                        options.AddServerHeader = false;
+                    });
+                });
+        }
+    }
+}
diff --git a/Engine/Source/Programs/Horde/HordeStorage/Jupiter.Common/BaseStartup.cs b/Engine/Source/Programs/Horde/HordeStorage/Jupiter.Common/BaseStartup.cs
new file mode 100644
--- /dev/null
+++ b/Engine/Source/Programs/Horde/HordeStorage/Jupiter.Common/BaseStartup.cs
@@ -0,0 +1,472 @@
+// Copyright Epic Games, Inc. All Rights Reserved.
+
+using System;
+using System.Collections.Generic;
+using System.ComponentModel.DataAnnotations;
+using System.IO;
+using System.Linq;
+using System.Net.Mime;
+using System.Reflection;
+using System.Text;
+using System.Threading.Tasks;
+using Amazon;
+using Datadog.Trace;
+using Microsoft.AspNetCore.Authentication;
+using Microsoft.AspNetCore.Authorization;
+using Microsoft.AspNetCore.Builder;
+using Microsoft.AspNetCore.Diagnostics.HealthChecks;
+using Microsoft.AspNetCore.Hosting;
+using Microsoft.AspNetCore.Http;
+using Microsoft.AspNetCore.HttpOverrides;
+using Microsoft.AspNetCore.Mvc;
+using Microsoft.AspNetCore.Routing;
+using Microsoft.Extensions.Configuration;
+using Microsoft.Extensions.DependencyInjection;
+using Microsoft.Extensions.Diagnostics.HealthChecks;
+using Microsoft.Extensions.Hosting;
+using Microsoft.Extensions.Options;
+using Microsoft.Extensions.Primitives;
+using Microsoft.OpenApi.Models;
+using Newtonsoft.Json;
+using Newtonsoft.Json.Serialization;
+using Okta.AspNet.Abstractions;
+using Okta.AspNetCore;
+using Serilog;
+using ILogger = Serilog.ILogger;
+using OktaWebApiOptions = Okta.AspNetCore.OktaWebApiOptions;
+
+namespace Jupiter
+{
+    public abstract class BaseStartup
+    {
+        protected readonly ILogger _logger = Log.ForContext<BaseStartup>();
+
+        public BaseStartup(IConfiguration configuration, IWebHostEnvironment environment)
+        {
+            Configuration = configuration;
+            Environment = environment;
+            Auth = new AuthSettings();
+        }
+
+        public IConfiguration Configuration { get; }
+        public IWebHostEnvironment Environment { get; }
+
+        protected AuthSettings Auth { get; }
+
+        // This method gets called by the runtime. Use this method to add services to the container.
+        public void ConfigureServices(IServiceCollection services)
+        {
+            // aws specific settings
+            services.AddOptions<AWSCredentialsSettings>().Bind(Configuration.GetSection("AWSCredentials")).ValidateDataAnnotations();
+            services.AddDefaultAWSOptions(Configuration.GetAWSOptions());
+            // send log4net logs to serilog and configure aws to log to log4net (they lack a serilog implementation)
+            AWSConfigs.LoggingConfig.LogTo = LoggingOptions.Log4Net;
+
+            services.AddOptions<AuthSettings>().Bind(Configuration.GetSection("Auth")).ValidateDataAnnotations();
+            Configuration.GetSection("Auth").Bind(Auth);
+
+            services.AddOptions<ServiceAccountAuthOptions>().Bind(Configuration.GetSection("ServiceAccounts")).ValidateDataAnnotations();
+
+            services.AddOptions<AuthorizationSettings>().Bind(Configuration.GetSection("Authorization")).ValidateDataAnnotations();
+            services.AddOptions<JupiterSettings>().Bind(Configuration.GetSection("Jupiter")).ValidateDataAnnotations();
+
+            // this is the same as invoke MvcBuilder.AddJsonOptions but with a service provider passed so we can use DI in the options creation
+            // see https://stackoverflow.com/questions/53288633/net-core-api-custom-json-resolver-based-on-request-values
+            // inject our custom json options and then pass a DI context to customize the serialization
+            services.Configure<MvcNewtonsoftJsonOptions>(OnJsonOptions);
+            services.AddTransient<IConfigureOptions<MvcNewtonsoftJsonOptions>, MvcJsonOptionsWrapper>();
+
+            services.AddControllers()
+                .AddNewtonsoftJson()
+                .AddMvcOptions(options =>
+                {
+                    options.InputFormatters.Add(new CompactBinaryInputFormatter());
+                    options.OutputFormatters.Add(new CompactBinaryOutputFormatter());
+                    options.OutputFormatters.Add(new RawOutputFormatter());
+
+                    options.FormatterMappings.SetMediaTypeMappingForFormat("raw", MediaTypeNames.Application.Octet);
+                    options.FormatterMappings.SetMediaTypeMappingForFormat("uecb", CustomMediaTypeNames.UnrealCompactBinary);
+
+                    OnAddControllers(options);
+                }).ConfigureApiBehaviorOptions(options =>
+                {
+                    options.InvalidModelStateResponseFactory = context =>
+                    {
+                        var result = new BadRequestObjectResult(context.ModelState);
+                        // always return errors as json objects
+                        // we could allow more types here, but we do not want raw for instance
+                        result.ContentTypes.Add(MediaTypeNames.Application.Json);
+
+                        return result;
+                    };
+                });
+
+            services.AddHttpContextAccessor();
+
+            // we change the name of the jwt scheme as okta uses the same name and does not allow us to reconfigure it.
+            const string JwtScheme = "JWTBearer";
+            List<string> defaultSchemes = new List<string>();
+
+            AuthenticationBuilder authenticationBuilder = services.AddAuthentication(options =>
+                {
+                    switch (Auth.Method)
+                    {
+                        case AuthMethod.JWTBearer:
+                            options.DefaultAuthenticateScheme = JwtScheme;
+                            options.DefaultChallengeScheme = JwtScheme;
+                            break;
+                        case AuthMethod.Okta:
+                            options.DefaultAuthenticateScheme = OktaDefaults.ApiAuthenticationScheme;
+                            options.DefaultChallengeScheme = OktaDefaults.ApiAuthenticationScheme;
+                            break;
+                        case AuthMethod.Disabled:
+                            options.DefaultAuthenticateScheme = DisabledAuthenticationHandler.AuthenticateScheme;
+                            options.DefaultChallengeScheme = DisabledAuthenticationHandler.AuthenticateScheme;
+                            break;
+                        default:
+                            throw new ArgumentOutOfRangeException();
+                    }
+                });
+            if (Auth.Method == AuthMethod.JWTBearer)
+            {
+                defaultSchemes.Add(JwtScheme);
+                authenticationBuilder.AddJwtBearer(JwtScheme, options =>
+                {
+                    options.Authority = Auth.JwtAuthority;
+                    options.Audience = Auth.JwtAudience;
+                });
+            }
+
+            if (Auth.Method == AuthMethod.Okta)
+            {
+                defaultSchemes.Add(OktaDefaults.ApiAuthenticationScheme);
+                authenticationBuilder.AddOktaWebApi(new OktaWebApiOptions
+                {
+                    OktaDomain = Auth.OktaDomain,
+                    AuthorizationServerId = Auth.AuthorizationServerId,
+                    Audience = Auth.JwtAudience,
+                });
+            }
+
+            if (Auth.Method == AuthMethod.Disabled)
+            {
+                defaultSchemes.Add(DisabledAuthenticationHandler.AuthenticateScheme);
+                authenticationBuilder.AddTestAuth(options => { });
+            }
+
+            defaultSchemes.Add(ServiceAccountAuthHandler.AuthenticationScheme);
+            authenticationBuilder.AddScheme<ServiceAccountAuthOptions, ServiceAccountAuthHandler>(ServiceAccountAuthHandler.AuthenticationScheme, options => { });
+
+            services.AddAuthorization(options =>
+            {
+                options.AddPolicy(NamespaceAccessRequirement.Name, policy =>
+                {
+                    policy.AuthenticationSchemes = defaultSchemes;
+                    policy.Requirements.Add(new NamespaceAccessRequirement());
+                });
+
+                OnAddAuthorization(options, defaultSchemes);
+            });
+            services.AddSingleton<IAuthorizationHandler, NamespaceAuthorizationHandler>();
+
+            services.Configure<ForwardedHeadersOptions>(options =>
+            {
+                options.ForwardedHeaders = ForwardedHeaders.XForwardedFor | ForwardedHeaders.XForwardedProto;
+            });
+
+
+            services.AddSwaggerGen(settings =>
+            {
+                string? assemblyName = Assembly.GetEntryAssembly()?.GetName().Name;
+                settings.SwaggerDoc("v1", info: new OpenApiInfo
+                {
+                    Title = $"{assemblyName} API",
+                    Contact = new OpenApiContact
+                    {
+                        Name = "Joakim Lindqvist",
+                        Email = "[REDACTED_EMAIL]",
+                    }
+                });
+
+                // Set the comments path for the Swagger JSON and UI.
+                string xmlFile = $"{assemblyName}.xml";
+                string xmlPath = Path.Combine(AppContext.BaseDirectory, xmlFile);
+                if (File.Exists(xmlPath))
+                {
+                    settings.IncludeXmlComments(xmlPath);
+                }
+            });
+
+            OnAddService(services);
+
+            OnAddHealthChecks(services);
+        }
+
+        private void OnAddHealthChecks(IServiceCollection services)
+        {
+            IHealthChecksBuilder healthChecks = services.AddHealthChecks()
+                .AddCheck("self", () => HealthCheckResult.Healthy(), tags: new[] {"self"});
+            OnAddHealthChecks(services, healthChecks);
+
+            healthChecks.AddDatadogPublisher("jupiter.healthchecks");
+        }
+
+        /// <summary>
+        /// Register health checks for individual services
+        /// </summary>
+        /// <remarks>Use the self tag for checks if the service is running while the services tag can be used for any dependencies which needs to work</remarks>
+        /// <param name="services">DI service injector</param>
+        /// <param name="healthChecks">A already configured builder that you can add more checks to</param>
+        protected abstract void OnAddHealthChecks(IServiceCollection services, IHealthChecksBuilder healthChecks);
+
+        protected abstract void OnAddAuthorization(AuthorizationOptions authorizationOptions, List<string> defaultSchemes);
+
+        protected virtual void OnJsonOptions(MvcNewtonsoftJsonOptions options)
+        {
+        }
+
+        protected virtual void OnAddControllers(MvcOptions options)
+        {
+        }
+
+        protected abstract void OnAddService(IServiceCollection services);
+
+        // This method gets called by the runtime. Use this method to configure the HTTP request pipeline.
+        public void Configure(IApplicationBuilder app, IWebHostEnvironment env)
+        {
+            JupiterSettings jupiterSettings = app.ApplicationServices.GetService<IOptionsMonitor<JupiterSettings>>()!.CurrentValue;
+            
+            bool IsHttpPort(HttpContext httpContext) => !IsHighPerfHttpPort(httpContext);
+            bool IsHighPerfHttpPort(HttpContext httpContext) => jupiterSettings.DisableAuthOnPorts.Any(port => port == httpContext.Connection.LocalPort);
+            
+            if (jupiterSettings.ShowPII)
+            {
+                _logger.Error("Personally Identifiable information being shown. This should not be generally enabled in prod.");
+
+                // do not hide personal information during development
+                Microsoft.IdentityModel.Logging.IdentityModelEventSource.ShowPII = true;
+            }
+
+            if (jupiterSettings.DisableAuthOnPorts.Any())
+            {
+                _logger.Information("Auth is disabled on port(s) {Ports}", jupiterSettings.DisableAuthOnPorts);
+            }
+
+            app.MapWhen(IsHttpPort, httpApp =>
+            {
+                ConfigureMiddlewares(false, jupiterSettings, httpApp, env);
+            });
+            
+            app.MapWhen(IsHighPerfHttpPort, highPerfApp =>
+            {
+                ConfigureMiddlewares(true, jupiterSettings, highPerfApp, env);
+            });
+        }
+
+        private void ConfigureMiddlewares(bool isHighPerf, JupiterSettings jupiterSettings, IApplicationBuilder app, IWebHostEnvironment env)
+        {
+            // enable use of forwarding headers as we expect a reverse proxy to be running in front of us
+            app.UseMiddleware<DatadogTraceMiddleware>("ForwardedHeaders");
+            app.UseForwardedHeaders();
+
+            if (!isHighPerf && jupiterSettings.UseRequestLogging)
+            {
+                app.UseMiddleware<DatadogTraceMiddleware>("RequestLogging");
+                app.UseSerilogRequestLogging();
+            }
+
+            OnConfigureAppEarly(app, env);
+                
+            if (!isHighPerf && env.IsDevelopment() && UseDeveloperExceptionPage)
+            {
+                app.UseDeveloperExceptionPage();
+            }
+            else
+            { 
+                app.UseExceptionHandler("/error");
+            }
+                
+            app.UseMiddleware<DatadogTraceMiddleware>("Routing");
+            app.UseRouting();
+
+            if (!isHighPerf)
+            {
+                app.UseMiddleware<DatadogTraceMiddleware>("Authentication");
+                app.UseAuthentication();
+                app.UseMiddleware<DatadogTraceMiddleware>("Authorization");
+                app.UseAuthorization();
+            }
+            
+            app.UseMiddleware<DatadogTraceMiddleware>("Endpoints");
+            app.UseEndpoints(endpoints =>
+            {
+                bool PassAllChecks(HealthCheckRegistration check) => true;
+
+                // Ready checks in Kubernetes is to verify that the service is working, if this returns false the app will not get any traffic (load balancer ignores it)
+                endpoints.MapHealthChecks("/health/ready", options: new HealthCheckOptions()
+                {
+                    Predicate = jupiterSettings.DisableHealthChecks ? PassAllChecks : (check) => check.Tags.Contains("self"),
+                });
+
+                // Live checks in Kubernetes to see if the pod is working as it should, if this returns false the entire pod is killed
+                endpoints.MapHealthChecks("/health/live", options: new HealthCheckOptions()
+                {
+                    Predicate = jupiterSettings.DisableHealthChecks ? PassAllChecks : (check) => check.Tags.Contains("services"),
+                });
+
+                OnUseEndpoints(env, endpoints);
+                
+                endpoints.MapControllers();
+            });
+
+            if (!isHighPerf && jupiterSettings.HostSwaggerDocumentation)
+            {
+                app.UseMiddleware<DatadogTraceMiddleware>("Swagger");
+                app.UseSwagger();
+                app.UseReDoc(options => { options.SpecUrl = "/swagger/v1/swagger.json"; });
+            }
+
+            OnConfigureApp(app, env);
+        }
+
+        protected virtual void OnUseEndpoints(IWebHostEnvironment env, IEndpointRouteBuilder endpoints)
+        {
+            
+        }
+
+        public virtual bool UseDeveloperExceptionPage { get; } = true;
+
+        protected virtual void OnConfigureAppEarly(IApplicationBuilder app, IWebHostEnvironment env)
+        {
+        }
+
+        protected virtual void OnConfigureApp(IApplicationBuilder app, IWebHostEnvironment env)
+        {
+            
+        }
+    }
+
+    public class MvcJsonOptionsWrapper : IConfigureOptions<MvcNewtonsoftJsonOptions>
+    {
+        IServiceProvider ServiceProvider;
+
+        public MvcJsonOptionsWrapper(IServiceProvider serviceProvider)
+        {
+            ServiceProvider = serviceProvider;
+        }
+        public void Configure(MvcNewtonsoftJsonOptions options)
+        {
+            options.SerializerSettings.ContractResolver = new FieldFilteringResolver(ServiceProvider);
+        }
+    }
+
+    public class FieldFilteringResolver : DefaultContractResolver
+    {
+        private IHttpContextAccessor _httpContextAccessor;
+
+        public FieldFilteringResolver(IServiceProvider sp)
+        {
+            _httpContextAccessor = sp.GetRequiredService<IHttpContextAccessor>();
+
+            NamingStrategy = new CamelCaseNamingStrategy(false, true);
+        }
+
+        protected override JsonProperty CreateProperty(MemberInfo member, MemberSerialization memberSerialization)
+        {
+
+            JsonProperty property = base.CreateProperty(member, memberSerialization);
+
+            property.ShouldSerialize = o =>
+            {
+                HttpContext? httpContext = _httpContextAccessor.HttpContext;
+
+                if (httpContext == null)
+                    return true;
+
+                // if no fields are being filtered we should serialize the property
+                if (!httpContext.Request.Query.ContainsKey("fields"))
+                    return true;
+
+                StringValues fields = httpContext.Request.Query["fields"];
+                bool ignore = true;
+                foreach (string field in fields)
+                {
+                    // a empty field to filter for is considered a match for everything as fields= should be the same as just omitting fields
+                    if (string.IsNullOrEmpty(field))
+                    {
+                        return true;
+                    }
+                    if (string.Equals(field, property.PropertyName, StringComparison.InvariantCultureIgnoreCase))
+                    {
+                        ignore = false;
+                    }
+                }
+
+                return !ignore;
+            };
+
+            return property;
+        }
+    }
+
+
+    public enum AuthMethod
+    {
+        JWTBearer,
+        Okta,
+        Disabled
+    };
+
+    public class AuthSettings
+    {
+        [Required] 
+        public AuthMethod Method { get; set; } = AuthMethod.JWTBearer;
+
+
+        [Required]
+        public string OktaDomain { get; set; } = "";
+        public string AuthorizationServerId { get; set; } = OktaWebOptions.DefaultAuthorizationServerId;
+
+        [Required]
+        public string JwtAuthority { get; set; } = "";
+
+        [Required]
+        public string JwtAudience { get; set; } = "";
+
+    }
+
+    public class JupiterSettings
+    {
+        // enable to unhide potentially personal information, see https://aka.ms/IdentityModel/PII
+        public bool ShowPII = false;
+        public bool DisableHealthChecks { get; set; } = false;
+        public bool HostSwaggerDocumentation { get; set; } = true;
+
+        public List<int> DisableAuthOnPorts { get; set; } = new();
+
+        // Enable to echo every request to the log file, usually this is more efficiently done on the load balancer
+        public bool UseRequestLogging { get; set; } = false;
+
+        //public Dictionary<string, int> DisableAuthOnPorts { get; set; } = new ();
+    }
+
+    public class DatadogTraceMiddleware
+    {
+        private readonly RequestDelegate _next;
+        private readonly string _scopeName;
+
+        public DatadogTraceMiddleware(RequestDelegate next, string scopeName)
+        {
+            _next = next;
+            _scopeName = scopeName;
+        }
+
+        public async Task InvokeAsync(HttpContext httpContext)
+        {
+            using Scope _ = Tracer.Instance.StartActive(_scopeName);
+
+            //Move to next delegate/middleware in the pipleline
+            await _next.Invoke(httpContext);
+        }
+    }
+}
diff --git a/Engine/Source/Programs/Horde/HordeStorage/Jupiter.Common/CompactBinaryInputFormatter.cs b/Engine/Source/Programs/Horde/HordeStorage/Jupiter.Common/CompactBinaryInputFormatter.cs
new file mode 100644
--- /dev/null
+++ b/Engine/Source/Programs/Horde/HordeStorage/Jupiter.Common/CompactBinaryInputFormatter.cs
@@ -0,0 +1,25 @@
+// Copyright Epic Games, Inc. All Rights Reserved.
+
+using System.Threading.Tasks;
+using Jupiter.Implementation;
+using Microsoft.AspNetCore.Mvc.Formatters;
+using Microsoft.Net.Http.Headers;
+
+namespace Jupiter
+{
+    public class CompactBinaryInputFormatter : InputFormatter
+    {
+        public CompactBinaryInputFormatter()
+        {
+            SupportedMediaTypes.Add(MediaTypeHeaderValue.Parse(CustomMediaTypeNames.UnrealCompactBinary));
+        }
+
+        public override async Task<InputFormatterResult> ReadRequestBodyAsync(InputFormatterContext context)
+        {
+            byte[] body = await RequestUtil.ReadRawBody(context.HttpContext.Request);
+            CompactBinaryObject cb = CompactBinaryObject.Load(body);
+            
+            return InputFormatterResult.Success(cb.ToPoco(context.ModelType));
+        }
+    }
+}
diff --git a/Engine/Source/Programs/Horde/HordeStorage/Jupiter.Common/CompactBinaryOutputFormatter.cs b/Engine/Source/Programs/Horde/HordeStorage/Jupiter.Common/CompactBinaryOutputFormatter.cs
new file mode 100644
--- /dev/null
+++ b/Engine/Source/Programs/Horde/HordeStorage/Jupiter.Common/CompactBinaryOutputFormatter.cs
@@ -0,0 +1,124 @@
+﻿// Copyright Epic Games, Inc. All Rights Reserved.
+
+using System;
+using System.Collections.Generic;
+using System.Linq;
+using System.Reflection;
+using System.Threading.Tasks;
+using Jupiter.Implementation;
+using Microsoft.AspNetCore.Http;
+using Microsoft.AspNetCore.Mvc.Formatters;
+using Microsoft.Net.Http.Headers;
+using Serilog;
+
+namespace Jupiter
+{
+    public class CompactBinaryOutputFormatter : OutputFormatter
+    {
+        private readonly ILogger _logger = Log.ForContext<CompactBinaryOutputFormatter>();
+
+        public CompactBinaryOutputFormatter()
+        {
+            SupportedMediaTypes.Add(MediaTypeHeaderValue.Parse(CustomMediaTypeNames.UnrealCompactBinary));
+        }
+
+        public override async Task WriteResponseBodyAsync(OutputFormatterWriteContext context)
+        {
+            if (context.Object == null)
+                return;
+
+            CompactBinaryWriter compactBinaryWriter = ConvertToCompactBinary(context.Object);
+            
+            HttpResponse response = context.HttpContext.Response;
+            byte[] value = compactBinaryWriter.Save();
+
+            response.ContentLength = value.LongLength;
+            await response.Body.WriteAsync(value);
+        }
+
+        private CompactBinaryWriter ConvertToCompactBinary(object o)
+        {
+            CompactBinaryWriter writer = new CompactBinaryWriter();
+            writer.BeginObject();
+            
+            foreach (PropertyInfo propertyInfo in o.GetType().GetProperties())
+            {
+                // convert to camelCase, looks better then PascalCase which is the default for a C# object
+                string name = propertyInfo.Name.Substring(0,1).ToLower() + propertyInfo.Name.Substring(1);
+                Type propertyType = propertyInfo.PropertyType;
+                Type? innerType = Nullable.GetUnderlyingType(propertyType);
+                if (innerType != null)
+                {
+                    propertyType = innerType;
+                }
+
+                Type[] propertyInterfaces = propertyType.GetInterfaces();
+
+                object? value = propertyInfo.GetValue(o);
+                if (value == null)
+                {
+                    writer.AddNull(name);
+                }
+                else if (propertyType == typeof(BlobIdentifier))
+                {
+                    BlobIdentifier blob = (BlobIdentifier) value!;
+                    writer.AddBinaryAttachment(blob, name);
+                }
+                else if (propertyType == typeof(string))
+                {
+                    string s = (string) value!;
+                    writer.AddString(s, name);
+                }
+                else if (propertyType == typeof(DateTime))
+                {
+                    DateTime dt = (DateTime) value!; 
+                    writer.AddDateTime(dt, name);
+                }
+                else if (propertyType == typeof(int))
+                {
+                    int i = (int) value!;
+                    // writer.AddInteger(i, name);
+                }
+                else if (propertyType.IsArray)
+                {
+                    Type elementType = propertyType.GetElementType()!;
+
+                    if (elementType == typeof(BlobIdentifier))
+                    {
+                        BlobIdentifier[] blobs = (BlobIdentifier[])value;
+                        writer.AddUniformArray(blobs, CompactBinaryFieldType.BinaryAttachment, name);
+                    }
+                    else if (elementType == typeof(string))
+                    {
+                        string[] strings = (string[])value;
+                        writer.AddUniformArray(strings, CompactBinaryFieldType.String, name);
+                    }
+                    else if (elementType == typeof(byte))
+                    {
+                        byte[] data = (byte[])value;
+                        writer.AddBinary(data, name);
+                    }
+                    else
+                    {
+                        throw new NotImplementedException($"Type {elementType} not support for arrays when converting to compact binary for field \"{name}\"");
+                    }
+                }
+                else if (propertyType.IsGenericType && propertyType.GetGenericTypeDefinition() == typeof(IDictionary<,>))
+                {
+                    // TODO: Serialize dictionaries as objects
+                }
+                else if (propertyInfo.PropertyType.IsClass)
+                {
+                    //writer.AddObject(name, ConvertToCompactBinary(value!));
+                }
+                else
+                {
+                    throw new NotImplementedException($"Unsupported type in compact binary conversion: {propertyType.Name} for field \"{name}\"");
+                }
+            }
+
+            writer.EndObject();
+            return writer;
+        }
+    }
+}
diff --git a/Engine/Source/Programs/Horde/HordeStorage/Jupiter.Common/CustomMediaTypeNames.cs b/Engine/Source/Programs/Horde/HordeStorage/Jupiter.Common/CustomMediaTypeNames.cs
new file mode 100644
--- /dev/null
+++ b/Engine/Source/Programs/Horde/HordeStorage/Jupiter.Common/CustomMediaTypeNames.cs
@@ -0,0 +1,9 @@
+﻿namespace Jupiter
+{
+    public class CustomMediaTypeNames
+    {
+        public const string UnrealCompactBinary = "application/x-ue-cb";
+        public const string UnrealCompressedBuffer = "application/x-ue-comp";
+    }
+
+}
diff --git a/Engine/Source/Programs/Horde/HordeStorage/Jupiter.Common/Implementation/BatchWriter.cs b/Engine/Source/Programs/Horde/HordeStorage/Jupiter.Common/Implementation/BatchWriter.cs
new file mode 100644
--- /dev/null
+++ b/Engine/Source/Programs/Horde/HordeStorage/Jupiter.Common/Implementation/BatchWriter.cs
@@ -0,0 +1,90 @@
+// Copyright Epic Games, Inc. All Rights Reserved.
+
+using System;
+using System.Collections.Generic;
+using System.IO;
+using System.Text;
+using System.Threading;
+using System.Threading.Tasks;
+using Datadog.Trace;
+using Serilog;
+
+namespace Jupiter.Implementation
+{
+    public class BatchWriter
+    {
+        public enum OpState : byte
+        {
+            OK = 0,
+            Failed = 1,
+            NotFound = 2,
+            Exists = 3 // special case for HEAD checks, same as OK but indicates that there will be no payload
+        };
+
+        private readonly ILogger _logger = Log.ForContext<BatchWriter>();
+
+        public async Task WriteToStream(Stream outStream, NamespaceId ns, List<Tuple<OpVerb, BucketId, KeyId>> ops, Func<OpVerb, NamespaceId, BucketId, KeyId, string[], Task<Tuple<ContentHash?, BlobContents?, OpState>>> fetchAction)
+        {
+            await outStream.WriteAsync(Encoding.ASCII.GetBytes("JPTR"));
+            await outStream.WriteAsync(BitConverter.GetBytes((uint) ops.Count));
+
+            byte[] entryHeaderBytes = Encoding.ASCII.GetBytes("JPEE");
+
+            SemaphoreSlim streamLock = new SemaphoreSlim(1, 1);
+
+            foreach ((OpVerb opVerb, BucketId bucket, KeyId key) in ops)
+            {
+                string opTypeName = opVerb.ToString();
+                using Scope scope = Tracer.Instance.StartActive($"batch.{opTypeName}");
+                scope.Span.ResourceName = $"{ns}.{bucket}.{key}";
+
+                (ContentHash? contentHash, BlobContents? blob, OpState opState) = await fetchAction(opVerb, ns, bucket, key, new [] {"blob", "contentHash"});
+
+                // update the scope with the state of the operation to allow us to track how often each request in the batch succeeds or fails
+                scope.Span.SetTag("Success", (opState is OpState.OK or OpState.Exists).ToString());
+                scope.Span.SetTag("Error", (opState == OpState.Failed).ToString());
+
+                try
+                {
+                    // make sure only one object writes to the result stream at a time
+                    await streamLock.WaitAsync();
+
+                    // Write custom header for each element, this header includes the cache key and its hash
+                    // start with a preamble to let us verify that the stream is correct.
+                    await outStream.WriteAsync(entryHeaderBytes);
+
+                    string opName = $"{ns}.{bucket}.{key}";
+                    // we null terminate the strings to be easier to use
+                    await outStream.WriteAsync(Encoding.ASCII.GetBytes(opName + "\0"));
+
+                    // write out the state byte
+                    await outStream.WriteAsync(new byte[] {(byte) opState});
+                    
+                    // if we have a failure we do not know if the hash data or blob is valid so we only write the state code
+                    // also for head checks we do not write anything more
+                    if (opState != OpState.OK || opState == OpState.Exists || contentHash == null || blob == null)
+                        continue;
+
+                    if (blob.Length == 0)
+                    {
+                        _logger.Error("0 byte blob encountered for hash {ContentHash} state was {OpState}. Skipping", contentHash.ToString(), opState);
+                    }
+
+                    await outStream.WriteAsync(contentHash.HashData);
+                    await outStream.WriteAsync(BitConverter.GetBytes((ulong) blob.Length));
+                    await blob.Stream.CopyToAsync(outStream);
+                }
+                finally
+                {
+                    streamLock.Release();
+                }
+            }
+        }
+
+        public enum OpVerb
+        {
+            GET,
+            HEAD
+        }
+    }
+}
diff --git a/Engine/Source/Programs/Horde/HordeStorage/Jupiter.Common/Implementation/BlobContents.cs b/Engine/Source/Programs/Horde/HordeStorage/Jupiter.Common/Implementation/BlobContents.cs
new file mode 100644
--- /dev/null
+++ b/Engine/Source/Programs/Horde/HordeStorage/Jupiter.Common/Implementation/BlobContents.cs
@@ -0,0 +1,35 @@
+﻿using System;
+using System.IO;
+using System.Threading.Tasks;
+
+namespace Jupiter.Implementation
+{
+    public class BlobContents : IDisposable, IAsyncDisposable
+    {
+        public BlobContents(Stream stream, long length)
+        {
+            Stream = stream;
+            Length = length;
+        }
+
+        public BlobContents(byte[] payload)
+        {
+            Stream = new MemoryStream(payload);
+            Length = payload.LongLength;
+        }
+
+        public Stream Stream { get; }
+        public long Length {get;}
+
+        public void Dispose()
+        {
+            Stream?.Dispose();
+        }
+
+        public ValueTask DisposeAsync()
+        {
+            Stream?.Dispose();
+            return ValueTask.CompletedTask;
+        }
+    }
+}
diff --git a/Engine/Source/Programs/Horde/HordeStorage/Jupiter.Common/Implementation/BlobIdentifier.cs b/Engine/Source/Programs/Horde/HordeStorage/Jupiter.Common/Implementation/BlobIdentifier.cs
new file mode 100644
--- /dev/null
+++ b/Engine/Source/Programs/Horde/HordeStorage/Jupiter.Common/Implementation/BlobIdentifier.cs
@@ -0,0 +1,176 @@
+﻿// Copyright Epic Games, Inc. All Rights Reserved.
+
+using System;
+using System.ComponentModel;
+using System.Globalization;
+using System.IO;
+using System.Threading.Tasks;
+using Blake3;
+using Newtonsoft.Json;
+
+namespace Jupiter.Implementation
+{
+    [JsonConverter(typeof(BlobIdentifierConverter))]
+    [TypeConverter(typeof(BlobIdentifierTypeConverter))]
+    public class BlobIdentifier : ContentHash,  IEquatable<BlobIdentifier>
+    {
+        // multi thread the hashing for blobs larger then this size
+        private const int MultiThreadedSize = 1_000_000;
+        private string? _stringIdentifier;
+
+        public BlobIdentifier(byte[] identifier) : base(identifier)
+        {
+        }
+
+        [JsonConstructor]
+        public BlobIdentifier(string identifier) : base(identifier)
+        {
+
+        }
+
+        public override int GetHashCode()
+        {
+            return Comparer.GetHashCode(Identifier);
+        }
+
+        public bool Equals(BlobIdentifier? other)
+        {
+            if (other == null)
+                return false;
+
+            return Comparer.Equals(Identifier, other.Identifier);
+        }
+
+        public override bool Equals(object? obj)
+        {
+            if (ReferenceEquals(null, obj))
+            {
+                return false;
+            }
+
+            if (ReferenceEquals(this, obj))
+            {
+                return true;
+            }
+
+            if (obj.GetType() != this.GetType())
+            {
+                return false;
+            }
+
+            return Equals((BlobIdentifier) obj);
+        }
+
+        public override string ToString()
+        {
+            if (_stringIdentifier == null)
+            {
+                _stringIdentifier = StringUtils.FormatAsHexString(Identifier);
+            }
+
+            return _stringIdentifier;
+        }
+
+        public new static BlobIdentifier FromBlob(byte[] blobMemory)
+        {
+            Hash blake3Hash;
+            if (blobMemory.Length < MultiThreadedSize)
+            {
+                using Hasher hasher = Hasher.New();
+                hasher.Update(blobMemory);
+                blake3Hash = hasher.Finalize();
+            }
+            else
+            {
+                using Hasher hasher = Hasher.New();
+                hasher.UpdateWithJoin(blobMemory);
+                blake3Hash = hasher.Finalize();
+            }
+            
+            // we only keep the first 20 bytes of the Blake3 hash
+            Span<byte> hash = blake3Hash.AsSpan().Slice(0, 20);
+            return new BlobIdentifier(hash.ToArray());
+        }
+
+        public static BlobIdentifier FromBlob(in Memory<byte> blobMemory)
+        {
+            Hash blake3Hash;
+            if (blobMemory.Length < MultiThreadedSize)
+            {
+                using Hasher hasher = Hasher.New();
+                hasher.Update(blobMemory.Span);
+                blake3Hash = hasher.Finalize();
+            }
+            else
+            {
+                using Hasher hasher = Hasher.New();
+                hasher.UpdateWithJoin(blobMemory.Span);
+                blake3Hash = hasher.Finalize();
+            }
+
+            // we only keep the first 20 bytes of the Blake3 hash
+            Span<byte> hash = blake3Hash.AsSpan().Slice(0, 20);
+            return new BlobIdentifier(hash.ToArray());
+        }
+
+        public static BlobIdentifier FromContentHash(ContentHash testObjectHash)
+        {
+            return new BlobIdentifier(testObjectHash.HashData);
+        }
+
+        public static async Task<BlobIdentifier> FromStream(Stream stream)
+        {
+            using Hasher hasher = Hasher.New();
+            const int bufferSize = 1024 * 1024 * 5;
+            byte[] buffer = new byte[bufferSize];
+            int read = await stream.ReadAsync(buffer, 0, buffer.Length);
+            while (read > 0)
+            {
+                hasher.UpdateWithJoin(new ReadOnlySpan<byte>(buffer, 0, read));
+                read = await stream.ReadAsync(buffer, 0, buffer.Length);
+            }
+            Hash blake3Hash = hasher.Finalize();
+
+            // we only keep the first 20 bytes of the Blake3 hash
+            byte[] hash = blake3Hash.AsSpan().Slice(0, 20).ToArray();
+            return new BlobIdentifier(hash);
+        }
+    }
+
+    public class BlobIdentifierTypeConverter : TypeConverter
+    {
+        public override bool CanConvertFrom(ITypeDescriptorContext? context, Type sourceType)
+        {  
+            if (sourceType == typeof(string))  
+            {  
+                return true;
+            }  
+            return base.CanConvertFrom(context, sourceType);
+        }  
+  
+        public override object? ConvertFrom(ITypeDescriptorContext? context, CultureInfo? culture, object value)  
+        {
+            if (value is string)
+            {
+                return new BlobIdentifier((string) value);
+            }
+
+            return base.ConvertFrom(context, culture, value);  
+        }  
+    }
+
+    public class BlobIdentifierConverter : JsonConverter<BlobIdentifier>
+    {
+        public override void WriteJson(JsonWriter writer, BlobIdentifier? value, JsonSerializer serializer)
+        {
+            writer.WriteValue(value!.ToString());
+        }
+
+        public override BlobIdentifier ReadJson(JsonReader reader, Type objectType, BlobIdentifier? existingValue, bool hasExistingValue, Newtonsoft.Json.JsonSerializer serializer)
+        {
+            string? s = (string?)reader.Value;
+
+            return new BlobIdentifier(s!);
+        }
+    }
+}
diff --git a/Engine/Source/Programs/Horde/HordeStorage/Jupiter.Common/Implementation/BucketId.cs b/Engine/Source/Programs/Horde/HordeStorage/Jupiter.Common/Implementation/BucketId.cs
new file mode 100644
--- /dev/null
+++ b/Engine/Source/Programs/Horde/HordeStorage/Jupiter.Common/Implementation/BucketId.cs
@@ -0,0 +1,124 @@
+﻿using System;
+using System.ComponentModel;
+using System.Globalization;
+using Newtonsoft.Json;
+
+namespace Jupiter.Implementation
+{
+    [JsonConverter(typeof(BucketIdJsonConverter))]
+    [TypeConverter(typeof(BucketIdTypeConverter))]
+
+    public readonly struct BucketId: IEquatable<BucketId>
+    {
+        public BucketId(string bucketId)
+        {
+             _text = bucketId;
+ 
+            if (_text.Length == 0)
+            {
+				throw new ArgumentException("Buckets must have at least one character.");
+            }
+ 
+            const int MaxLength = 64;
+            if (_text.Length > MaxLength)
+            {
+                throw new ArgumentException($"Buckets may not be longer than {MaxLength} characters");
+            }
+ 
+            for (int Idx = 0; Idx < _text.Length; Idx++)
+            {
+                if (!IsValidCharacter(_text[Idx]))
+                {
+                    throw new ArgumentException($"{_text} is not a valid bucket id");
+                }
+            }
+        }
+
+        static bool IsValidCharacter(char c)
+        {
+            if (c >= 'a' && c <= 'z')
+            {
+                return true;
+            }
+            if (c >= '0' && c <= '9')
+            {
+                return true;
+            }
+            if (c == '-' || c == '_' || c == '.')
+            {
+                return true;
+            }
+            return false;
+        }
+
+        readonly string _text;
+
+        public bool Equals(BucketId other)
+        {
+            return string.Equals(_text , other._text, StringComparison.Ordinal);
+        }
+
+        public override bool Equals(object? obj)
+        {
+            return obj is BucketId other && Equals(other);
+        }
+
+        public override int GetHashCode()
+        {
+            return _text.GetHashCode();
+        }
+
+        public override string ToString()
+        {
+            return _text;
+        }
+
+        public static bool operator ==(BucketId left, BucketId right)
+        {
+            return left.Equals(right);
+        }
+ 
+        public static bool operator !=(BucketId left, BucketId right)
+        {
+            return !left.Equals(right);
+        }
+    }
+
+    public class BucketIdJsonConverter: JsonConverter<BucketId>
+    {
+        public override void WriteJson(JsonWriter writer, BucketId value, JsonSerializer serializer)
+        {
+            writer.WriteValue(value!.ToString());
+        }
+
+        public override BucketId ReadJson(JsonReader reader, Type objectType, BucketId existingValue, bool hasExistingValue, Newtonsoft.Json.JsonSerializer serializer)
+        {
+            string? s = (string?)reader.Value;
+
+            return new BucketId(s!);
+        }
+    }
+
+    public sealed class BucketIdTypeConverter : TypeConverter
+    {
+        /// <inheritdoc/>
+        public override bool CanConvertFrom(ITypeDescriptorContext? context, Type sourceType)
+        {
+            if (sourceType == typeof(string))
+                return true;
+
+            return base.CanConvertFrom(context, sourceType);
+        }
+ 
+        public override object? ConvertFrom(ITypeDescriptorContext? context, CultureInfo? culture, object value)  
+        {
+            if (value is string s)
+            {
+                return new BucketId(s);
+            }
+
+            return base.ConvertFrom(context, culture, value);  
+        }
+    }
+}
+
diff --git a/Engine/Source/Programs/Horde/HordeStorage/Jupiter.Common/Implementation/CallistoReader.cs b/Engine/Source/Programs/Horde/HordeStorage/Jupiter.Common/Implementation/CallistoReader.cs
new file mode 100644
--- /dev/null
+++ b/Engine/Source/Programs/Horde/HordeStorage/Jupiter.Common/Implementation/CallistoReader.cs
@@ -0,0 +1,210 @@
+﻿// Copyright Epic Games, Inc. All Rights Reserved.
+
+using System;
+using System.Collections.Generic;
+using System.Linq;
+using System.Net;
+using System.Runtime.CompilerServices;
+using System.Threading;
+using System.Threading.Tasks;
+using Newtonsoft.Json;
+using Newtonsoft.Json.Linq;
+using RestSharp;
+
+namespace Jupiter.Implementation
+{
+    public class CallistoReader
+    {
+        public class CallistoGetResponse
+        {
+            [JsonConstructor]
+            public CallistoGetResponse(List<TransactionEvent> events, Guid generation, long currentOffset)
+            {
+                Events = events;
+                Generation = generation;
+                CurrentOffset = currentOffset;
+            }
+
+            public CallistoGetResponse()
+            {
+                Events = null!;
+                Generation = Guid.Empty;
+                CurrentOffset = 0;
+            }
+
+            public Guid Generation { get; set; }
+            public long CurrentOffset { get; }
+            public List<TransactionEvent> Events { get; set; }
+        }
+
+        private readonly NamespaceId _ns;
+        private IRestClient _client;
+
+        public CallistoReader(IRestClient callistoClient, NamespaceId ns)
+        {
+            _ns = ns;
+            _client = callistoClient;
+        }
+
+        public async IAsyncEnumerable<TransactionEvent> GetOps(long startingOffset = 0, Guid? expectedReplicationGeneration = null, string? siteFilter = null, int? maxOffsetsAttempted = null, OpsEnumerationState? enumerationState = null, [EnumeratorCancellation] CancellationToken cancellationToken = default(CancellationToken))
+        {
+            if (expectedReplicationGeneration == null)
+            {
+                Transactions initialInfo = await FetchTransactions(0, countOfEventsToFetch: 1, siteFilter: null, replicatingGeneration: null, maxOffsetsAttempted: maxOffsetsAttempted, cancellationToken);
+                expectedReplicationGeneration = initialInfo.ReplicatingGeneration;
+            }
+
+            if (cancellationToken.IsCancellationRequested)
+                yield break;
+
+            bool hadNewEvents;
+            long currentOffset = startingOffset;
+            do
+            {
+                Transactions currentBatch = await FetchTransactions(currentOffset, countOfEventsToFetch: 100, siteFilter, expectedReplicationGeneration, maxOffsetsAttempted: maxOffsetsAttempted, cancellationToken);
+
+                hadNewEvents = currentBatch.Events?.Any() ?? false;
+                currentOffset = currentBatch.LastOffset;
+
+                if (enumerationState != null)
+                {
+                    enumerationState.ReplicatingGeneration = currentBatch.ReplicatingGeneration;
+                    enumerationState.LastOffset = currentOffset;
+                }
+
+                if (currentBatch.Events == null)
+                    yield break;
+
+                foreach (TransactionEvent transactionEvent in currentBatch.Events)
+                {
+                    yield return transactionEvent;
+                }
+
+
+            } while (hadNewEvents);
+        }
+
+        private async Task<Transactions> FetchTransactions(long startingOffset, int countOfEventsToFetch, string? siteFilter, Guid? replicatingGeneration, int? maxOffsetsAttempted, CancellationToken replicationToken)
+        {
+            const int maxAttempts = 3;
+            CallistoGetResponse? callistoGetResponse = null;
+            IRestResponse<CallistoGetResponse>? response = null;
+            for (int i = 0; i < maxAttempts; i++)
+            {
+                RestRequest request = new RestRequest("api/v1/t/{ns}/{offset}");
+                request.AddUrlSegment("ns", _ns);
+                request.AddUrlSegment("offset", startingOffset);
+
+                request.AddQueryParameter("count", countOfEventsToFetch.ToString());
+
+                if (siteFilter != null)
+                    request.AddQueryParameter("notSeenAtSite", siteFilter);
+
+                if (replicatingGeneration != null)
+                    request.AddQueryParameter("expectedGeneration", replicatingGeneration.ToString()!);
+
+                if (maxOffsetsAttempted.HasValue)
+                    request.AddQueryParameter("maxOffsetsAttempted", maxOffsetsAttempted.Value.ToString());
+
+                response = await _client.ExecuteGetAsync<CallistoGetResponse>(request, replicationToken);
+                if (response.StatusCode == HttpStatusCode.BadRequest)
+                {
+                    JObject badRequestObj = JObject.Parse(response.Content);
+                    if (badRequestObj.ContainsKey("type") && badRequestObj["type"]!.Value<string>() == "http://jupiter.epicgames.com/callisto/newGeneration")
+                    {
+                        string errorMsg = badRequestObj["title"]!.Value<string>()!;
+                        throw new ReplicatingGenerationChangedException($"Generation mismatch when replicating remote callisto. Received error message: {errorMsg}");
+                    }
+
+                    if (badRequestObj.ContainsKey("type") && badRequestObj["type"]!.Value<string>() ==
+                        "http://jupiter.epicgames.com/callisto/unknownNamespace")
+                    {
+                        throw new UnknownNamespaceException($"Namespace {_ns} was not known by Callisto.");
+                    }
+
+                    if (badRequestObj.ContainsKey("type") && badRequestObj["type"]!.Value<string>() ==
+                        "http://jupiter.epicgames.com/callisto/transactionLogOffsetMismatch")
+                    {
+                        // increase offsets by a order of magnitude so we have a better chance of recovering
+                        maxOffsetsAttempted = maxOffsetsAttempted.GetValueOrDefault(100) * 10;
+
+                        // if we reached max attempts throw a exception other retry with more offsets attempted
+                        if (i == maxAttempts - 1)
+                        {
+                            string errorMsg = badRequestObj["title"]!.Value<string>()!;
+                            throw new UnknownNamespaceException($"Transaction log mismatch in {_ns}. Error message {errorMsg}");
+                        }
+                        continue;
+                    }
+
+                    throw new Exception("Unknown bad request response from remote callisto: " + response.Content);
+                }
+
+                if (response.StatusCode == HttpStatusCode.BadGateway)
+                {
+                    // the request timed out, likely because we requested to many objects to reduce the size of the batch
+                    countOfEventsToFetch = countOfEventsToFetch / 2;
+                    continue;
+                }
+
+                if (!response.IsSuccessful)
+                    throw new Exception($"Unsuccessful response from remote callisto. Status Code: {response.StatusCode}. Url: {response.ResponseUri} . Message: {response.ErrorMessage}");
+
+                // we have data so stop retrying
+                callistoGetResponse = response.Data;
+                break;
+            }
+
+            if (callistoGetResponse == null)
+                throw new Exception($"Unsuccessful response from remote callisto, gave up after {maxAttempts} attempts. Status Code: {response!.StatusCode}. Url: {response.ResponseUri} . Message: {response.ErrorMessage}");
+
+            List<TransactionEvent> events = callistoGetResponse.Events;
+            if (events == null)
+                throw new Exception($"Invalid response from remote callisto: {_client} expected objects with events array");
+
+            if (events.Count == 0)
+            {
+                // nothing new
+                return new Transactions
+                {
+                    ReplicatingGeneration = callistoGetResponse.Generation,
+                    LastOffset = callistoGetResponse.CurrentOffset,
+                    Events = null
+                };
+            }
+
+            return new Transactions()
+            {
+                ReplicatingGeneration = callistoGetResponse.Generation,
+                LastOffset = callistoGetResponse.CurrentOffset,
+                Events = events
+            };
+        }
+
+        private class Transactions
+        {
+            public Guid ReplicatingGeneration { get; set; }
+            public long LastOffset { get; set; }
+            public List<TransactionEvent>? Events { get; set; }
+        }
+    }
+    public class OpsEnumerationState
+    {
+        public Guid ReplicatingGeneration { get; set; }
+        public long LastOffset { get; set; }
+    }
+
+    public class UnknownNamespaceException : Exception
+    {
+        public UnknownNamespaceException(string message) : base(message)
+        {
+        }
+    }
+
+    public class ReplicatingGenerationChangedException : Exception
+    {
+        public ReplicatingGenerationChangedException(string message) : base(message)
+        {
+        }
+    }
+}
diff --git a/Engine/Source/Programs/Horde/HordeStorage/Jupiter.Common/Implementation/ClientCredentialOAuthAuthenticator.cs b/Engine/Source/Programs/Horde/HordeStorage/Jupiter.Common/Implementation/ClientCredentialOAuthAuthenticator.cs
new file mode 100644
--- /dev/null
+++ b/Engine/Source/Programs/Horde/HordeStorage/Jupiter.Common/Implementation/ClientCredentialOAuthAuthenticator.cs
@@ -0,0 +1,116 @@
+// Copyright Epic Games, Inc. All Rights Reserved.
+
+using System;
+using System.Net;
+using RestSharp;
+using RestSharp.Authenticators;
+
+namespace Jupiter.Implementation
+{
+    public class ClientCredentialOAuthAuthenticator: IAuthenticator
+    {
+        public ClientCredentialOAuthAuthenticator(string authUrl, string clientId, string clientSecret, string scope)
+        {
+            _authUrl = authUrl;
+            _clientId = clientId;
+            _clientSecret = clientSecret;
+            _scope = scope;
+        }
+
+        private string? _accessToken;
+
+        private readonly string _authUrl;
+        private readonly string _clientId;
+        private readonly string _clientSecret;
+        private readonly string _scope;
+        private DateTime _expiresAt;
+
+        public void Authenticate(IRestClient client, IRestRequest request)
+        {
+            if (string.IsNullOrEmpty(_accessToken) || DateTime.Now > _expiresAt)
+            {
+                PreAuthenticate();
+            }
+
+            request.AddHeader("Authorization", $"Bearer {_accessToken}");
+        }
+
+        public string? Authenticate()
+        {
+            if (string.IsNullOrEmpty(_accessToken) || DateTime.Now > _expiresAt)
+            {
+                PreAuthenticate();
+            }
+
+            return _accessToken;
+        }
+
+        private void PreAuthenticate()
+        {
+            IRestResponse<ClientCredentialsResponse> authRequest = DoAuthenticationRequest();
+            try
+            {
+                if (!authRequest.IsSuccessful)
+                {
+                    throw new AuthenticationFailedException(authRequest.Content);
+                }
+
+                ClientCredentialsResponse result = authRequest.Data;
+                string? accessToken = result?.access_token;
+                if (string.IsNullOrEmpty(accessToken))
+                {
+                    throw new InvalidOperationException("The authentication token received by the server is null or empty. Body received was: " + authRequest.Content);
+                }
+                _accessToken = accessToken;
+                // renew after half the renewal time
+                _expiresAt = DateTime.Now + TimeSpan.FromSeconds((result?.expires_in ?? 3200) / 2.0);
+            }
+            catch (WebException ex)
+            {
+                if (ex.Response is HttpWebResponse response && response.StatusCode != HttpStatusCode.InternalServerError)
+                {
+                    string errorResult = authRequest.Content;
+                    if (errorResult == null)
+                    {
+                        throw;
+                    }
+                    throw new AuthenticationFailedException(errorResult);
+                }
+                throw;
+            }
+
+        }
+
+        private IRestResponse<ClientCredentialsResponse> DoAuthenticationRequest()
+        {
+            Uri requestUri = new Uri(_authUrl);
+            RestClient client = new RestClient(requestUri);
+            RestRequest request = new RestRequest(Method.POST);
+
+            request.AddParameter("grant_type", "client_credentials", ParameterType.GetOrPost);
+            request.AddParameter("client_id", _clientId, ParameterType.GetOrPost);
+            request.AddParameter("client_secret", _clientSecret, ParameterType.GetOrPost);
+            request.AddParameter("scope", _scope, ParameterType.GetOrPost);
+            IRestResponse<ClientCredentialsResponse> response = client.Execute<ClientCredentialsResponse>(request);
+
+            return response;
+        }
+    }
+
+    // ReSharper disable once ClassNeverInstantiated.Global
+    internal class ClientCredentialsResponse
+    {
+        public string? access_token { get; set; }
+
+        public string? token_type { get; set; }
+        
+        public int? expires_in { get; set; }
+
+        public string? scope { get; set; }
+    }
+
+    internal class AuthenticationFailedException : Exception
+    {
+        public AuthenticationFailedException(object errorResult) : base(errorResult.ToString()) { }
+    }
+}
