Moved Horde.Storage from git to p4

diff --git a/Engine/Source/Programs/Horde/HordeStorage/Jupiter.Common/Implementation/CompactBinaryObject.cs b/Engine/Source/Programs/Horde/HordeStorage/Jupiter.Common/Implementation/CompactBinaryObject.cs
new file mode 100644
--- /dev/null
+++ b/Engine/Source/Programs/Horde/HordeStorage/Jupiter.Common/Implementation/CompactBinaryObject.cs
@@ -0,0 +1,1052 @@
+﻿using System;
+using System.Buffers;
+using System.Collections;
+using System.Collections.Generic;
+using System.Dynamic;
+using System.Linq;
+using System.Numerics;
+using System.Reflection;
+using System.Text;
+using System.Text.Json;
+using BindingFlags = System.Reflection.BindingFlags;
+
+namespace Jupiter.Implementation
+{
+    static class CompactBinaryFieldUtils
+    {
+        private const CompactBinaryFieldType SerializedTypeMask = (CompactBinaryFieldType)0b_1001_1111;
+        private const CompactBinaryFieldType TypeMask              = (CompactBinaryFieldType)0b_0001_1111;
+
+        private const CompactBinaryFieldType ObjectMask             = (CompactBinaryFieldType)0b_0001_1110;
+        private const CompactBinaryFieldType ObjectBase             = (CompactBinaryFieldType)0b_0001_0010;
+        
+        private const CompactBinaryFieldType ArrayMask             = (CompactBinaryFieldType)0b_0001_1110;
+        private const CompactBinaryFieldType ArrayBase             = (CompactBinaryFieldType)0b_0000_0100;
+
+        private const CompactBinaryFieldType IntegerMask             = (CompactBinaryFieldType)0b_0001_1110;
+        private const CompactBinaryFieldType IntegerBase             = (CompactBinaryFieldType)0b_0000_1000;
+
+        private const CompactBinaryFieldType FloatMask             = (CompactBinaryFieldType)0b_0001_1100;
+        private const CompactBinaryFieldType FloatBase             = (CompactBinaryFieldType)0b_0000_1000;
+
+        private const CompactBinaryFieldType BoolMask             = (CompactBinaryFieldType)0b_0001_1110;
+        private const CompactBinaryFieldType BoolBase             = (CompactBinaryFieldType)0b_0000_1100;
+
+        private const CompactBinaryFieldType AttachmentMask             = (CompactBinaryFieldType)0b_0001_1110;
+        private const CompactBinaryFieldType AttachmentBase             = (CompactBinaryFieldType)0b_0000_1110;
+
+
+        public static bool IsArray(CompactBinaryFieldType type)
+        {
+            return (type & ArrayMask) == ArrayBase;
+        }
+
+        public static bool IsObject(CompactBinaryFieldType type)
+        {
+            return (type & ObjectMask) == ObjectBase;
+        }
+
+        public static bool IsString(CompactBinaryFieldType type)
+        {
+            return RemoveFlags(type) == CompactBinaryFieldType.String;
+        }
+
+        public static bool IsInteger(CompactBinaryFieldType type)
+        {
+            return (type & IntegerMask) == IntegerBase;
+        }
+
+        public static bool IsFloat(CompactBinaryFieldType type)
+        {
+            return (type & FloatMask) == FloatBase;
+        }
+
+        public static bool IsBool(CompactBinaryFieldType type)
+        {
+            return (type & BoolMask) == BoolBase;
+        }
+        
+        public static bool IsAttachment(CompactBinaryFieldType type)
+        {
+            return (type & AttachmentMask) == AttachmentBase;
+        }
+
+        public static bool IsNull(CompactBinaryFieldType type)
+        {
+            return RemoveFlags(type) == CompactBinaryFieldType.Null;
+        }
+
+        public static bool IsCompactBinaryAttachment(CompactBinaryFieldType type)
+        {
+            return RemoveFlags(type) == CompactBinaryFieldType.CompactBinaryAttachment;
+        }
+
+        public static bool IsBinaryAttachment(CompactBinaryFieldType type)
+        {
+            return RemoveFlags(type) == CompactBinaryFieldType.BinaryAttachment;
+        }
+
+        public static bool IsHash(CompactBinaryFieldType type)
+        {
+            return RemoveFlags(type) == CompactBinaryFieldType.Hash || IsAttachment(type);
+        }
+
+        public static bool IsUuid(CompactBinaryFieldType type)
+        {
+            return RemoveFlags(type) == CompactBinaryFieldType.Uuid;
+        }
+        
+        public static bool IsDateTime(CompactBinaryFieldType type)
+        {
+            return RemoveFlags(type) == CompactBinaryFieldType.DateTime;
+        }
+
+        public static bool IsBinary(CompactBinaryFieldType type)
+        {
+            return RemoveFlags(type) == CompactBinaryFieldType.Binary;
+        }
+
+        public static bool HasUniformFields(CompactBinaryFieldType type)
+        {
+            CompactBinaryFieldType localType = RemoveFlags(type);
+            return localType == CompactBinaryFieldType.UniformObject || localType == CompactBinaryFieldType.UniformArray;
+        }
+
+        public static bool HasFields(CompactBinaryFieldType type)
+        {
+            CompactBinaryFieldType noFlags = RemoveFlags(type);
+            return noFlags >= CompactBinaryFieldType.Object && noFlags <= CompactBinaryFieldType.UniformArray;
+        }
+
+        public static bool HasFieldType(CompactBinaryFieldType type)
+        {
+            return type.HasFlag(CompactBinaryFieldType.HasFieldType);
+        }
+
+        public static bool HasFieldName(CompactBinaryFieldType type)
+        {
+            return type.HasFlag(CompactBinaryFieldType.HasFieldName);
+        }
+
+        public static CompactBinaryFieldType RemoveFlags(CompactBinaryFieldType type)
+        {
+            return type & TypeMask;
+        }
+
+        public static ulong GetPayloadSize(CompactBinaryFieldType type, ReadOnlyMemory<byte> payload)
+        {
+            switch (RemoveFlags(type))
+            {
+                case CompactBinaryFieldType.None:
+                case CompactBinaryFieldType.Null:
+                    return 0;
+                case CompactBinaryFieldType.Object:
+                case CompactBinaryFieldType.UniformObject:
+                case CompactBinaryFieldType.Array:
+                case CompactBinaryFieldType.UniformArray:
+                case CompactBinaryFieldType.Binary:
+                case CompactBinaryFieldType.String:
+                case CompactBinaryFieldType.CustomByName:
+                case CompactBinaryFieldType.CustomById:
+                {
+                    ReadOnlyMemory<byte> localPayload = payload.Slice(0);
+                    ulong payloadSize = BitUtils.ReadVarUInt(ref localPayload, out int bytesRead);
+                    return payloadSize + (ulong)bytesRead;
+                }
+
+                case CompactBinaryFieldType.IntegerPositive:
+                case CompactBinaryFieldType.IntegerNegative:
+                {
+                    ReadOnlyMemory<byte> localPayload = payload.Slice(0);
+                    return BitUtils.MeasureVarUInt(ref localPayload);
+                }
+                case CompactBinaryFieldType.Float32:
+                    return 4;
+                case CompactBinaryFieldType.Float64:
+                    return 8;
+                case CompactBinaryFieldType.BoolFalse:
+                case CompactBinaryFieldType.BoolTrue:
+                    return 0;
+                case CompactBinaryFieldType.CompactBinaryAttachment:
+                case CompactBinaryFieldType.BinaryAttachment:
+                case CompactBinaryFieldType.Hash:
+                    return 20;
+                case CompactBinaryFieldType.Uuid:
+                    return 16;
+                case CompactBinaryFieldType.DateTime:
+                case CompactBinaryFieldType.TimeSpan:
+                    return 8;
+                case CompactBinaryFieldType.ObjectId:
+                    return 12;
+                default:
+                    return 0;
+            }
+        }
+    }
+    [Flags]
+    public enum CompactBinaryFieldType : byte
+    {
+        /** A field type that does not occur in a valid object. */
+        None                             = 0x00,
+
+        /** Null. Payload is empty. */
+        Null                             = 0x01,
+
+        /**
+         * Object is an array of fields with unique non-empty names.
+         *
+         * Payload is a VarUInt byte count for the encoded fields followed by the fields.
+         */
+        Object                           = 0x02,
+        /**
+         * UniformObject is an array of fields with the same field types and unique non-empty names.
+         *
+         * Payload is a VarUInt byte count for the encoded fields followed by the fields.
+         */
+        UniformObject                    = 0x03,
+
+        /**
+         * Array is an array of fields with no name that may be of different types.
+         *
+         * Payload is a VarUInt byte count, followed by a VarUInt item count, followed by the fields.
+         */
+        Array                            = 0x04,
+        /**
+         * UniformArray is an array of fields with no name and with the same field type.
+         *
+         * Payload is a VarUInt byte count, followed by a VarUInt item count, followed by field type,
+         * followed by the fields without their field type.
+         */
+        UniformArray                     = 0x05,
+
+        /** Binary. Payload is a VarUInt byte count followed by the data. */
+        Binary                           = 0x06,
+
+        /** String in UTF-8. Payload is a VarUInt byte count then an unterminated UTF-8 string. */
+        String                           = 0x07,
+
+        /**
+         * Non-negative integer with the range of a 64-bit unsigned integer.
+         *
+         * Payload is the value encoded as a VarUInt.
+         */
+        IntegerPositive                  = 0x08,
+        /**
+         * Negative integer with the range of a 64-bit signed integer.
+         *
+         * Payload is the ones' complement of the value encoded as a VarUInt.
+         */
+        IntegerNegative                  = 0x09,
+
+        /** Single precision float. Payload is one big endian IEEE 754 binary32 float. */
+        Float32                          = 0x0a,
+        /** Double precision float. Payload is one big endian IEEE 754 binary64 float. */
+        Float64                          = 0x0b,
+
+        /** Boolean false value. Payload is empty. */
+        BoolFalse                        = 0x0c,
+        /** Boolean true value. Payload is empty. */
+        BoolTrue                         = 0x0d,
+
+        /**
+         * CompactBinaryAttachment is a reference to a compact binary attachment stored externally.
+         *
+         * Payload is a 160-bit hash digest of the referenced compact binary data.
+         */
+        CompactBinaryAttachment          = 0x0e,
+        /**
+         * BinaryAttachment is a reference to a binary attachment stored externally.
+         *
+         * Payload is a 160-bit hash digest of the referenced binary data.
+         */
+        BinaryAttachment                 = 0x0f,
+
+        /** Hash. Payload is a 160-bit hash digest. */
+        Hash                             = 0x10,
+        /** UUID/GUID. Payload is a 128-bit UUID as defined by RFC 4122. */
+        Uuid                             = 0x11,
+
+        /**
+         * Date and time between 0001-01-01 00:00:00.0000000 and 9999-12-31 23:59:59.9999999.
+         *
+         * Payload is a big endian int64 count of 100ns ticks since 0001-01-01 00:00:00.0000000.
+         */
+        DateTime                         = 0x12,
+        /**
+         * Difference between two date/time values.
+         *
+         * Payload is a big endian int64 count of 100ns ticks in the span, and may be negative.
+         */
+        TimeSpan                         = 0x13,
+
+        /**
+         * ObjectId is an opaque object identifier. See FCbObjectId.
+         *
+         * Payload is a 12-byte object identifier.
+         */
+        ObjectId                         = 0x14,
+
+        /**
+         * CustomById identifies the sub-type of its payload by an integer identifier.
+         *
+         * Payload is a VarUInt byte count of the sub-type identifier and the sub-type payload, followed
+         * by a VarUInt of the sub-type identifier then the payload of the sub-type.
+         */
+        CustomById                       = 0x1e,
+        /**
+         * CustomByType identifies the sub-type of its payload by a string identifier.
+         *
+         * Payload is a VarUInt byte count of the sub-type identifier and the sub-type payload, followed
+         * by a VarUInt byte count of the unterminated sub-type identifier, then the sub-type identifier
+         * without termination, then the payload of the sub-type.
+         */
+        CustomByName                     = 0x1f,
+
+        /** Reserved for future use as a flag. Do not add types in this range. */
+        Reserved                         = 0x20,
+
+        /**
+         * A transient flag which indicates that the object or array containing this field has stored
+         * the field type before the payload and name. Non-uniform objects and fields will set this.
+         *
+         * Note: Since the flag must never be serialized, this bit may be repurposed in the future.
+         */
+        HasFieldType                     = 0x40,
+
+        /** A persisted flag which indicates that the field has a name stored before the payload. */
+        HasFieldName                     = 0x80,
+
+    }
+
+
+    public class CompactBinaryField
+    {
+        /*
+        * These has to follow Engine\Source\Runtime\Core\Public\Serialization\CompactBinary.h
+        *
+        * DO NOT CHANGE THE VALUE OF ANY MEMBERS OF THIS ENUM!
+        * BACKWARD COMPATIBILITY REQUIRES THAT THESE VALUES BE FIXED!
+        * SERIALIZATION USES HARD-CODED CONSTANTS BASED ON THESE VALUES!
+        */
+
+        private readonly CompactBinaryFieldType _type;
+        // the memory storing the name
+        private readonly ReadOnlyMemory<byte> _nameMemory;
+        // the memory area of the payload
+        private readonly ReadOnlyMemory<byte> _payload;
+        private readonly int _objectSize;
+
+        public CompactBinaryField(ReadOnlyMemory<byte> data, CompactBinaryFieldType type = CompactBinaryFieldType.HasFieldType)
+        {
+            _objectSize = 0;
+            CompactBinaryFieldType fieldType = type;
+            if (CompactBinaryFieldUtils.HasFieldType(fieldType))
+            {
+                fieldType = (CompactBinaryFieldType)data.Span[0] | CompactBinaryFieldType.HasFieldType;
+                data = data.Slice(1);
+                _objectSize += 1;
+            }
+            _type = fieldType;
+
+            _nameMemory = ReadOnlyMemory<byte>.Empty;
+            if (CompactBinaryFieldUtils.HasFieldName(fieldType))
+            {
+                int nameLength = (int)BitUtils.ReadVarUInt(ref data, out int nameBytesRead);
+                _nameMemory = data.Slice(nameBytesRead, nameLength);
+                data = data.Slice(nameLength + nameBytesRead);
+                _objectSize += nameLength + nameBytesRead;
+            }
+
+            int payloadSize = (int) CompactBinaryFieldUtils.GetPayloadSize(_type, data);
+            // we reassign the memory range based on how long the payload is expected to be
+            _payload = data.Slice(0, payloadSize);
+            _objectSize += payloadSize;
+        }
+
+        
+        public string Name 
+        { 
+            get
+            {
+                return _nameMemory.Length == 0 ? "" : Encoding.ASCII.GetString(_nameMemory.Span);
+            }
+        }
+
+        public CompactBinaryFieldType Type
+        {
+            get
+            {
+                return CompactBinaryFieldUtils.RemoveFlags(_type);
+            }
+        }
+
+        /** Value accessors **/
+        public string AsString(string @default = "")
+        {
+            if (CompactBinaryFieldUtils.IsString(_type))
+            {
+                ReadOnlyMemory<byte> localPayload = _payload;
+                ulong stringLength = BitUtils.ReadVarUInt(ref localPayload, out int bytesRead);
+                ReadOnlyMemory<byte> stringMemory = _payload.Slice(bytesRead, (int)stringLength);
+                return Encoding.ASCII.GetString(stringMemory.Span);
+            }
+
+            return @default;
+        }
+
+        public BlobIdentifier? AsCompactBinaryAttachment()
+        {
+            if (CompactBinaryFieldUtils.IsCompactBinaryAttachment(_type))
+            {
+                return new BlobIdentifier(_payload.ToArray());
+            }
+
+            return null;
+        }
+
+        public BlobIdentifier? AsBinaryAttachment()
+        {
+            if (CompactBinaryFieldUtils.IsBinaryAttachment(_type))
+            {
+                return new BlobIdentifier(_payload.ToArray());
+            }
+
+            return null;
+        }
+
+        public BlobIdentifier? AsAttachment()
+        {
+            if (CompactBinaryFieldUtils.IsAttachment(_type))
+            {
+                return new BlobIdentifier(_payload.ToArray());
+            }
+
+            return null;
+        }
+
+        public BlobIdentifier? AsHash()
+        {
+            if (CompactBinaryFieldUtils.IsHash(_type))
+            {
+                return new BlobIdentifier(_payload.ToArray());
+            }
+
+            return null;
+        }
+
+        public DateTime? AsDateTime()
+        {
+            if (CompactBinaryFieldUtils.IsDateTime(_type))
+            {
+                long? ticks = AsLongInteger();
+                if (ticks == null)
+                    return null;
+
+                return DateTime.FromFileTime(ticks.Value);
+            }
+
+            return null;
+        }
+
+        public int? AsInteger()
+        {
+            if (CompactBinaryFieldUtils.IsInteger(_type))
+            {
+                return BitConverter.ToInt32(_payload.Span);
+            }
+
+            return null;
+        }
+
+        public long? AsLongInteger()
+        {
+            if (CompactBinaryFieldUtils.IsInteger(_type))
+            {
+                return BitConverter.ToInt64(_payload.Span);
+            }
+
+            return null;
+        }
+
+        public bool? AsBool()
+        {
+            if (CompactBinaryFieldUtils.IsBool(_type))
+            {
+                CompactBinaryFieldType type = CompactBinaryFieldUtils.RemoveFlags(_type);
+                if (type == CompactBinaryFieldType.BoolFalse)
+                    return false;
+                else if (type == CompactBinaryFieldType.BoolTrue)
+                    return true;
+
+                throw new NotImplementedException("Unknown bool mapping, was not false or true");
+            }
+
+            return null;
+        }
+
+        public byte[]? AsBinary()
+        {
+            if (CompactBinaryFieldUtils.IsBinary(_type))
+            {
+                ReadOnlyMemory<byte> localPayload = _payload.Slice(0);
+                ulong payloadSize = BitUtils.ReadVarUInt(ref localPayload, out int payloadBytesRead);
+                localPayload = localPayload.Slice(payloadBytesRead);
+
+                return localPayload.ToArray();
+            }
+
+            return null;
+        }
+
+        public IEnumerable<CompactBinaryField> AsArray()
+        {
+            return GetFields();
+        }
+        
+        public CompactBinaryObject AsObject()
+        {
+            ReadOnlyMemory<byte> localMemory = _payload.Slice(0);
+            return CompactBinaryObject.Load(ref localMemory);
+        }
+
+        public Guid AsUuid(Guid @default = default(Guid))
+        {
+            if (CompactBinaryFieldUtils.IsUuid(_type))
+            {
+                return new Guid(_payload.ToArray());
+            }
+
+            return @default;
+        }
+
+        public bool IsObject()
+        {
+            return CompactBinaryFieldUtils.IsObject(_type);
+        }
+
+        public bool IsArray()
+        {
+            return CompactBinaryFieldUtils.IsArray(_type);
+        }
+
+        public bool IsInteger()
+        {
+            return CompactBinaryFieldUtils.IsInteger(_type);
+        }
+
+        public bool IsDateTime()
+        {
+            return CompactBinaryFieldUtils.IsDateTime(_type);
+        }
+
+        public bool IsNull()
+        {
+            return CompactBinaryFieldUtils.IsNull(_type);
+        }
+
+        public bool IsBool()
+        {
+            return CompactBinaryFieldUtils.IsBool(_type);
+        }
+
+        public bool IsString()
+        {
+            return CompactBinaryFieldUtils.IsString(_type);
+        }
+
+        public bool IsHash()
+        {
+            return CompactBinaryFieldUtils.IsHash(_type);
+        }
+
+        public bool IsBinary()
+        {
+            return CompactBinaryFieldUtils.IsBinary(_type);
+        }
+
+        public bool IsBinaryAttachment()
+        {
+            return CompactBinaryFieldUtils.IsBinaryAttachment(_type);
+        }
+
+        public bool IsCompactBinaryAttachment()
+        {
+            return CompactBinaryFieldUtils.IsCompactBinaryAttachment(_type);
+        }
+
+        public bool IsAttachment()
+        {
+            return CompactBinaryFieldUtils.IsAttachment(_type);
+        }
+
+        public int GetObjectLength()
+        {
+            return _objectSize;
+        }
+
+        public IEnumerable<CompactBinaryField> GetFields()
+        {
+            if (CompactBinaryFieldUtils.HasFields(_type))
+            {
+                ReadOnlyMemory<byte> payloadBytes = _payload.Slice(0);
+                ulong payloadSize = BitUtils.ReadVarUInt(ref payloadBytes, out int payloadBytesRead);
+                payloadBytes = payloadBytes.Slice(payloadBytesRead);
+                uint byteCount = CompactBinaryFieldUtils.IsArray(_type) ? BitUtils.MeasureVarUInt(ref payloadBytes) : 0;
+                if (payloadSize > byteCount)
+                {
+                    payloadBytes = payloadBytes.Slice((int) byteCount, (int) (payloadSize - byteCount));
+                    CompactBinaryFieldType uniformType = CompactBinaryFieldType.HasFieldType;
+                    if (CompactBinaryFieldUtils.HasUniformFields(_type))
+                    {
+                        uniformType = (CompactBinaryFieldType) payloadBytes.Span[0];
+                        payloadBytes = payloadBytes.Slice(1);
+                    }
+
+                    while (payloadBytes.Length != 0)
+                    {
+                        CompactBinaryField field = new CompactBinaryField(payloadBytes, uniformType);
+                        payloadBytes = payloadBytes.Slice(field.GetObjectLength());
+
+                        yield return field;
+                    }
+                }
+            }
+        }
+
+        // Fetch a field by case sensitive name
+        public CompactBinaryField? this[string name]
+        {
+            get { return GetFields().FirstOrDefault(field => field.Name == name); }
+        }
+
+        public CompactBinaryField? this[int i]
+        {
+            get { return GetFields().Skip(i).FirstOrDefault(); }
+        }
+    }
+
+    public class CompactBinaryObject : CompactBinaryField
+    {
+        private CompactBinaryObject(byte[] buffer) : base(new ReadOnlyMemory<byte>(buffer))
+        {
+            Data = buffer;
+        }
+
+        public byte[] Data { get; }
+
+        public static CompactBinaryObject Load(ref ReadOnlyMemory<byte> data)
+        {
+            ReadOnlyMemory<byte> localMemory = data.Slice(0);
+            if (!TryMeasureCompactBinary(ref localMemory, out CompactBinaryFieldType _, out ulong size))
+            {
+                throw new Exception("Failed to load from invalid compact binary data.");
+            }
+
+            byte[] buffer = data.Slice(0, (int)size).ToArray();
+            return new CompactBinaryObject(buffer);
+        }
+
+        public static CompactBinaryObject Load(byte[] data)
+        {
+            ReadOnlyMemory<byte> localMemory = new ReadOnlyMemory<byte>(data);
+            if (!TryMeasureCompactBinary(ref localMemory, out CompactBinaryFieldType _, out ulong size))
+            {
+                throw new Exception("Failed to load from invalid compact binary data.");
+            }
+
+            byte[] buffer = new byte[size];
+            Array.Copy(data, buffer, (long)size);
+            return new CompactBinaryObject(buffer);
+        }
+
+        static bool TryMeasureCompactBinary(ref ReadOnlyMemory<byte> data, out CompactBinaryFieldType outFieldType, out ulong outSize, CompactBinaryFieldType assumeFieldType = CompactBinaryFieldType.HasFieldType)
+        {
+            CompactBinaryFieldType type = assumeFieldType;
+            ulong size = 0;
+
+            if (CompactBinaryFieldUtils.HasFieldType(type))
+            {
+                if (data.Length == 0)
+                {
+                    outFieldType = CompactBinaryFieldType.None;
+                    outSize = 1;
+                    return false;
+                }
+
+                type = (CompactBinaryFieldType) data.Span[0];
+                data = data.Slice(1);
+                size += 1;
+            }
+
+            bool dynamicSize = false;
+            ulong fixedSize = 0;
+            switch (CompactBinaryFieldUtils.RemoveFlags(type))
+            {
+            case CompactBinaryFieldType.Null:
+                break;
+            case CompactBinaryFieldType.Object:
+            case CompactBinaryFieldType.UniformObject:
+            case CompactBinaryFieldType.Array:
+            case CompactBinaryFieldType.UniformArray:
+            case CompactBinaryFieldType.Binary:
+            case CompactBinaryFieldType.String:
+            case CompactBinaryFieldType.IntegerPositive:
+            case CompactBinaryFieldType.IntegerNegative:
+            case CompactBinaryFieldType.CustomById:
+            case CompactBinaryFieldType.CustomByName:
+                dynamicSize = true;
+                break;
+            case CompactBinaryFieldType.Float32:
+                fixedSize = 4;
+                break;
+            case CompactBinaryFieldType.Float64:
+                fixedSize = 8;
+                break;
+            case CompactBinaryFieldType.BoolFalse:
+            case CompactBinaryFieldType.BoolTrue:
+                break;
+            case CompactBinaryFieldType.CompactBinaryAttachment:
+            case CompactBinaryFieldType.BinaryAttachment:
+            case CompactBinaryFieldType.Hash:
+                fixedSize = 20;
+                break;
+            case CompactBinaryFieldType.Uuid:
+                fixedSize = 16;
+                break;
+            case CompactBinaryFieldType.DateTime:
+            case CompactBinaryFieldType.TimeSpan:
+                fixedSize = 8;
+                break;
+            case CompactBinaryFieldType.ObjectId:
+                fixedSize = 12;
+                break;
+            case CompactBinaryFieldType.None:
+            default:
+                outFieldType = CompactBinaryFieldType.None;
+                outSize = 0;
+                return false;
+            }
+
+            outFieldType = type;
+
+            if (CompactBinaryFieldUtils.HasFieldName(type))
+            {
+                if (data.Length == 0)
+                {
+                    outSize = size + 1;
+                    return false;
+                }
+
+                uint nameLenByteCount = BitUtils.MeasureVarUInt(ref data);
+                if (data.Length < nameLenByteCount)
+                {
+                    outSize = size + nameLenByteCount;
+                    return false;
+                }
+
+                ulong nameLen = BitUtils.ReadVarUInt(ref data, out int bytesRead);
+                ulong nameSize = nameLen + nameLenByteCount;
+
+                if (dynamicSize && data.Length < (int)nameSize)
+                {
+                    outSize = size + nameSize;
+                    return false;
+                }
+
+                data = data.Slice((int)nameSize);
+                size += nameSize;
+            }
+
+            switch (CompactBinaryFieldUtils.RemoveFlags(type))
+            {
+            case CompactBinaryFieldType.Object:
+            case CompactBinaryFieldType.UniformObject:
+            case CompactBinaryFieldType.Array:
+            case CompactBinaryFieldType.UniformArray:
+            case CompactBinaryFieldType.Binary:
+            case CompactBinaryFieldType.String:
+            case CompactBinaryFieldType.CustomById:
+            case CompactBinaryFieldType.CustomByName:
+                if (data.Length == 0)
+                {
+                    outSize = size + 1;
+                    return false;
+                }
+                else
+                {
+                    uint payloadSizeByteCount = BitUtils.MeasureVarUInt(ref data);
+                    if (data.Length < payloadSizeByteCount)
+                    {
+                        outSize = size + payloadSizeByteCount;
+                        return false;
+                    }
+                    ulong payloadSize = BitUtils.ReadVarUInt(ref data, out int _);
+                    outSize = size + payloadSize + payloadSizeByteCount;
+                    return true;
+                }
+            case CompactBinaryFieldType.IntegerPositive:
+            case CompactBinaryFieldType.IntegerNegative:
+                if (data.Length == 0)
+                {
+                    outSize = size + 1;
+                    return false;
+                }
+                outSize = size + BitUtils.MeasureVarUInt(ref data);
+                return true;
+            default:
+                outSize = size + fixedSize;
+                return true;
+            }
+        }
+
+        public IEnumerable<CompactBinaryField> GetAllFields()
+        {
+            Queue<CompactBinaryField> fieldQueue = new Queue<CompactBinaryField>(GetFields());
+
+            while(fieldQueue.Count != 0)
+            {
+                CompactBinaryField? field = fieldQueue.Dequeue();
+                yield return field;
+
+                foreach (CompactBinaryField f in field.GetFields())
+                {
+                    fieldQueue.Enqueue(f);
+                }
+            }
+        }
+
+        public string ToJson()
+        {
+            var buffer = new ArrayBufferWriter<byte>();
+            using Utf8JsonWriter jsonWriter = new Utf8JsonWriter(buffer);
+
+            jsonWriter.WriteStartObject();
+            foreach (CompactBinaryField field in GetFields())
+            {
+                WriteField(field, jsonWriter);
+            }
+            jsonWriter.WriteEndObject();
+            jsonWriter.Flush();
+            return Encoding.UTF8.GetString(buffer.WrittenMemory.Span);
+        }
+
+        private static void WriteField(CompactBinaryField field, Utf8JsonWriter jsonWriter)
+        {
+            if (field.IsObject())
+            {
+                jsonWriter.WriteStartObject();
+                CompactBinaryObject o = field.AsObject();
+                foreach (CompactBinaryField f in o.GetFields())
+                {
+                    WriteField(f, jsonWriter);
+                }
+                jsonWriter.WriteEndObject();
+            }
+            else if (field.IsArray())
+            {
+                jsonWriter.WriteStartArray();
+                jsonWriter.WriteEndArray();
+            }
+            else if (field.IsInteger())
+            {
+                jsonWriter.WriteNumber(field.Name, field.AsInteger()!.Value);
+            }
+            else if (field.IsBool())
+            {
+                jsonWriter.WriteBoolean(field.Name, field.AsBool()!.Value);
+            }
+            else if (field.IsNull())
+            {
+                jsonWriter.WriteNullValue();
+            }
+            else if (field.IsDateTime())
+            {
+                jsonWriter.WriteString(field.Name, field.AsDateTime()!.Value);
+            }
+            else if (field.IsHash())
+            {
+                jsonWriter.WriteString(field.Name, field.AsHash()!.ToString()!);
+            }
+            else if (field.IsString())
+            {
+                jsonWriter.WriteString(field.Name, field.AsString());
+            }
+            else
+            {
+                throw new NotImplementedException($"Not handled type {field.Type} when attempting to convert to json");
+            }
+
+        }
+
+        
+        public object ToPoco(Type pocoType)
+        {
+            object? ToPocoField(CompactBinaryField field)
+            {
+                if (field.IsObject())
+                {
+                    ExpandoObject o = new ExpandoObject();
+                    foreach (CompactBinaryField f in field.GetFields())
+                    {
+                        o.TryAdd(f.Name, ToPocoField(f));
+                    }
+
+                    return o;
+                }
+                else if (field.IsArray())
+                {
+                    List<object?> objs = new();
+                    foreach (CompactBinaryField f in field.AsArray())
+                    {
+                        objs.Add(ToPocoField(f));
+                    }
+
+                    return objs.ToArray();
+                }
+                else if (field.IsInteger())
+                {
+                    return field.AsInteger()!.Value;
+                }
+                else if (field.IsBool())
+                {
+                    return field.AsBool()!.Value;
+                }
+                else if (field.IsNull())
+                {
+                    return null;
+                }
+                else if (field.IsDateTime())
+                {
+                    return field.AsDateTime()!.Value;
+                }
+                else if (field.IsHash())
+                {
+                    return field.AsHash()!;
+                }
+                else if (field.IsString())
+                {
+                    return field.AsString();
+                }
+                else
+                {
+                    throw new NotImplementedException($"Not handled type {field.Type} when attempting to convert to poco type");
+                }
+            }
+
+            if (IsArray())
+            {
+                if (!pocoType.IsArray)
+                    throw new Exception($"Trying to convert a compact binary array to a none array poco type: {pocoType}");
+
+                Type arrayType = pocoType.GetElementType()!;
+                List<object?> objs = new();
+                foreach (CompactBinaryField f in GetFields())
+                {
+                    objs.Add(ToPocoField(f));
+                }
+
+                Array array = Array.CreateInstance(arrayType, objs.Count);
+                for (int i = 0; i < objs.Count; i++)
+                {
+                    array.SetValue(objs[i], i);
+                }
+
+                return array;
+            }
+
+            ConstructorInfo? defaultConstructor = pocoType.GetConstructor(BindingFlags.CreateInstance, null, Array.Empty<Type>(), Array.Empty<ParameterModifier>());
+            if (defaultConstructor == null)
+                throw new Exception($"Failed to find a parameter-less constructor for type {pocoType} which is required to serialize from a compact binary");
+            object o = defaultConstructor.Invoke(Array.Empty<object>());
+
+            foreach (CompactBinaryField field in GetFields())
+            {
+                FieldInfo? fieldInfo = pocoType.GetField(field.Name);
+                if (fieldInfo == null)
+                    throw new Exception($"Failed to map compact binary field \"{field.Name}\" to a field in type {pocoType}");
+
+                fieldInfo.SetValue(o, ToPocoField(field));
+            }
+
+            return o;
+        }
+
+    }
+
+    static class BitUtils
+    { 
+        /**
+             * Read a variable-length unsigned integer.
+             *
+             * @param InData A variable-length encoding of an unsigned integer.
+             * @param OutByteCount The number of bytes consumed from the input.
+             * @return An unsigned integer.
+             */
+        public static ulong ReadVarUInt(ref ReadOnlyMemory<byte> buffer, out int bytesRead)
+        {
+            bytesRead = (int)MeasureVarUInt(ref buffer);
+
+            ulong value = (ulong) (buffer.Span[0] & (0xff >> bytesRead));
+
+            for (int i = 1; i < bytesRead; i++)
+            {
+                value <<= 8;
+                value |= buffer.Span[i];
+            }
+
+            return value;
+        }
+
+        /**
+         * Measure the length in bytes (1-9) of an encoded variable-length integer.
+         *
+         * @param InData A variable-length encoding of an (signed or unsigned) integer.
+         * @return The number of bytes used to encode the integer, in the range 1-9.
+         */
+        public static uint MeasureVarUInt(ref ReadOnlyMemory<byte> buffer)
+        {
+            byte b = buffer.Span[0];
+            b = (byte)~b;
+            return (uint)BitOperations.LeadingZeroCount (b) - 23;
+        }
+
+        /** Measure the number of bytes (1-5) required to encode the 32-bit input. */
+
+        public static uint MeasureVarUInt(uint value)
+        {
+            //return uint32(int32(FPlatformMath::FloorLog2(InValue)) / 7 + 1);
+
+            return (uint) (BitOperations.Log2(value) / 7 + 1);
+        }
+
+        /** Measure the number of bytes (1-9) required to encode the 64-bit input. */
+
+        public static uint MeasureVarUInt(ulong value)
+        {
+            //return uint32(FPlatformMath::Min(int32(FPlatformMath::FloorLog2_64(InValue)) / 7 + 1, 9));
+
+            return (uint) (Math.Min(BitOperations.Log2(value) / 7 + 1, 9));
+        }
+
+        /**
+         * Write a variable-length unsigned integer.
+         *
+         * @param InValue An unsigned integer to encode.
+         * @param OutData A buffer of at least 9 bytes to write the output to.
+         * @return The number of bytes used in the output.
+         */
+
+        public static uint WriteVarUInt(ulong value, byte[] buffer, int bufferOffset = 0)
+        {
+            uint byteCount = MeasureVarUInt(value);
+
+            for (int i = 1; i < byteCount; i++)
+            {
+                buffer[bufferOffset + byteCount - i] = (byte)value;
+                value >>= 8;
+            }
+            buffer[bufferOffset] = (byte)((0xff << (9 - (int)byteCount)) | (byte)value);
+            return byteCount;
+        }
+    }
+}
diff --git a/Engine/Source/Programs/Horde/HordeStorage/Jupiter.Common/Implementation/CompactBinaryWriter.cs b/Engine/Source/Programs/Horde/HordeStorage/Jupiter.Common/Implementation/CompactBinaryWriter.cs
new file mode 100644
--- /dev/null
+++ b/Engine/Source/Programs/Horde/HordeStorage/Jupiter.Common/Implementation/CompactBinaryWriter.cs
@@ -0,0 +1,322 @@
+﻿// Copyright Epic Games, Inc. All Rights Reserved.
+
+using System;
+using System.Collections.Generic;
+using System.Linq;
+using System.Text;
+
+namespace Jupiter.Implementation
+{
+    public class CompactBinaryWriter
+    {
+        readonly List<Field> _fields = new List<Field>();
+        private Field _topField;
+
+        struct Field
+        {
+            public CompactBinaryFieldType FieldType;
+            public byte[]? Name;
+            public byte[]? Payload;
+        }
+
+        public byte[] Save()
+        {
+            uint bufferLength = 1; // field type;
+
+            if (_topField.Name != null)
+            {
+                bufferLength += (uint)_topField.Name!.Length;
+            }
+
+            uint payloadBytesLength = 0;
+            if (_topField.Payload != null)
+            {
+                payloadBytesLength = BitUtils.MeasureVarUInt((uint)_topField.Payload!.Length);
+                bufferLength += payloadBytesLength;
+                bufferLength += (uint)_topField.Payload!.Length;
+            }
+            byte[] buffer = new byte[bufferLength];
+            buffer[0] = (byte) _topField.FieldType;
+            int offset = 1;
+
+            if (_topField.Name != null)
+            {
+                Array.Copy(_topField.Name!, 0, buffer, offset, _topField.Name!.Length);
+                offset += _topField.Name.Length;
+            }
+
+            if (_topField.Payload != null)
+            {
+                // an array already contains the payload length
+                if (!CompactBinaryFieldUtils.IsArray(_topField.FieldType))
+                {
+                    BitUtils.WriteVarUInt((uint) _topField.Payload!.Length, buffer, offset);
+                    offset += (int)payloadBytesLength;
+                }
+                Array.Copy(_topField.Payload!, 0, buffer, offset, _topField.Payload!.Length);
+                offset += _topField.Payload.Length;
+            }
+
+            return buffer;
+        }
+
+        public void BeginObject()
+        {
+            _topField = BeginField();
+        }
+
+        public void EndObject()
+        {
+            _topField.Payload = SerializePayload(_fields);
+            
+            EndField(ref this._topField, CompactBinaryFieldType.Object);
+        }
+
+        public void AsArray<T>(T[] members, CompactBinaryFieldType fieldType)
+        {
+            AddUniformArray<T>(members, fieldType);
+            _topField = _fields.Last();
+        }
+
+        private static byte[] SerializePayload(List<Field> fields)
+        {
+            uint totalSizeFields = 0;
+            foreach (Field field in fields)
+            {
+                totalSizeFields += 1; // the field type
+                if (field.Payload != null)
+                    totalSizeFields += (uint)field.Payload.Length;
+                if (field.Name != null)
+                    totalSizeFields += (uint)field.Name.Length;
+            }
+
+            byte[] buffer = new byte[totalSizeFields];
+            long payloadOffset = 0;
+            foreach (Field field in fields)
+            {
+                buffer[payloadOffset] = (byte) field.FieldType;
+                payloadOffset += 1;
+
+                if (CompactBinaryFieldUtils.HasFieldName(field.FieldType))
+                {
+                    Array.Copy(field.Name!, 0, buffer, payloadOffset, field.Name!.Length);
+                    payloadOffset += field.Name.Length;
+                }
+
+                if ((field.Payload?.Length ?? 0) != 0)
+                {
+                    Array.Copy(field.Payload!, 0, buffer, payloadOffset, field.Payload!.Length);
+                    payloadOffset += field.Payload.Length;
+                }
+            }
+
+            return buffer;
+        }
+
+        public void AddString(string value, string? fieldName = null)
+        {
+            Field field = BeginField();
+            SetName(ref field, fieldName);
+
+            uint length = (uint) Encoding.UTF8.GetByteCount(value);
+            uint bytesCount = BitUtils.MeasureVarUInt(length);
+            byte[] buffer = new byte[length + bytesCount];
+            BitUtils.WriteVarUInt(length, buffer);
+            int bytesWritten = Encoding.UTF8.GetBytes(value, 0, value.Length, buffer, (int)bytesCount);
+            if (bytesWritten != length)
+                throw new Exception("Failed to write string into buffer");
+
+            field.Payload = buffer;
+            EndField(ref field, CompactBinaryFieldType.String);
+        }
+
+
+        private Field BeginField()
+        {
+            // TODO: We could add checks similar to what we do in C++ to make sure the writer is used correctly
+            return new Field();
+        }
+
+        private void EndField(ref Field field, CompactBinaryFieldType fieldType)
+        {
+            field.FieldType = fieldType;
+            if (field.Name != null && field.Name.Length != 0)
+            {
+                field.FieldType |= CompactBinaryFieldType.HasFieldName;
+            }
+            // TODO: check if types are uniform and check the field type to be uniform thus removing redudant type info
+
+            _fields.Add(field);
+        }
+
+        public void AddBinaryAttachment(BlobIdentifier hash, string? fieldName = null)
+        {
+            Field field = BeginField();
+            SetName(ref field, fieldName);
+
+            if (hash.HashData.Length != 20)
+                throw new Exception("Hash data is assumed to be 20 bytes");
+            field.Payload = hash.HashData;
+
+            EndField(ref field, CompactBinaryFieldType.BinaryAttachment);
+        }
+
+        public void AddCompactBinaryAttachment(BlobIdentifier hash, string? fieldName = null)
+        {
+            Field field = BeginField();
+            SetName(ref field, fieldName);
+
+            if (hash.HashData.Length != 20)
+                throw new Exception("Hash data is assumed to be 20 bytes");
+            field.Payload = hash.HashData;
+
+            EndField(ref field, CompactBinaryFieldType.CompactBinaryAttachment);
+        }
+
+        public void AddNull(string? fieldName = null)
+        {
+            Field field = BeginField();
+            SetName(ref field, fieldName);
+            EndField(ref field, CompactBinaryFieldType.Null);
+        }
+
+        
+        public void AddDateTime(DateTime dateTime, string? fieldName = null)
+        {
+            Field field = BeginField();
+            SetName(ref field, fieldName);
+            long ticks = dateTime.Ticks;
+            byte[] payload = BitConverter.GetBytes(ticks);
+            if (BitConverter.IsLittleEndian)
+            {
+                Array.Reverse(payload, 0, payload.Length);
+            }
+            field.Payload = payload;
+
+            EndField(ref field, CompactBinaryFieldType.DateTime);
+        }
+
+        
+        public void AddBinary(byte[] data, string? fieldName = null)
+        {
+            Field field = BeginField();
+            SetName(ref field, fieldName);
+
+            long bytesLength = BitUtils.MeasureVarUInt((uint)data.Length);
+            byte[] payload = new byte[data.Length + bytesLength];
+            BitUtils.WriteVarUInt((uint)data.Length, payload);
+            Array.Copy(data, 0, payload, bytesLength, data.Length);
+            field.Payload = payload;
+            EndField(ref field, CompactBinaryFieldType.Binary);
+        }
+
+        public void AddUniformArray<T>(T[] members, CompactBinaryFieldType fieldType, string? fieldName = null)
+        {
+            Field field = BeginField();
+            SetName(ref field, fieldName);
+
+            uint bytesMembers = BitUtils.MeasureVarUInt((uint)members.Length);
+
+            List<byte[]> payloads = new List<byte[]>();
+
+            foreach (T member in members)
+            {
+                byte[] buffer = GetByteValue<T>(member, fieldType);
+                payloads.Add(buffer);
+            }
+
+            long payloadsBytesLength = payloads.Sum(payload => (uint)payload.Length);
+            uint payloadsLength = 1 + bytesMembers + (uint)payloadsBytesLength;
+            uint bytesPayloads = BitUtils.MeasureVarUInt(payloadsLength);
+
+            // Uniform array payload is a VarUInt byte count, followed by a VarUInt item count, followed by field type, followed by the fields without their field type.
+            byte[] payloadBuffer = new byte[payloadsLength + bytesPayloads];
+            BitUtils.WriteVarUInt(payloadsLength, payloadBuffer);
+            int bufferOffset = (int)bytesPayloads;
+            BitUtils.WriteVarUInt((uint)members.Length, payloadBuffer, bufferOffset);
+            bufferOffset += (int)bytesMembers;
+            payloadBuffer[bufferOffset] = (byte)fieldType;
+            bufferOffset += 1;
+
+            foreach (byte[] payload in payloads)
+            {
+                int payloadLength = payload.Length;
+                Array.Copy(payload, 0, payloadBuffer, bufferOffset, payloadLength);
+                bufferOffset += payloadLength;
+            }
+
+            field.Payload = payloadBuffer;
+            EndField(ref field, CompactBinaryFieldType.UniformArray);
+        }
+
+        private byte[] GetByteValue<T>(T member, CompactBinaryFieldType fieldType)
+        {
+            switch (fieldType)
+            {
+                case CompactBinaryFieldType.None:
+                case CompactBinaryFieldType.Null:
+                    return Array.Empty<byte>();
+                case CompactBinaryFieldType.Object:
+                case CompactBinaryFieldType.UniformObject:
+                case CompactBinaryFieldType.Array:
+                case CompactBinaryFieldType.UniformArray:
+                    throw new NotImplementedException();
+                case CompactBinaryFieldType.Binary:
+                    throw new NotImplementedException();
+                case CompactBinaryFieldType.String:
+                    if (member is string s)
+                    {
+                        return Encoding.UTF8.GetBytes(s);
+                    }
+                    break;
+                case CompactBinaryFieldType.IntegerPositive:
+                case CompactBinaryFieldType.IntegerNegative:
+                    throw new NotImplementedException();
+                case CompactBinaryFieldType.Float32:
+                case CompactBinaryFieldType.Float64:
+                    throw new NotImplementedException();
+                case CompactBinaryFieldType.BoolFalse:
+                case CompactBinaryFieldType.BoolTrue:
+                    return Array.Empty<byte>();
+                case CompactBinaryFieldType.CompactBinaryAttachment:
+                case CompactBinaryFieldType.BinaryAttachment:
+                case CompactBinaryFieldType.Hash:
+                    if (member is BlobIdentifier blob)
+                    {
+                        return blob.HashData;
+                    }
+
+                    break;
+                case CompactBinaryFieldType.Uuid:
+                case CompactBinaryFieldType.DateTime:
+                case CompactBinaryFieldType.TimeSpan:
+                case CompactBinaryFieldType.ObjectId:
+                case CompactBinaryFieldType.CustomById:
+                case CompactBinaryFieldType.CustomByName:
+                    throw new NotImplementedException();
+                default:
+                    throw new ArgumentOutOfRangeException(nameof(fieldType), fieldType, null);
+            }
+            throw new ArgumentException($"{member} of type {typeof(T).Name} not convert-able to field-type {fieldType}", nameof(member));
+        }
+
+        private void SetName(ref Field field, string? fieldName)
+        {
+            if (string.IsNullOrEmpty(fieldName))
+                return;
+
+            // TODO: We could add checks similar to what we do in C++ to make sure the writer is used correctly
+            
+            uint nameLength = (uint) fieldName.Length;
+            uint nameBytesCount = BitUtils.MeasureVarUInt(nameLength);
+
+            byte[] buffer = new byte[nameLength + nameBytesCount];
+            BitUtils.WriteVarUInt(nameLength, buffer);
+            int bytesWritten = Encoding.ASCII.GetBytes(fieldName, 0, fieldName.Length, buffer, (int)nameBytesCount);
+            if (bytesWritten != nameLength)
+                throw new Exception("Failed to write name into buffer");
+
+            field.Name = buffer;
+        }
+    }
+}
diff --git a/Engine/Source/Programs/Horde/HordeStorage/Jupiter.Common/Implementation/ContentHash.cs b/Engine/Source/Programs/Horde/HordeStorage/Jupiter.Common/Implementation/ContentHash.cs
new file mode 100644
--- /dev/null
+++ b/Engine/Source/Programs/Horde/HordeStorage/Jupiter.Common/Implementation/ContentHash.cs
@@ -0,0 +1,115 @@
+﻿// Copyright Epic Games, Inc. All Rights Reserved.
+
+using System;
+using System.IO;
+using System.Text;
+using System.Threading.Tasks;
+using Blake3;
+using Newtonsoft.Json;
+
+namespace Jupiter.Implementation
+{
+    [JsonConverter(typeof(ContentHashConverter))]
+    public class ContentHash : IEquatable<ContentHash>, IEquatable<byte[]>
+    {
+        protected readonly ByteArrayComparer Comparer = new ByteArrayComparer();
+        protected readonly byte[] Identifier;
+        public static int HashLength = 20;
+
+        public ContentHash(byte[] identifier)
+        {
+            Identifier = identifier;
+            if (identifier.Length != HashLength)
+                throw new ArgumentException("Supplied identifier was not 20 bytes, this is not a valid identifier", nameof(identifier));
+        }
+
+        [JsonConstructor]
+        public ContentHash(string identifier)
+        {
+            if (identifier == null)
+                throw new ArgumentNullException(nameof(identifier));
+
+            byte[] byteIdentifier = StringUtils.ToHashFromHexString(identifier);
+            Identifier = byteIdentifier;
+            if (byteIdentifier.Length != HashLength)
+                throw new ArgumentException("Supplied identifier was not 20 bytes, this is not a valid identifier", nameof(identifier));
+        }
+
+        public byte[] HashData
+        {
+            get { return Identifier; }
+        }
+
+        public override int GetHashCode()
+        {
+            return Comparer.GetHashCode(Identifier);
+        }
+
+        public bool Equals(ContentHash? other)
+        {
+            if (other == null)
+                return false;
+
+            return Comparer.Equals(Identifier, other.Identifier);
+        }
+
+        public bool Equals(byte[]? other)
+        {
+            if (other == null)
+                return false;
+            return Comparer.Equals(Identifier, other);
+        }
+
+        public override bool Equals(object? obj)
+        {
+            if (ReferenceEquals(null, obj))
+            {
+                return false;
+            }
+
+            if (ReferenceEquals(this, obj))
+            {
+                return true;
+            }
+
+            if (obj.GetType() != this.GetType())
+            {
+                return false;
+            }
+
+            return Equals((ContentHash) obj);
+        }
+
+        public override string ToString()
+        {
+            return StringUtils.FormatAsHexString(Identifier);
+        }
+
+        public static ContentHash FromBlob(byte[] blobMemory)
+        {
+            using Hasher hasher = Hasher.New();
+            hasher.UpdateWithJoin(blobMemory);
+            Hash blake3Hash = hasher.Finalize();
+
+            // we only keep the first 20 bytes of the Blake3 hash
+            Span<byte> hash = blake3Hash.AsSpan().Slice(0, HashLength);
+
+            return new ContentHash(hash.ToArray());
+        }
+    }
+    
+    public class ContentHashConverter : JsonConverter<ContentHash>
+    {
+        public override void WriteJson(JsonWriter writer, ContentHash? value, JsonSerializer serializer)
+        {
+            writer.WriteValue(value!.ToString());
+        }
+
+        public override ContentHash ReadJson(JsonReader reader, Type objectType, ContentHash? existingValue, bool hasExistingValue, JsonSerializer serializer)
+        {
+            string? s = (string?)reader.Value;
+
+            return new ContentHash(s!);
+        }
+    }
+}
diff --git a/Engine/Source/Programs/Horde/HordeStorage/Jupiter.Common/Implementation/InstrumentedStream.cs b/Engine/Source/Programs/Horde/HordeStorage/Jupiter.Common/Implementation/InstrumentedStream.cs
new file mode 100644
--- /dev/null
+++ b/Engine/Source/Programs/Horde/HordeStorage/Jupiter.Common/Implementation/InstrumentedStream.cs
@@ -0,0 +1,104 @@
+﻿using System;
+using System.Collections.Generic;
+using System.IO;
+using System.Linq;
+using System.Threading;
+using System.Threading.Tasks;
+using Serilog;
+
+namespace Jupiter.Common.Implementation
+{
+    /// <summary>
+    /// Wraps a stream and adds instrumentation for throughput which is written to a counter
+    /// </summary>
+    public class InstrumentedStream : Stream
+    {
+        private readonly Stream _streamImplementation;
+        private readonly string _sourceIdentifier;
+        private bool _started = false;
+        private long _totalBytesWritten = 0;
+        private DateTime _consumeStartedAt;
+
+        private readonly ILogger _logger = Log.ForContext<InstrumentedStream>();
+
+        public InstrumentedStream(Stream s, string sourceIdentifier)
+        {
+            _streamImplementation = s;
+            _sourceIdentifier = sourceIdentifier;
+        }
+
+        public override void Close()
+        {
+            TimeSpan duration = DateTime.Now - _consumeStartedAt;
+            double rate = _totalBytesWritten / duration.TotalSeconds;
+
+            StatsdClient.DogStatsd.Histogram("jupiter.stream_throughput", rate, tags: new string[] {"sourceIdentifier:" + _sourceIdentifier});
+            base.Close();
+        }
+
+        public override void Flush()
+        {
+            _streamImplementation.Flush();
+        }
+
+        public override int Read(byte[] buffer, int offset, int count)
+        {
+            if (!_started)
+            {
+                _consumeStartedAt = DateTime.Now;
+                _started = true;
+            }
+
+            Interlocked.Add(ref _totalBytesWritten, count);
+            return _streamImplementation.Read(buffer, offset, count);
+        }
+
+        public override long Seek(long offset, SeekOrigin origin)
+        {
+            return _streamImplementation.Seek(offset, origin);
+        }
+
+        public override void SetLength(long value)
+        {
+            _streamImplementation.SetLength(value);
+        }
+
+        public override void Write(byte[] buffer, int offset, int count)
+        {
+            if (!_started)
+            {
+                _consumeStartedAt = DateTime.Now;
+                _started = true;
+            }
+
+            Interlocked.Add(ref _totalBytesWritten, count);
+            _streamImplementation.Write(buffer, offset, count);
+        }
+
+        public override bool CanRead
+        {
+            get { return _streamImplementation.CanRead; }
+        }
+
+        public override bool CanSeek
+        {
+            get { return _streamImplementation.CanSeek; }
+        }
+
+        public override bool CanWrite
+        {
+            get { return _streamImplementation.CanWrite; }
+        }
+
+        public override long Length
+        {
+            get { return _streamImplementation.Length; }
+        }
+
+        public override long Position
+        {
+            get { return _streamImplementation.Position; }
+            set { _streamImplementation.Position = value; }
+        }
+    }
+}
diff --git a/Engine/Source/Programs/Horde/HordeStorage/Jupiter.Common/Implementation/KeyId.cs b/Engine/Source/Programs/Horde/HordeStorage/Jupiter.Common/Implementation/KeyId.cs
new file mode 100644
--- /dev/null
+++ b/Engine/Source/Programs/Horde/HordeStorage/Jupiter.Common/Implementation/KeyId.cs
@@ -0,0 +1,99 @@
+﻿using System;
+using System.ComponentModel;
+using System.Globalization;
+using Newtonsoft.Json;
+
+namespace Jupiter.Implementation
+{
+    [JsonConverter(typeof(KeyIdJsonConverter))]
+    [TypeConverter(typeof(KeyIdTypeConverter))]
+
+    public readonly struct KeyId: IEquatable<KeyId>
+    {
+        public KeyId(string key)
+        {
+             _text = key;
+ 
+            if (_text.Length == 0)
+            {
+				throw new ArgumentException("Keys must have at least one character.");
+            }
+ 
+            const int MaxLength = 250;
+            if (_text.Length > MaxLength)
+            {
+                throw new ArgumentException($"Keys may not be longer than {MaxLength} characters");
+            }
+        }
+
+        readonly string _text;
+
+        public bool Equals(KeyId other)
+        {
+            return string.Equals(_text , other._text, StringComparison.Ordinal);
+        }
+
+        public override bool Equals(object? obj)
+        {
+            return obj is KeyId other && Equals(other);
+        }
+
+        public override int GetHashCode()
+        {
+            return _text.GetHashCode();
+        }
+
+        public override string ToString()
+        {
+            return _text;
+        }
+
+        public static bool operator ==(KeyId left, KeyId right)
+        {
+            return left.Equals(right);
+        }
+ 
+        public static bool operator !=(KeyId left, KeyId right)
+        {
+            return !left.Equals(right);
+        }
+    }
+
+    public class KeyIdJsonConverter: JsonConverter<KeyId>
+    {
+        public override void WriteJson(JsonWriter writer, KeyId value, JsonSerializer serializer)
+        {
+            writer.WriteValue(value!.ToString());
+        }
+
+        public override KeyId ReadJson(JsonReader reader, Type objectType, KeyId existingValue, bool hasExistingValue, Newtonsoft.Json.JsonSerializer serializer)
+        {
+            string? s = (string?)reader.Value;
+
+            return new KeyId(s!);
+        }
+    }
+
+    public sealed class KeyIdTypeConverter : TypeConverter
+    {
+        /// <inheritdoc/>
+        public override bool CanConvertFrom(ITypeDescriptorContext? context, Type sourceType)
+        {
+            if (sourceType == typeof(string))
+                return true;
+
+            return base.CanConvertFrom(context, sourceType);
+        }
+ 
+        public override object? ConvertFrom(ITypeDescriptorContext? context, CultureInfo? culture, object value)  
+        {
+            if (value is string s)
+            {
+                return new KeyId(s);
+            }
+
+            return base.ConvertFrom(context, culture, value);  
+        }
+    }
+}
+
diff --git a/Engine/Source/Programs/Horde/HordeStorage/Jupiter.Common/Implementation/NamespaceId.cs b/Engine/Source/Programs/Horde/HordeStorage/Jupiter.Common/Implementation/NamespaceId.cs
new file mode 100644
--- /dev/null
+++ b/Engine/Source/Programs/Horde/HordeStorage/Jupiter.Common/Implementation/NamespaceId.cs
@@ -0,0 +1,123 @@
+﻿using System;
+using System.ComponentModel;
+using System.Globalization;
+using Newtonsoft.Json;
+
+namespace Jupiter.Implementation
+{
+    [JsonConverter(typeof(NamespaceIdJsonConverter))]
+    [TypeConverter(typeof(NamespaceIdTypeConverter))]
+    public readonly struct NamespaceId: IEquatable<NamespaceId>
+    {
+        public NamespaceId(string ns)
+        {
+             _text = ns;
+ 
+            if (_text.Length == 0)
+            {
+				throw new ArgumentException("Namespaces must have at least one character.");
+            }
+ 
+            const int MaxLength = 64;
+            if (_text.Length > MaxLength)
+            {
+                throw new ArgumentException($"Namespaces may not be longer than {MaxLength} characters");
+            }
+ 
+            for (int Idx = 0; Idx < _text.Length; Idx++)
+            {
+                if (!IsValidCharacter(_text[Idx]))
+                {
+                    throw new ArgumentException($"{_text} is not a valid namespace id");
+                }
+            }
+        }
+
+        static bool IsValidCharacter(char c)
+        {
+            if (c >= 'a' && c <= 'z')
+            {
+                return true;
+            }
+            if (c >= '0' && c <= '9')
+            {
+                return true;
+            }
+            if (c == '-' || c == '_' || c == '.')
+            {
+                return true;
+            }
+            return false;
+        }
+
+        readonly string _text;
+
+        public bool Equals(NamespaceId other)
+        {
+            return string.Equals(_text , other._text, StringComparison.Ordinal);
+        }
+
+        public override bool Equals(object? obj)
+        {
+            return obj is NamespaceId other && Equals(other);
+        }
+
+        public override int GetHashCode()
+        {
+            return _text.GetHashCode();
+        }
+
+        public override string ToString()
+        {
+            return _text;
+        }
+
+        public static bool operator ==(NamespaceId left, NamespaceId right)
+        {
+            return left.Equals(right);
+        }
+ 
+        public static bool operator !=(NamespaceId left, NamespaceId right)
+        {
+            return !left.Equals(right);
+        }
+    }
+
+    public class NamespaceIdJsonConverter: JsonConverter<NamespaceId>
+    {
+        public override void WriteJson(JsonWriter writer, NamespaceId value, JsonSerializer serializer)
+        {
+            writer.WriteValue(value!.ToString());
+        }
+
+        public override NamespaceId ReadJson(JsonReader reader, Type objectType, NamespaceId existingValue, bool hasExistingValue, Newtonsoft.Json.JsonSerializer serializer)
+        {
+            string? s = (string?)reader.Value;
+
+            return new NamespaceId(s!);
+        }
+    }
+
+    public sealed class NamespaceIdTypeConverter : TypeConverter
+    {
+        /// <inheritdoc/>
+        public override bool CanConvertFrom(ITypeDescriptorContext? context, Type sourceType)
+        {
+            if (sourceType == typeof(string))
+                return true;
+
+            return base.CanConvertFrom(context, sourceType);
+        }
+ 
+        public override object? ConvertFrom(ITypeDescriptorContext? context, CultureInfo? culture, object value)  
+        {
+            if (value is string s)
+            {
+                return new NamespaceId(s);
+            }
+
+            return base.ConvertFrom(context, culture, value);  
+        }
+    }
+}
+
diff --git a/Engine/Source/Programs/Horde/HordeStorage/Jupiter.Common/Implementation/TransactionEvent.cs b/Engine/Source/Programs/Horde/HordeStorage/Jupiter.Common/Implementation/TransactionEvent.cs
new file mode 100644
--- /dev/null
+++ b/Engine/Source/Programs/Horde/HordeStorage/Jupiter.Common/Implementation/TransactionEvent.cs
@@ -0,0 +1,88 @@
+﻿// Copyright Epic Games, Inc. All Rights Reserved.
+
+using System.Collections.Generic;
+using System.ComponentModel.DataAnnotations;
+using JsonSubTypes;
+using Newtonsoft.Json;
+
+namespace Jupiter.Implementation
+{
+    public class AddTransactionEvent : TransactionEvent
+    {
+        [Required] public BlobIdentifier[] Blobs { get; }
+
+        public Dictionary<string, object>? Metadata { get; }
+
+        [JsonConstructor]
+        public AddTransactionEvent(string name, string bucket, BlobIdentifier[] blobs, Dictionary<string, object>? metadata = null, long? identifier = null, long? nextIdentifier = null) 
+            : base("Add", identifier, nextIdentifier, name, bucket)
+        {
+            Blobs = blobs;
+            Metadata = metadata;
+        }
+
+        public AddTransactionEvent(string name, string bucket, BlobIdentifier[] blobs, List<string> locations, Dictionary<string, object>? metadata = null, long? identifier = null, long? nextIdentifier = null)
+            : base("Add", identifier, nextIdentifier, name, bucket, locations)
+        {
+            Blobs = blobs;
+            Metadata = metadata;
+        }
+    }
+
+    public class RemoveTransactionEvent : TransactionEvent
+    {
+        [JsonConstructor]
+        public RemoveTransactionEvent(string name, string bucket, long? identifier = null, long? nextIdentifier = null) 
+            : base("Remove", identifier, nextIdentifier, name, bucket)
+        {
+        }
+
+        public RemoveTransactionEvent(string name, string bucket, List<string> locations, long? identifier = null, long? nextIdentifier = null)
+            : base("Remove", identifier, nextIdentifier, name, bucket, locations)
+        {
+        }
+    }
+
+    public interface IReplicationEvent
+    {
+        public long? Identifier { get; }
+        public long? NextIdentifier { get; }
+    }
+
+    [JsonConverter(typeof(JsonSubtypes), "type")]
+    [JsonSubtypes.KnownSubType(typeof(AddTransactionEvent), "Add")]
+    [JsonSubtypes.KnownSubType(typeof(RemoveTransactionEvent), "Remove")]
+    public abstract class TransactionEvent : IReplicationEvent
+    {
+        /// <summary>
+        /// Constructor for use by the controller
+        /// </summary>
+        /// <param name="type">The type of transaction event, either add or remove</param>
+        /// <param name="identifier">The starting offset for this event in the log</param>
+        /// <param name="nextIdentifier">The end offset of this event (can be used to fetch the next event)</param>
+        /// <param name="name">The name of the ref key that is being added for this event</param>
+        /// <param name="bucket">The bucket identifier of this event</param>
+        /// <param name="locations">Optional specification of which locations this event has been seen in</param>
+        protected TransactionEvent(string type, long? identifier, long? nextIdentifier, string name, string bucket, List<string>? locations = null)
+        {
+            Type = type;
+            Identifier = identifier;
+            NextIdentifier = nextIdentifier;
+            Name = name;
+            Bucket = bucket;
+            Locations = locations ?? new List<string>();
+        }
+
+        public string Type { get; }
+        public long? Identifier { get; }
+        public long? NextIdentifier { get; }
+
+        [Required]
+        public string Name { get; }
+        [Required]
+        public string Bucket { get; }
+
+        public List<string> Locations { get; }
+
+    }
+}
diff --git a/Engine/Source/Programs/Horde/HordeStorage/Jupiter.Common/Implementation/TreeHash.cs b/Engine/Source/Programs/Horde/HordeStorage/Jupiter.Common/Implementation/TreeHash.cs
new file mode 100644
--- /dev/null
+++ b/Engine/Source/Programs/Horde/HordeStorage/Jupiter.Common/Implementation/TreeHash.cs
@@ -0,0 +1,118 @@
+﻿// Copyright Epic Games, Inc. All Rights Reserved.
+
+using System;
+using System.Collections.Generic;
+using System.ComponentModel;
+using System.Globalization;
+using Blake3;
+using Newtonsoft.Json;
+
+namespace Jupiter.Implementation
+{
+    [JsonConverter(typeof(TreeHashConverter))]
+    [TypeConverter(typeof(TreeHashTypeConverter))]
+    public class TreeHash : ContentHash, IEquatable<TreeHash>
+    {
+        public TreeHash(byte[] identifier) : base(identifier)
+        {
+        }
+
+        [JsonConstructor]
+        public TreeHash(string identifier) : base(identifier)
+        {
+        }
+
+        public override int GetHashCode()
+        {
+            return Comparer.GetHashCode(Identifier);
+        }
+
+        public bool Equals(TreeHash? other)
+        {
+            if (other == null)
+                return false;
+            return Comparer.Equals(Identifier, other.Identifier);
+        }
+
+        public override bool Equals(object? obj)
+        {
+            if (ReferenceEquals(null, obj))
+            {
+                return false;
+            }
+
+            if (ReferenceEquals(this, obj))
+            {
+                return true;
+            }
+
+            if (obj.GetType() != this.GetType())
+            {
+                return false;
+            }
+
+            return Equals((TreeHash) obj);
+        }
+
+         public static TreeHash FromTreeContents(TreeHash[] treeHashes, BlobIdentifier[] blobHashes)
+        {
+            int countOfHashes = treeHashes.Length + blobHashes.Length;
+            byte[] bytes = new byte[countOfHashes * HashLength];
+            int hashIndex = 0;
+            foreach (TreeHash tree in treeHashes)
+            {
+                Array.Copy(tree.HashData, 0, bytes, hashIndex * HashLength, HashLength);
+                ++hashIndex;
+            }
+            foreach (BlobIdentifier blob in blobHashes)
+            {
+                Array.Copy(blob.HashData, 0, bytes, hashIndex * HashLength, HashLength);
+                ++hashIndex;
+            }
+
+            Hash blake3Hash = Hasher.Hash(bytes);
+
+            // we only keep the first 20 bytes of the Blake3 hash
+            TreeHash hash = new TreeHash(blake3Hash.AsSpan().Slice(0, HashLength).ToArray());
+            return hash;
+        }
+    }
+
+    public class TreeHashConverter : JsonConverter<TreeHash>
+    {
+        public override void WriteJson(JsonWriter writer, TreeHash? value, JsonSerializer serializer)
+        {
+            writer.WriteValue(value!.ToString());
+        }
+
+        public override TreeHash ReadJson(JsonReader reader, Type objectType, TreeHash? existingValue, bool hasExistingValue, Newtonsoft.Json.JsonSerializer serializer)
+        {
+            string? s = (string?)reader.Value;
+
+            return new TreeHash(s!);
+        }
+    }
+
+
+    public sealed class TreeHashTypeConverter : TypeConverter
+    {
+        /// <inheritdoc/>
+        public override bool CanConvertFrom(ITypeDescriptorContext? context, Type sourceType)
+        {
+            if (sourceType == typeof(string))
+                return true;
+
+            return base.CanConvertFrom(context, sourceType);
+        }
+
+        public override object? ConvertFrom(ITypeDescriptorContext? context, CultureInfo? culture, object value)
+        {
+            if (value is string s)
+            {
+                return new TreeHash(s);
+            }
+
+            return base.ConvertFrom(context, culture, value);
+        }
+    }
+}
diff --git a/Engine/Source/Programs/Horde/HordeStorage/Jupiter.Common/Jupiter.Common.csproj b/Engine/Source/Programs/Horde/HordeStorage/Jupiter.Common/Jupiter.Common.csproj
new file mode 100644
--- /dev/null
+++ b/Engine/Source/Programs/Horde/HordeStorage/Jupiter.Common/Jupiter.Common.csproj
@@ -0,0 +1,55 @@
+﻿<Project Sdk="Microsoft.NET.Sdk.Web">
+
+  <PropertyGroup>
+    <TargetFramework>net6.0</TargetFramework>
+    <ApplicationIcon />
+    <OutputType>Library</OutputType>
+    <StartupObject />
+    <Nullable>enable</Nullable>
+    <TreatWarningsAsErrors>true</TreatWarningsAsErrors>
+    <WarningsAsErrors />
+    <OpenApiGenerateDocuments>false</OpenApiGenerateDocuments>
+    <LangVersion>10</LangVersion>
+  </PropertyGroup>
+
+  <ItemGroup>
+    <PackageReference Include="AspNetCore.HealthChecks.Publisher.Datadog" Version="5.0.1" />
+    <PackageReference Include="AWSSDK.Extensions.NETCore.Setup" Version="3.7.0.1" />
+    <PackageReference Include="AWSSDK.SecretsManager" Version="3.7.1.39" />
+    <PackageReference Include="AWSSDK.SecurityToken" Version="3.7.1.4" />
+    <PackageReference Include="Blake3" Version="0.2.1" />
+    <PackageReference Include="Datadog.Trace" Version="1.23.0" />
+    <PackageReference Include="IdentityServer4.AccessTokenValidation" Version="3.0.1" />
+    <PackageReference Include="JsonSubTypes" Version="1.8.0" />
+    <PackageReference Include="Log4net.Appender.Serilog" Version="1.0.11" />
+    <PackageReference Include="Microsoft.AspNet.WebApi.Client" Version="5.2.7" />
+    <PackageReference Include="Microsoft.AspNetCore.Authentication.JwtBearer" Version="6.0.0" />
+    <PackageReference Include="Microsoft.AspNetCore.Mvc.NewtonsoftJson" Version="6.0.0" />
+    <PackageReference Include="Microsoft.Extensions.Diagnostics.HealthChecks" Version="6.0.0" />
+    <PackageReference Include="Microsoft.Extensions.Hosting" Version="6.0.0" />
+    <PackageReference Include="NetEscapades.Configuration.Yaml" Version="2.1.0" />
+    <PackageReference Include="Okta.AspNetCore" Version="3.5.1" />
+    <PackageReference Include="RestSharp" Version="106.11.7" />
+    <PackageReference Include="RestSharp.Serializers.NewtonsoftJson" Version="106.11.7" />
+    <PackageReference Include="Serilog" Version="2.10.0" />
+    <PackageReference Include="Serilog.AspNetCore" Version="3.4.0" />
+    <PackageReference Include="Serilog.Formatting.Compact" Version="1.1.0" />
+    <PackageReference Include="Serilog.Settings.Configuration" Version="3.1.0" />
+    <PackageReference Include="Serilog.Sinks.File" Version="4.1.0" />
+    <PackageReference Include="SerilogAnalyzer" Version="0.15.0" />
+    <PackageReference Include="SharpYaml" Version="1.6.6" />
+    <PackageReference Include="Swashbuckle.AspNetCore" Version="5.6.3" />
+    <PackageReference Include="Swashbuckle.AspNetCore.ReDoc" Version="5.6.3" />
+    <PackageReference Include="WebApiContrib.Core.Formatter.Bson" Version="2.1.0" />
+    <PackageReference Include="WebApiContrib.Core.Formatter.Yaml" Version="2.1.0" />
+  </ItemGroup>
+
+  <ItemGroup>
+    <None Include="Shared\Oodle\2.9.3\linux\lib\liboo2corelinux64.so.9" Link="liboo2corelinux64.so.9" CopyToOutputDirectory="PreserveNewest" />
+    <None Include="Shared\Oodle\2.9.3\linux\lib\liboo2corelinux.so.9" Link="liboo2corelinux.so.9" CopyToOutputDirectory="PreserveNewest" />
+    <None Include="Shared\Oodle\2.9.3\win\redist\oo2core_9_win64.dll" Link="oo2core_9_win64.dll" CopyToOutputDirectory="PreserveNewest" />
+    <None Include="Shared\Oodle\2.9.3\win\redist\oo2core_9_win32.dll" Link="oo2core_9_win32.dll" CopyToOutputDirectory="PreserveNewest" />
+  </ItemGroup>
+
+</Project>
+ 
\ No newline at end of file
diff --git a/Engine/Source/Programs/Horde/HordeStorage/Jupiter.Common/PollingService.cs b/Engine/Source/Programs/Horde/HordeStorage/Jupiter.Common/PollingService.cs
new file mode 100644
--- /dev/null
+++ b/Engine/Source/Programs/Horde/HordeStorage/Jupiter.Common/PollingService.cs
@@ -0,0 +1,143 @@
+﻿// Copyright Epic Games, Inc. All Rights Reserved.
+
+using System;
+using System.Threading;
+using System.Threading.Tasks;
+using Microsoft.Extensions.Hosting;
+using Serilog;
+
+namespace Jupiter
+{
+    public abstract class PollingService<T> : IHostedService
+    {
+        private struct ThreadState
+        {
+            public string ServiceName { get; set; }
+            public TimeSpan PollFrequency { get; set; }
+
+            public CancellationToken StopPollingToken { get; set; }
+            public T ServiceState { get; set; }
+            public PollingService<T> Instance { get; set; }
+        }
+
+        private readonly ILogger _logger = Log.ForContext<PollingService<T>>();
+        private readonly string _serviceName;
+        private readonly TimeSpan _pollFrequency;
+        private readonly T _state;
+        private readonly Thread _pollingThread = new Thread(OnUpdate);
+        private readonly CancellationTokenSource _stopPolling = new CancellationTokenSource();
+
+        public PollingService(string serviceName, TimeSpan pollFrequency, T state)
+        {
+            _serviceName = serviceName;
+            _pollFrequency = pollFrequency;
+            _state = state;
+        }
+
+        public bool Running
+        {
+            get { return _pollingThread.IsAlive; }
+        }
+
+        public T State
+        {
+            get { return _state; }
+        }
+
+        public virtual bool ShouldStartPolling()
+        {
+            return true;
+        }
+
+        public Task StartAsync(CancellationToken cancellationToken)
+        {
+            bool shouldPoll = ShouldStartPolling();
+            _logger.Information("Polling service {Service} initialized {@State} , will poll: {WillPoll}.", _serviceName, _state, shouldPoll);
+
+            if (shouldPoll)
+            {
+                _pollingThread.Start(new ThreadState
+                {
+                    ServiceName = _serviceName, 
+                    PollFrequency = _pollFrequency, 
+                    ServiceState = _state, 
+                    StopPollingToken = _stopPolling.Token,
+                    Instance = this,
+                });
+            }
+
+            return Task.CompletedTask;
+        }
+
+        private static void OnUpdate(object? state)
+        {
+            ThreadState? ts = (ThreadState?)state;
+            if (ts == null)
+            {
+                throw new ArgumentNullException(nameof(ts));
+            }
+            ThreadState threadState = ts.Value;
+
+            PollingService<T> instance = threadState.Instance;
+            string serviceName = threadState.ServiceName;
+            CancellationToken stopPollingToken = threadState.StopPollingToken;
+            ILogger logger = Log.ForContext<PollingService<T>>();
+
+            while (!stopPollingToken.IsCancellationRequested)
+            {
+                var startTime = DateTime.Now;
+                try
+                {
+                    bool _ = instance.OnPoll(threadState.ServiceState, stopPollingToken).Result;
+                }
+                catch (AggregateException e)
+                {
+                    logger.Error(e, "{Service} Aggregate exception in polling thread", serviceName);
+                    foreach (Exception inner in e.InnerExceptions)
+                    {
+                        logger.Error(inner, "{Service} inner exception in polling thread. Trace: {StackTrace}", serviceName, inner.StackTrace);
+                        
+                    }
+                }
+                catch (Exception e)
+                {
+                    logger.Error(e, "{Service} Exception in polling thread", serviceName);
+                }
+                finally
+                {
+                    TimeSpan duration = DateTime.Now - startTime;
+                    TimeSpan pollFrequency = threadState.PollFrequency;
+
+                    // if we spent less then the poll frequency we sleep to avoid ping the remote server to much 
+                    if (duration < pollFrequency)
+                    {
+                        TimeSpan remainingDuration = pollFrequency - duration;
+                        logger.Information("{Service} ran polled for {Duration} Sleeping for {RemainingDuration} to stay within poll frequency", serviceName, duration, remainingDuration);
+
+                        Thread.Sleep(remainingDuration);
+                    }
+                }
+            }
+        }
+
+        public abstract Task<bool> OnPoll(T state, CancellationToken cancellationToken);
+
+        protected virtual Task OnStopping(T state)
+        {
+            return Task.CompletedTask;
+        }
+
+        public async Task StopAsync(CancellationToken cancellationToken)
+        {
+            _logger.Information("{Service} poll service stopping.", _serviceName);
+
+            await OnStopping(_state);
+
+            _stopPolling.Cancel();
+            if (_pollingThread.IsAlive)
+            {
+                _pollingThread.Join();
+            }
+        }
+    }
+}
