Moved Horde.Storage from git to p4

diff --git a/Engine/Source/Programs/Horde/HordeStorage/Jupiter.Common/Shared/Oodle/2.9.3/linux/examples/cbvector.h b/Engine/Source/Programs/Horde/HordeStorage/Jupiter.Common/Shared/Oodle/2.9.3/linux/examples/cbvector.h
new file mode 100644
--- /dev/null
+++ b/Engine/Source/Programs/Horde/HordeStorage/Jupiter.Common/Shared/Oodle/2.9.3/linux/examples/cbvector.h
@@ -0,0 +1,1217 @@
+#pragma once
+
+#ifndef CB_VECTOR_H_INCLUDED
+#define CB_VECTOR_H_INCLUDED
+
+/***********************
+
+
+This license applies to all software available on cbloom.com unless otherwise stated.
+
+Copyright (c) 1998-2020, Charles Bloom
+
+This software is provided 'as-is', without any express or implied
+warranty. In no event will the authors be held liable for any damages
+arising from the use of this software.
+
+Permission is granted to anyone to use this software for any purpose,
+including commercial applications, and to alter it and redistribute it
+freely, subject to the following restrictions:
+
+   1. The origin of this software must not be misrepresented; you must not
+   claim that you wrote the original software. If you use this software
+   in a product, an acknowledgment in the product documentation would be
+   appreciated but is not required.
+
+   2. Altered source versions must be plainly marked as such, and must not be
+   misrepresented as being the original software.
+
+   3. This notice may not be removed or altered from any source
+   distribution.
+
+-------------------------------------------
+
+cbvector.h
+
+Single file self-contained vector implementation.
+
+This header defined a class name CB_VECTOR in namespace CB_VECTOR_NAMESPACE.
+If you don't #define those macros, they will default to "vector" and "", respectively.
+
+cbvector.h was written by Charles Bloom.
+Public domain.
+Don't be a jerk.
+
+---------------------------
+
+This code was ripped out of cblib/vector_flex and not cleaned up.  It's ugly.
+
+The use of size_t/intptr_t/ptrdiff_t is often wrong and always ugly.
+
+The split into vector & vector_base is pointless obfuscation in this context.
+
+---------------------------
+
+Currently tested on these compilers :
+
+MSVC 2005
+GCC 4.1.1 (PS3)
+
+***************************/
+
+// options :
+
+// class name :
+#ifndef CB_VECTOR
+#define CB_VECTOR   vector
+#endif
+
+// var type for m_size :
+#ifndef CB_SIZETYPE
+#define CB_SIZETYPE size_t
+#endif
+
+// assert :
+#ifndef CB_ASSERT
+#define CB_ASSERT(exp)
+#endif
+
+// compiler assert :
+#ifndef CB_COMPILER_ASSERT
+#define CB_COMPILER_ASSERT(exp)
+#endif
+
+// assert used for the vector extend malloc :
+#ifndef CB_ASSERT_MALLOC
+#define CB_ASSERT_MALLOC(exp)       do { if ( ! (exp) ) exit(20); } while(0)
+#endif
+
+// alloc memory :
+#ifndef CB_ALLOC
+#define CB_ALLOC(size)  malloc(size)
+#endif
+
+// free mem :
+#ifndef CB_FREE
+#define CB_FREE(ptr,size)   free(ptr)
+#endif
+
+// min :
+#ifndef CB_MIN
+#define CB_MIN(a,b) ( (a) < (b) ? (a) : (b) )
+#endif
+
+// max :
+#ifndef CB_MAX
+#define CB_MAX(a,b) ( (a) > (b) ? (a) : (b) )
+#endif
+
+// option : use swap instead of copy when extending vector
+//#define CB_VECTOR_USE_SWAP
+
+// option : set max grow size : (otherwise doubles)
+// you can define CB_VECTOR_MAX_GROW_BYTES yourself before including vector.h if you like
+#ifdef CB_VECTOR_MAX_GROW_BYTES
+#if CB_VECTOR_MAX_GROW_BYTES == 0
+// it was defined but not a valid value, set default :
+#define CB_VECTOR_MAX_GROW_BYTES    (1024*1024)
+#endif  //CB_VECTOR_MAX_GROW_BYTES
+#endif
+
+// namespace :
+//  blank namespace ensures that includes with different options don't conflict in the linker
+// define CB_VECTOR_NO_NAMESPACE if you really want none (not recommended)
+#ifndef CB_VECTOR_NO_NAMESPACE
+#ifndef CB_VECTOR_NAMESPACE
+#define CB_VECTOR_NAMESPACE     // blank
+#endif
+#endif
+
+//#include <new> // need this to get placement new
+
+//=======================================================
+
+#ifdef _MSC_VER
+#define CB_VECTOR_FORCEINLINE __forceinline
+#elif defined(__GNUC__)
+#define CB_VECTOR_FORCEINLINE inline __attribute((always_inline))
+#else
+#define CB_VECTOR_FORCEINLINE
+#endif
+
+//=======================================================
+
+#ifdef CB_VECTOR_NAMESPACE
+namespace CB_VECTOR_NAMESPACE
+{
+#endif
+
+#ifdef _MSC_VER
+#pragma warning(push)
+#pragma warning(disable : 4345)
+//d:\Exoddus\Code\Engine\Core\cb_entry_array.h(83) : warning C4345: behavior change: an object of POD type constructed with an initializer of the form () will be default-initialized
+#pragma warning(disable : 4100) // unreferenced formal parameter
+#pragma warning(disable : 4127) // conditional on constant
+#pragma warning(disable : 4328) // indirection alignment of formal parameter parameter_number (alignment_value_for_parameter) is greater than the actual argument
+#endif
+
+//#pragma pack(push)
+//#pragma pack(4)
+
+namespace cb_entry_array
+{
+    /*******************
+
+    cb_entry_array
+
+    This is a set of helper template functions for arrays of a templated type.
+    It's interesting in a few ways.  
+
+    *******************/
+
+    //-----------------------------------------------------------------------------------------------
+
+    // EntryCopy
+    //  like memcpy
+    //  pTo and pFm must not overlap
+    template <typename t_entry>
+    static inline void copy(t_entry * pTo,const t_entry * pFm,const size_t count)
+    {
+        // assert we don't overlap in a bad way :
+        CB_ASSERT( pTo != pFm || count == 0 );
+        CB_ASSERT( pTo < pFm || pTo >= pFm + count );
+
+        for(size_t i=0;i<count;i++)
+        {
+            pTo[i] = pFm[i];
+        }
+    }
+
+    //-----------------------------------------------------------------------------------------------
+
+    // EntryMove
+    //  like memmove
+    //  pTo and pFm may overlap
+    template <typename t_entry>
+    static inline void move(t_entry * pTo,const t_entry * pFm,const size_t count)
+    {
+        CB_ASSERT( pTo != pFm || count == 0 );
+        
+        if ( pTo > pFm )
+        {
+            // go backwards
+            for(intptr_t i = count-1; i>= 0;i--)
+            {
+                pTo[i] = pFm[i];
+            }
+        }
+        else // ( pTo < pFm )
+        {
+            // go forwards
+            copy(pTo,pFm,count);
+        }
+    }
+
+    //-----------------------------------------------------------------------------------------------
+
+    template <typename t_entry>
+    static inline void construct(t_entry * pEntry)
+    {
+        new (pEntry) t_entry();
+    }
+
+    //-----------------------------------------------------------------------------------------------
+
+    template <typename t_entry>
+    static inline void destruct(t_entry * pEntry)
+    {
+        CB_ASSERT(pEntry);
+
+        pEntry->~t_entry();
+    }
+    
+    //-----------------------------------------------------------------------------------------------
+
+    template <typename t_entry>
+    static inline void construct(t_entry * pArray,const size_t size)
+    {
+        CB_ASSERT(pArray);
+        // placement new an array :
+        for(size_t i=0;i<size;i++)
+        {
+            //new (ePlacementNew, pArray+i) t_entry();
+            new (pArray+i) t_entry();
+        }
+    }
+
+    //-----------------------------------------------------------------------------------------------
+
+    template <typename t_entry>
+    static inline void destruct(t_entry * pArray,const size_t size)
+    {
+        for(size_t i=0;i<size;i++)
+        {
+            CB_ASSERT(pArray);
+            destruct(pArray+i);
+        }
+    }
+
+    //-----------------------------------------------------------------------------------------------
+    
+    template <typename t_entry>
+    static inline void copy_construct(t_entry * pTo,const t_entry & from)
+    {
+        CB_ASSERT(pTo);
+
+        //new (ePlacementNew, pTo) t_entry(from);
+        new (pTo) t_entry(from);
+    }
+        
+    //-----------------------------------------------------------------------------------------------
+    
+    template <typename t_entry>
+    static inline void copy_construct(t_entry * pArray,const t_entry * pFrom,const size_t size)
+    {
+        // placement new an array :
+        for(size_t i=0;i<size;i++)
+        {
+            CB_ASSERT( pArray && pFrom );
+            copy_construct(pArray+i,pFrom[i]);
+        }
+    }
+    
+    //-----------------------------------------------------------------------------------------------
+    // not used but nice :
+
+    /*
+    // Entryswap
+    //  like memcpy
+    //  pArray1 and pArray2 must not overlap
+    template <class t_entry>
+    static inline void swap_array(t_entry * pArray1,t_entry * pArray2,const size_t count)
+    {
+        // assert we don't overlap in a bad way :
+        CB_ASSERT( pArray1 != pArray2 );
+        CB_ASSERT( pArray1 < pArray2 || pArray1 >= pArray2 + count );
+
+        for(size_t i=0;i<count;i++)
+        {
+            CB_ASSERT( pArray1 && pArray2 );
+            Swap(pArray1[i],pArray2[i]);
+        }
+    }
+    */
+
+    //-----------------------------------------------------------------------------------------------
+
+    template <typename t_entry>
+    static inline void swap_construct(t_entry * pTo,t_entry & from)
+    {
+        CB_ASSERT(pTo);
+
+        new (pTo) t_entry();
+        Swap(from,*pTo);
+    }
+    
+    template <typename t_entry>
+    static inline void swap_construct(t_entry * pArray,t_entry * pFrom,const size_t size)
+    {
+        // placement new an array :
+        for(intptr_t i=0;i<size;i++)
+        {
+            CB_ASSERT( pArray && pFrom );
+            swap_construct(pArray+i,pFrom[i]);
+        }
+    }
+        
+    //-----------------------------------------------------------------------------------------------
+    
+    // We need this in vector, but I don't want to include <memory> there.
+    //  Note that this is slightly less flexible than the stl version,
+    //  it only works with pointer iterators, not general iterators
+    // has [FROM,TO] ordering unlike everything else
+    template <typename t_entry>
+    inline t_entry * uninitialized_copy(const t_entry*  first,
+                                            const t_entry*  last,
+                                            t_entry* result)
+    {
+        while (first != last)
+        {
+            new (result) t_entry(*first);
+            first++;
+            result++;
+        }
+
+        return result;
+    }
+    
+    //-----------------------------------------------------------------------------------------------
+};
+
+//}{=======================================================================================
+// cb_vector_base
+
+template <typename t_entry> class cb_vector_base
+{
+public:
+
+    typedef cb_vector_base<t_entry> this_type;
+    typedef CB_SIZETYPE             size_type;
+
+    //-----------------------------------------
+    // constructors :
+
+    CB_VECTOR_FORCEINLINE cb_vector_base()
+    {
+        init();
+    }
+
+    CB_VECTOR_FORCEINLINE ~cb_vector_base()
+    {
+        release();
+    }
+
+    cb_vector_base(const cb_vector_base & other)
+    {
+        init();
+        extend_copy(other.begin(),other.size());
+    }
+
+    template <class input_iterator>
+    cb_vector_base(const input_iterator first, const input_iterator last)
+    {
+        init();
+        assign_construct(first,last);
+    }
+
+    //-----------------------------------------
+
+    // increase size and construct added entries with their default constructor
+    void extend_default(const size_type count)
+    {
+        const size_type oldsize = m_size;
+        if ( needmakefit(m_size+count) )
+        {
+            const size_type oldcapacity = capacity();
+            makefit2( makefit1(m_size + count) , m_size , oldcapacity);
+        }
+        m_size += count;
+        cb_entry_array::construct(begin() + oldsize,count);
+    }
+
+    // increase size and don't construct!!!!
+    void extend_no_construct(const size_type count)
+    {
+        if ( needmakefit(m_size+count) )
+        {
+            const size_type oldcapacity = capacity();
+            makefit2( makefit1(m_size + count) , m_size , oldcapacity);
+        }
+        m_size += count;
+    }
+    
+    // increase size and construct added entries with their copy constructor
+    // this must work fine if pFrom is pointing into myself
+    //  the key is that makefit1() doesn't free the old memory, so we must
+    //  do any work between makefit1 and makefit2
+    void extend_copy(const t_entry * pFrom,const size_type count)
+    {
+        if ( needmakefit(m_size+count) )
+        {
+            const size_type oldsize = m_size;
+            const size_type oldcapacity = capacity();
+            t_entry * pOld = makefit1(m_size + count);
+            cb_entry_array::copy_construct(begin() + m_size,pFrom,count);
+            m_size += count;
+            makefit2(pOld, oldsize, oldcapacity);
+        }
+        else
+        {
+            cb_entry_array::copy_construct(begin() + m_size,pFrom,count);
+            m_size += count;
+        }
+    }
+
+    // fast specialization of extend_copy for one addition (for push_back)
+    CB_VECTOR_FORCEINLINE void extend_copy(const t_entry & from)
+    {
+        if ( needmakefit(m_size+1) )
+        {
+            const size_type oldsize = m_size;
+            const size_type oldcapacity = capacity();
+            t_entry * pOld = makefit1(m_size + 1);
+            cb_entry_array::copy_construct(begin() + m_size,from);
+            m_size ++;
+            makefit2(pOld, oldsize, oldcapacity);
+        }
+        else
+        {
+            cb_entry_array::copy_construct(begin() + m_size,from);
+            m_size ++;
+        }
+    }
+
+    // see notes on base extend_copy
+    CB_VECTOR_FORCEINLINE void extend_copy(const t_entry & from,const size_type count)
+    {
+        if ( needmakefit(m_size+count) )
+        {
+            const size_type oldsize = m_size;
+            const size_type oldcapacity = capacity();
+            t_entry * pOld = makefit1(m_size + count);
+            for(size_type i=0;i<count;i++)
+            {
+                cb_entry_array::copy_construct(begin() + m_size + i,from);
+            }
+            m_size += count;
+            makefit2(pOld, oldsize, oldcapacity);
+        }
+        else
+        {
+            for(size_type i=0;i<count;i++)
+            {
+                cb_entry_array::copy_construct(begin() + m_size + i,from);
+            }
+            m_size += count;
+        }
+    }
+
+    // reduce size and destruct stuff left behind :
+    CB_VECTOR_FORCEINLINE void shrink(const size_type newsize)
+    {
+        CB_ASSERT( newsize >= 0 && newsize <= m_size );
+        const size_type count = m_size - newsize;
+        cb_entry_array::destruct(begin() + newsize,count);
+        m_size = newsize;
+    }
+
+    //-----------------------------------------
+
+    void reserve(const size_type newcap)
+    {
+        if ( needmakefit(newcap) )
+        {
+            const size_type oldcapacity = capacity();
+            makefit2( makefit1( newcap ) , m_size, oldcapacity);
+        }
+    }
+
+    // release frees any memory allocated and resizes to zero
+    void release()
+    {
+        shrink(0);
+        if ( m_begin )
+            CB_FREE(m_begin,m_capacity*sizeof(t_entry));
+        //CBFREE(m_begin);
+        init();
+    }
+
+    void swap(this_type & other)
+    {
+        AssignmentSwap(m_begin,other.m_begin);
+        AssignmentSwap(m_size,other.m_size);
+        AssignmentSwap(m_capacity,other.m_capacity);
+    }
+
+    void insert(const size_type n_pos,
+                 const t_entry * first, const t_entry * last);
+
+    //-----------------------------------------
+
+    // assign_construct :
+    //  used by iterator-range and copy constructor for whole vectors
+    //  fills out a randbool-new (empty) cb_vector_base with the data from a range of iterators
+    template <class input_iterator>
+    void assign_construct(const input_iterator first,const input_iterator last)
+    {
+        CB_ASSERT( m_size == 0 );
+        // w64
+        const size_type count = (size_type)( last - first );
+        if ( needmakefit(count) )
+        {
+            const size_type oldcapacity = capacity();
+            t_entry * pOld = makefit1(count);
+            m_size = count;
+            cb_entry_array::uninitialized_copy(first, last, begin() );
+            makefit2(pOld,(size_type)0,oldcapacity);
+        }
+        else
+        {
+            m_size = count;
+            cb_entry_array::uninitialized_copy(first, last, begin() );
+        }
+    }
+
+    // specialization to simple iterators :
+    void assign_construct(const t_entry * const first,const t_entry * const last)
+    {
+        CB_ASSERT( m_size == 0 );
+        // w64
+        const size_type count = (size_type)( last - first );
+        if ( needmakefit(count) )
+        {
+            const size_type oldcapacity = capacity();
+            t_entry * pOld = makefit1(count);
+            m_size = count;
+            cb_entry_array::copy_construct(begin(),first,count);
+            makefit2(pOld,0,oldcapacity);
+        }
+        else
+        {
+            m_size = count;
+            cb_entry_array::copy_construct(begin(),first,count);
+        }
+    }
+
+    //-----------------------------------------
+    // simple accessors :
+
+    t_entry *           begin()         { return m_begin; }
+    const t_entry *     begin() const   { return m_begin; }
+    t_entry *           end()           { return begin()+m_size; }
+    const t_entry *     end() const     { return begin()+m_size; }
+    size_type           size() const    { return m_size; }
+    size_type           capacity() const { return m_capacity; }
+    size_type           max_size() const { return (1UL)<<30; }
+
+    //-----------------------------------------
+    
+    //-------------------------------------------------------
+protected :
+
+    template<typename Type> 
+    static inline void AssignmentSwap(Type &a,Type &b)
+    {
+        Type c = a;  // Type c(a);
+        a = b;
+        b = c;
+    }
+    
+    CB_VECTOR_FORCEINLINE bool needmakefit(const size_type newsize) const
+    {
+        return (newsize > m_capacity);
+    }
+
+    // makefit1
+    // returns the *old* pointer for passing into makefit2
+    //
+    t_entry * makefit1(const size_type newsize)
+    {
+        const size_type oldsize = m_size;
+
+        CB_ASSERT( needmakefit(newsize) );
+
+        t_entry * pOld = m_begin;
+
+        #ifdef CB_VECTOR_MAX_GROW_BYTES
+        // Be much more careful about growing the memory conservatively.  This changes
+        //  push_back from amortized O(N) to O(N^2), but results in tighter vectors.
+
+        enum { c_maxGrowBytes = CB_VECTOR_MAX_GROW_BYTES }; // 1 MB
+        enum { c_maxGrowCount = c_maxGrowBytes/sizeof(t_entry) };
+        CB_COMPILER_ASSERT( c_maxGrowCount > 1 );
+        // c_maxGrowCount should limit the doubling, but not the passed in newsize
+        
+        // m_capacity is 0 the first time we're called
+        // newsize can be passed in from reserve() so don't put a +1 on it
+
+        // grow by doubling until we get to max grow count, then grow linearly
+        size_type newcapacity = CB_MIN( m_capacity * 2 , (m_capacity + c_maxGrowCount) );
+        
+        #else
+        
+        size_type newcapacity = m_capacity * 2;
+        
+        #endif // CB_VECTOR_MAX_GROW_BYTES
+        
+        newcapacity = CB_MAX( newcapacity, newsize );
+
+        size_type newbytes;
+            
+        // if on constant should optimize out
+        if ( sizeof(t_entry) == 1 )
+        {
+            // round up newcapacity to be a multiple of 8
+            newcapacity = (newcapacity+7)&(~7);
+            newbytes = newcapacity;
+        }
+        else
+        {
+            newbytes = newcapacity * sizeof(t_entry);
+            
+            if ( newbytes > 65536 )
+            {
+                // align newbytes up :
+                newbytes = (newbytes + 65535) & (~65535);
+            }
+            else
+            {
+                if ( newbytes < 512 )
+                {
+                    newbytes = (newbytes + 15) & (~15);
+                }
+                else
+                {
+                    // align up to 4096 :
+                    newbytes = (newbytes + 4095) & (~4095);
+                }
+            }
+            
+            // align newcapacity down :
+            newcapacity = newbytes / sizeof(t_entry);
+        }
+
+        CB_ASSERT( newcapacity >= newsize );            
+        CB_ASSERT( newbytes >= newcapacity * sizeof(t_entry) );
+        
+        t_entry * pNew = (t_entry *) CB_ALLOC( newbytes );
+
+        CB_ASSERT_MALLOC( pNew != 0 );
+
+        // copy existing :
+        #ifdef CB_VECTOR_USE_SWAP
+        // swap instead of copy can be a big performance advantage
+        //  this requires entries to have a default constructor !!
+        cb_entry_array::swap_construct(pNew,pOld,oldsize);
+        #else
+        cb_entry_array::copy_construct(pNew,pOld,oldsize);
+        #endif
+        
+        m_begin = pNew;
+        m_capacity = newcapacity;
+        // m_size not changed
+
+        return pOld;
+    }
+
+    void makefit2(t_entry * pOld, const size_type oldsize, const size_type oldcapacity)
+    {
+        if ( pOld )
+        {
+            cb_entry_array::destruct(pOld,oldsize);
+
+            CB_FREE(pOld,oldcapacity*sizeof(t_entry));
+        }
+    }
+
+private :
+    //-----------------------------------------
+    // data :
+
+    t_entry *   m_begin;
+    size_type   m_size;     // how many in use
+    size_type   m_capacity; // how many allocated
+
+    void init()
+    {
+        m_begin = 0;
+        m_capacity = 0;
+        m_size = 0;
+    }
+
+    //-----------------------------------------
+}; // end cb_vector_base
+
+/**
+ cb_vector_base::insert for a chunk of elements
+ this does NOT work if the elements to insert are from the same vector,
+ since insert can cause allocation, and the elements to add are referred-to by pointer
+**/
+template<typename t_entry> void cb_vector_base<t_entry>::insert(
+        const size_type n_pos,
+        const t_entry * first, 
+        const t_entry * last)
+{
+    const size_type n_insert = (size_type)( last - first );
+    const size_type n_insert_end = (size_type) ( n_pos + n_insert );
+    const size_type old_size = size();
+    const size_type oldcapacity = capacity();
+
+    t_entry * pOld;
+    if ( needmakefit( m_size+n_insert ) )
+        pOld = makefit1( m_size + n_insert);
+    else
+        pOld = 0;
+
+    m_size += n_insert;
+
+    if ( n_insert_end > old_size )
+    {
+        // copy construct out to n_insert_end
+        //  from the insert chunk
+        const ptrdiff_t n_insert_copy = old_size - n_pos;
+        const ptrdiff_t n_insert_append = n_insert_end - old_size;
+        CB_ASSERT( n_insert_copy >= 0 && n_insert_append >= 0);
+        CB_ASSERT( (n_insert_copy + n_insert_append) == n_insert );
+
+        cb_entry_array::copy_construct( begin() + old_size, first + n_insert_copy, n_insert_append );
+
+        // now copy from old end to new end
+        cb_entry_array::copy_construct( begin() + old_size + n_insert_append, begin() + n_pos, n_insert_copy );
+
+        // finally copy in the n_insert_copy chunk :
+        cb_entry_array::copy( begin() + n_pos, first, n_insert_copy );
+    }
+    else
+    {
+        cb_entry_array::copy_construct( begin() + old_size, begin() + old_size - n_insert, n_insert );
+    
+        t_entry * newpos = begin() + n_pos;
+    
+        const ptrdiff_t count = old_size - n_insert_end;
+
+        cb_entry_array::move( newpos + n_insert, newpos, count );
+                
+        cb_entry_array::copy( newpos, first, n_insert );
+    }
+    
+    if ( pOld )
+        makefit2(pOld, old_size, oldcapacity);
+}
+
+//}{=======================================================================================
+// vector
+
+template <typename t_entry> class CB_VECTOR : 
+    protected cb_vector_base<t_entry>
+{
+public:
+    //----------------------------------------------------------------------
+    // typedefs :
+    typedef t_entry             value_type;
+    typedef t_entry*            iterator;
+    typedef const t_entry*      const_iterator;
+    typedef t_entry&            reference;
+    typedef const t_entry&      const_reference;
+    typedef CB_SIZETYPE         size_type;
+    
+    typedef CB_VECTOR<t_entry>  this_type;
+    typedef cb_vector_base<t_entry> parent_type;
+
+    //----------------------------------------------------------------------
+    // constructors
+
+    CB_VECTOR_FORCEINLINE  CB_VECTOR() { }
+    /*CB_VECTOR_FORCEINLINE*/ ~CB_VECTOR() { }
+
+    CB_VECTOR_FORCEINLINE CB_VECTOR(const size_type size,const value_type & init)
+    {
+        //parent_type::extend_default(size); 
+        parent_type::extend_copy(init,size); 
+    }
+
+    CB_VECTOR_FORCEINLINE CB_VECTOR(const this_type & other) :
+        parent_type(other)
+    {
+    }
+
+    template <class input_iterator>
+    CB_VECTOR_FORCEINLINE CB_VECTOR(const input_iterator first,const input_iterator last)
+        : parent_type(first,last)
+    {
+    }
+    
+    //---------------------------------------------------------------------------
+    // simple accessors :
+
+    // iterator support
+    iterator        begin()         { return parent_type::begin(); }
+    const_iterator  begin() const   { return parent_type::begin(); }
+    iterator        end()           { return parent_type::end(); }
+    const_iterator  end() const     { return parent_type::end(); }
+
+    // at() with range check
+    reference at(const size_type i)
+    {
+        CB_ASSERT( i >= 0 && i < parent_type::size() );
+        return *(parent_type::begin() + i); 
+    }
+    const_reference at(const size_type i) const
+    {
+        CB_ASSERT( i >= 0 && i < parent_type::size() );
+        return *(parent_type::begin() + i); 
+    }
+
+    // operator[]
+    reference       operator[](const size_type i)       { return at(i); }
+    const_reference operator[](const size_type i) const { return at(i); }
+
+    // front() and back()
+    reference       front()         { return at(0); }
+    const_reference front() const   { return at(0); }
+    reference       back()          { return at(parent_type::size()-1); }
+    const_reference back() const    { return at(parent_type::size()-1); }
+
+    // size queries :
+    size_type   size() const        { return parent_type::size(); }
+    bool        empty() const       { return parent_type::size() == 0; }
+    size_type   capacity() const    { return parent_type::capacity(); }
+    size_type   max_size() const    { return parent_type::max_size(); }
+
+    int         size32() const      { return (int)(size()); }
+    
+    void reserve(const size_type newcap)    { parent_type::reserve(newcap); }
+
+    //---------------------------------------------------------------------------
+    // mutators :
+
+    void push_back(const value_type & e)
+                                    { parent_type::extend_copy(e); }
+
+    void push_back()                { parent_type::extend_default(1); }
+
+    // push_back_no_construct is dangerous!!  You must immediately construct it yourself!!
+    void push_back_no_construct()   { parent_type::extend_no_construct(1); }
+    
+    void pop_back()                 { parent_type::shrink( parent_type::size() - 1 ); }
+
+    void clear()                    { parent_type::shrink(0); }
+
+    void resize(const size_type new_size);
+    void resize(const size_type new_size, const value_type & e);
+
+    //----------------------------------------------------------------------
+    // serious entry mutators :
+
+    iterator insert(iterator position, const value_type & e);
+
+    iterator insert(iterator position,
+              const_iterator first, const_iterator last);
+
+    iterator erase(iterator position);
+
+    iterator erase(iterator first, iterator last);
+
+    //----------------------------------------------------------------------
+    // serious whole mutators :
+
+    void swap(this_type & other)
+    {
+        parent_type::swap(other);
+    }
+
+    //@@ I'm not fond of operator= on a whole vector; it's error prone and expensive, but here for compatibility
+    void operator=(const this_type & other)
+    {
+        assign(other.begin(),other.end());
+    }
+
+    //----------------------------------------------------------------------
+    // assign :
+    //  assign does work fine for iterator ranges that are in yourself
+
+    template <class input_iterator>
+    void assign(const input_iterator first,const input_iterator last)
+    {
+        const size_type count = (size_type)( last - first );
+        if ( size() >= count )
+        {
+            // don't move on top of self; can happen eg. if you do operator= on yourself
+            if ( begin() != first )
+            {
+                cb_entry_array::move( begin(), first, count );
+            }
+            parent_type::shrink(count);
+        }
+        else
+        {
+            CB_ASSERT( ! is_iterator_in_range(first) );
+            CB_ASSERT( (first == last) || ! is_iterator_in_range(last-1) );
+            parent_type::shrink(0);
+            parent_type::assign_construct(first,last);
+        }
+    }
+
+    //----------------------------------------------------------------------
+    // extensions :
+
+    // @@ ? what should a call to data() do on an empty vector?  This will assert;
+    //  maybe it should return NULL ?
+    const t_entry * data() const    { CB_ASSERT( ! empty() ); return &(at(0)); }
+    t_entry * data()                { CB_ASSERT( ! empty() ); return &(at(0)); }
+
+    size_type size_bytes() const { return size() * sizeof(t_entry); } 
+    void memset_zero() { memset(data(),0,size_bytes()); }
+
+    // tighten releases memory that's not in use
+    void tighten()
+    {
+        if ( capacity() != size() )
+        {
+            this_type other(*this);
+            swap(other);
+        }
+    }
+    
+    // release frees any memory allocated and resizes to zero
+    void release()
+    {
+        parent_type::release();
+    }
+
+    // fast unordered erase :
+    //  now with STL-compliant signature
+    iterator erase_u(iterator position)
+    {
+        CB_ASSERT( is_iterator_in_range(position) );
+        if ( position != end()-1 )
+            *position = back();
+        pop_back();
+        return position;
+    }
+
+    // erase_u an index :
+    void erase_u(const size_type i)
+    {
+        if ( i != size()-1 )
+            at(i) = back();
+        pop_back();
+    }
+
+    // member function find :
+    // just write it outrselves so we don't have to bring in <algorithm>
+
+    template <class what_to_find>
+    iterator find(const what_to_find & what) 
+    {
+        for(iterator it = begin();it < end();it++)
+        {
+            if ( *it == what )
+                return it;
+        }
+        return end();
+    }
+
+    template <class what_to_find>
+    const_iterator find(const what_to_find & what) const 
+    {
+        for(const_iterator it = begin();it < end();it++)
+        {
+            if ( *it == what )
+                return it;
+        }
+        return end();
+    }
+
+    template <class _Predicate>
+    iterator find_if(_Predicate __pred)
+    {
+        for(iterator it = begin();it < end();it++)
+        {
+            if ( __pred(*it) )
+                return it;
+        }
+        return end();
+    }
+
+    template <class _Predicate>
+    const_iterator find_if(_Predicate __pred) const
+    {
+        for(iterator it = begin();it < end();it++)
+        {
+            if ( __pred(*it) )
+                return it;
+        }
+        return end();
+    }
+
+    // some sugar :
+
+    template <class other_vector>
+    void appendv( const other_vector & other )
+    {
+        insert( end(), other.begin(), other.end() );
+    }
+
+    template <class other_vector>
+    void insertv(iterator position, const other_vector & other )
+    {
+        insert( position, other.begin(), other.end() );
+    }
+
+    void append(const_iterator first, const_iterator last)
+    {
+        insert( end(), first, last );
+    }
+    
+    void append(const_iterator first, const size_type count)
+    {
+        insert( end(), first, first + count );
+    }           
+
+    template <class other_vector>
+    void assignv(const other_vector & other )
+    {
+        assign( other.begin(), other.end() );
+    }
+    
+
+    //----------------------------------------------------------------------
+
+private:
+    //----------------------------------------------------------------------
+    template <class input_iterator>
+    bool is_iterator_in_range(const input_iterator it) const
+    {
+        return ( it >= begin() && it < end() );
+    }
+    template <class input_iterator>
+    bool is_iterator_in_range_or_end(const input_iterator it) const
+    {
+        return ( it >= begin() && it <= end() );
+    }   
+
+}; // end vector
+
+//}{=======================================================================================
+// vector serious mutators
+
+#define T_PRE1 template<typename t_entry>
+#define T_PRE2 CB_VECTOR<t_entry>
+#define T_PRE_VOID  T_PRE1 void T_PRE2
+#define T_PRE_IT    T_PRE1 typename CB_VECTOR<t_entry>::iterator T_PRE2
+
+//----------------------------------------------------------------------
+// serious entry mutators :
+
+T_PRE_IT::insert(const iterator position, const value_type & e)
+{
+    // not maximally efficient if position == end; use push_back
+    CB_ASSERT( is_iterator_in_range_or_end(position) );
+
+    const ptrdiff_t n = ( position - begin() );
+    CB_ASSERT( n <= (ptrdiff_t) size() );
+
+    const ptrdiff_t move_count = size() - n;
+
+    if ( move_count == 0 )
+    {
+        push_back(e);
+        return begin();
+    }
+    else
+    {
+        // extend by copy construction :
+        if ( size() > 0 )
+        {
+            parent_type::extend_copy( back() );
+        }
+        // now cannot use "position" any more !
+
+        iterator newpos = begin() + n;
+        cb_entry_array::move(newpos+1,newpos, move_count );
+
+        *newpos = e;
+
+        return newpos;
+    }
+}
+
+T_PRE_IT::insert(iterator position,
+          const_iterator first, const_iterator last)
+{
+    CB_ASSERT( is_iterator_in_range_or_end(position) );
+    CB_ASSERT( ! is_iterator_in_range(first) );
+    CB_ASSERT( (first == last) || ! is_iterator_in_range(last-1) );
+
+    const ptrdiff_t n_pos = ( position - begin() );
+
+    parent_type::insert( (size_type) n_pos,first,last);
+
+    return begin() + n_pos;
+}
+
+T_PRE_IT::erase(const iterator position)
+{
+    CB_ASSERT( is_iterator_in_range(position) );
+    CB_ASSERT( position != end() );
+    const ptrdiff_t n_pos = ( position - begin() );
+    
+    // slide it down :
+    const ptrdiff_t newsize = size() - 1;
+    cb_entry_array::move(position,position+1, newsize - n_pos );
+
+    parent_type::shrink( (size_type) newsize);
+    
+    return begin() + n_pos;
+}
+
+T_PRE_IT::erase(const iterator first, const iterator last)
+{
+    CB_ASSERT( is_iterator_in_range_or_end(first) );
+    CB_ASSERT( is_iterator_in_range_or_end(last) );
+    if ( last <= first )
+        return end();
+
+    const ptrdiff_t ifirst = (first - begin());
+    const ptrdiff_t ilast  = (last - begin() );
+    const ptrdiff_t num_removed = ilast - ifirst;
+    const ptrdiff_t num_to_move = parent_type::size() - ilast;
+
+    cb_entry_array::move(first,last, num_to_move);
+    parent_type::shrink( (size_type)( parent_type::size() - num_removed ) );
+    
+    return begin() + ifirst;
+}
+
+T_PRE_VOID::resize(const size_type new_size)
+{
+    // if shrinking, must clear things being left :
+    if ( new_size < parent_type::size() )
+    {
+        parent_type::shrink(new_size);
+    }
+    else
+    {
+        parent_type::extend_default(new_size - parent_type::size());
+    }
+}
+
+T_PRE_VOID::resize(const size_type new_size, const value_type & e)
+{
+    // if shrinking, must clear things being left :
+    if ( new_size < parent_type::size() )
+    {
+        parent_type::shrink(new_size);
+    }
+    else
+    {
+        const size_type count = new_size - parent_type::size();
+        parent_type::extend_copy(e,count);
+    }
+}
+
+#undef T_PRE_IT
+#undef T_PRE_VOID
+#undef T_PRE1 
+#undef T_PRE2
+
+//}{=======================================================================================
+
+#ifdef _MSC_VER
+//#pragma pack(pop)
+#pragma warning(pop)
+#endif
+
+/*
+// partial specialize swap_functor to all vectors
+//  for cb::Swap
+
+template<class t_entry> 
+struct swap_functor< cb::vector<t_entry> >
+{
+    void operator () ( cb::vector<t_entry> & _Left, cb::vector<t_entry> & _Right)
+    {
+        _Left.swap(_Right);
+    }
+};
+*/
+
+#ifdef CB_VECTOR_NAMESPACE
+}; // namespace
+#endif
+
+//}{=======================================================================================
+
+
+#endif // CB_VECTOR_H_INCLUDED
diff --git a/Engine/Source/Programs/Horde/HordeStorage/Jupiter.Common/Shared/Oodle/2.9.3/linux/examples/example_jobify.cpp b/Engine/Source/Programs/Horde/HordeStorage/Jupiter.Common/Shared/Oodle/2.9.3/linux/examples/example_jobify.cpp
new file mode 100644
--- /dev/null
+++ b/Engine/Source/Programs/Horde/HordeStorage/Jupiter.Common/Shared/Oodle/2.9.3/linux/examples/example_jobify.cpp
@@ -0,0 +1,146 @@
+/****
+
+example_jobify :
+
+example implementations of the Oodle job system
+
+to be plugged into OodleCore_Plugins_SetJobSystemAndCount
+ OodleNet_Plugins_SetJobSystemAndCount
+
+You likely have a worker thread system already and may wish to plug that in instead.
+
+etc.
+
+****/
+
+/* 
+@cdep pre
+
+    $if($equals($BuildPlatform,linux),
+    
+        $// use TBB on linux if found :
+        $if($exists(/usr/include/tbb),
+            $addtocswitches(-DHAVE_TBB)
+            $requiresbinary( -ltbb ),
+            $addtocswitches(-DUSE_OODLEX)           
+        )
+        , )
+   
+    $if($equals($BuildPlatform,mac),
+        $requires( $projectspath/oodle2/examples/example_jobify_gcd.mm )
+        $requiresbinary(-framework Foundation)
+        , )
+*/
+
+/* 
+@cdep post
+
+*/
+
+#include "example_jobify.h"
+#ifdef USE_OODLEX
+#include "../include/oodle2x.h"
+#endif
+
+#include "ooex.h" // for platform detection
+
+#include <stdlib.h> 
+#include <stdio.h> 
+
+//===================================
+
+t_fp_example_jobify_RunJob * example_jobify_run_job_fptr = NULL;
+t_fp_example_jobify_WaitJob *example_jobify_wait_job_fptr = NULL;
+int example_jobify_target_parallelism = 0;
+
+//===================================
+
+#ifdef HAVE_TBB 
+// enabled if /usr/include/tbb is found by cdep
+// modern clang & gcc have "__has_include" , could use that instead
+// HAVE_TBB is semi-standard name set by autoconf
+
+#include "example_jobify_linuxtbb.inl"
+
+#elif defined(OOEX_PLATFORM_NT)
+
+#include "example_jobify_win32tp.inl"
+
+#elif defined(OOEX_PLATFORM_MAC)
+
+// in example_jobify_gcd.mm
+
+extern "C"
+void * example_jobify_init_gcd();
+
+void * example_jobify_init()
+{
+    void * userPtr = example_jobify_init_gcd();
+    
+    return userPtr;
+}
+
+#else
+
+// host tool platforms (Linux,Windows,Mac) should have been caught above
+// here we are either on a non-host platform
+// or Linux without TBB
+
+#ifdef USE_OODLEX
+
+// if you have OodleX, you can use it :
+
+void * example_jobify_init()
+{
+    // OodleX_Init installs jobify plugins to Oodle Core
+    
+    OodleXInitOptions oo;
+    OodleX_Init_GetDefaults(OODLE_HEADER_VERSION,&oo,
+        OodleX_Init_GetDefaults_DebugSystems_No,
+        OodleX_Init_GetDefaults_Threads_Yes),
+    // Texture work likes to use all hyper threads :
+    oo.m_OodleInit_Workers_Count = OODLE_WORKERS_COUNT_ALL_HYPER_CORES;
+    if ( ! OodleX_Init(OODLE_HEADER_VERSION,&oo) )
+    {
+        fprintf(stderr,"OodleX_Init failed\n");
+        exit(10);
+    }
+    
+    // OodleX automatically installs itself to Oodle Core
+    // but it doesn't do so to Oodle Net or others
+    // you can manually do it thusly :
+    
+    /*  
+    OodleNet_Plugins_SetPrintf(OodleXLog_Printf_Raw);
+    OodleNet_Plugins_SetAssertion(OodleX_DisplayAssertion);
+    OodleNet_Plugins_SetAllocators(OodleXMallocAligned,OodleXFree);
+    OodleNet_Plugins_SetJobSystemAndCount(OodleX_CorePlugin_RunJob,OodleX_CorePlugin_WaitJob,OodleX_GetNumWorkerThreads());
+    */
+    
+    example_jobify_run_job_fptr = OodleX_CorePlugin_RunJob;
+    example_jobify_wait_job_fptr = OodleX_CorePlugin_WaitJob;
+    example_jobify_target_parallelism = OodleX_GetNumWorkerThreads();
+    
+    // no jobify user ptr
+    
+    return NULL;
+}
+
+#else
+
+void * example_jobify_init()
+{
+    fprintf(stderr,"warning: no jobify.\n");
+    
+    example_jobify_run_job_fptr = NULL;
+    example_jobify_wait_job_fptr = NULL;
+    example_jobify_target_parallelism = 0;
+    
+    return NULL;
+}
+
+#endif
+
+#endif
+    
+    
diff --git a/Engine/Source/Programs/Horde/HordeStorage/Jupiter.Common/Shared/Oodle/2.9.3/linux/examples/example_jobify.h b/Engine/Source/Programs/Horde/HordeStorage/Jupiter.Common/Shared/Oodle/2.9.3/linux/examples/example_jobify.h
new file mode 100644
--- /dev/null
+++ b/Engine/Source/Programs/Horde/HordeStorage/Jupiter.Common/Shared/Oodle/2.9.3/linux/examples/example_jobify.h
@@ -0,0 +1,17 @@
+#pragma once
+
+// use only oodle2base so we work with core/net/tex :
+#include "../include/oodle2base.h"
+
+// returns void * jobifyUserPtr you must pass back
+void * example_jobify_init();
+
+OODEFFUNC typedef void (OODLE_CALLBACK t_fp_example_jobify_Job)( void * job_data );
+OODEFFUNC typedef OO_U64 (OODLE_CALLBACK t_fp_example_jobify_RunJob)( t_fp_example_jobify_Job * fp_job, void * job_data , OO_U64 * dependencies, int num_dependencies, void * user_ptr );
+OODEFFUNC typedef void (OODLE_CALLBACK t_fp_example_jobify_WaitJob)( OO_U64 job_handle, void * user_ptr );
+
+
+extern t_fp_example_jobify_RunJob * example_jobify_run_job_fptr;
+extern t_fp_example_jobify_WaitJob *example_jobify_wait_job_fptr;
+extern int example_jobify_target_parallelism;
+
diff --git a/Engine/Source/Programs/Horde/HordeStorage/Jupiter.Common/Shared/Oodle/2.9.3/linux/examples/example_jobify_gcd.mm b/Engine/Source/Programs/Horde/HordeStorage/Jupiter.Common/Shared/Oodle/2.9.3/linux/examples/example_jobify_gcd.mm
new file mode 100644
--- /dev/null
+++ b/Engine/Source/Programs/Horde/HordeStorage/Jupiter.Common/Shared/Oodle/2.9.3/linux/examples/example_jobify_gcd.mm
@@ -0,0 +1,144 @@
+/*
+
+  Oodle pluggable job system example using NSOperationQueue on OS X (part of the Foundation framework),
+  which in turn builds on top of GCD (Grand Central Dispatch).
+
+  NSOperationQueue and NSOperations have direct support for DAGs of jobs. Therefore, the only thing
+  this implementation has to do is create the relevant operation objects and set up dependencies.
+
+  One important caveat is that a OodleLZ_Compress call using jobify callbacks will complete
+  synchronously, which includes waiting for its spawned jobs to finish. That means that if you
+  call OodleLZ_Compress on a GCD/worker thread, it will perform a blocking wait on it. GCD will
+  run more threads than the number of logical cores, but it won't keep adding threads indefinitely
+  (a typical limit seems to be 64 worker threads); therefore it is possible for an application to
+  deadlock if there are (say) 64 worker threads all running OodleLZ_Compress and waiting for a child
+  job to complete, with no worker thread actually running say child job. If you are concerned about this,
+  your options are:
+
+  1. Disable Oodle jobs. If you have enough independent parallel work to keep 64 worker threads
+     busy, you presumably don't need the extra parallelism. OodleLZ_Compress calls with jobs
+	 disabled contain no blocking waits and will always complete with no further intervention.
+  2. Throttle the number of in-flight OodleLZ_Compress calls. Having more compressors running
+     than logical cores available has no benefit to begin with.
+  3. Use the OodleX thread pool. The OodleX thread pool will not actually block a thread waiting
+     for child jobs; it will have blocked threads run other work instead until the work they
+	 depend on has been completed.
+
+*/
+
+#include "example_jobify.h"
+#include "ooex.h" // for platform detection
+
+#ifdef OOEX_PLATFORM_MAC
+#include <stdio.h>
+#include <stdlib.h>
+#include <string.h>
+#include <assert.h>
+#import <Foundation/Foundation.h>
+
+
+// ---- Implementation of the job callbacks using GCD (well, NSOperation)
+
+static void * g_example_jobify_jobifyUserPtr = NULL;
+
+static OO_U64 OODLE_CALLBACK gcd_run_job(t_fp_Oodle_Job * fp_job, void * job_data, OO_U64 * dependencies, int num_dependencies, void * user_ptr)
+{
+	if ( user_ptr == NULL )
+	{
+		// someone invoked a Jobify'ed Oodle call but didn't pass the user pointer
+		// in this example, we'll make it still work by using a global pointer
+		// you could also choose to fail or assert here
+		//	if you want to enforce in your code that the correct user pointer is always used
+		printf("example_jobify : user_ptr was not passed through Jobified call; using global.\n");
+
+		user_ptr = g_example_jobify_jobifyUserPtr;
+		
+		#if 0
+		// alternatively you could run synchronously in this case
+		// dependencies must all be done already
+		(*fp_job)(job_data);
+		// no handle :
+		return 0;
+		#endif
+	}
+
+	NSBlockOperation * op = [NSBlockOperation blockOperationWithBlock: ^{
+		fp_job(job_data);
+	}];
+	[op retain]; // Need to do this manually since we'll cast it to a OO_U64 which ARC doesn't track
+
+	for (int i = 0; i < num_dependencies; i++)
+		[op addDependency:reinterpret_cast<NSOperation *>(dependencies[i])];
+
+	NSOperationQueue * queue = reinterpret_cast<NSOperationQueue*>(user_ptr);
+	[queue addOperation: op];
+
+	// Our block operation is the handle
+	return reinterpret_cast<OO_U64>(op);
+}
+
+static void OODLE_CALLBACK gcd_wait_job(OO_U64 job_handle, void * /*user_ptr*/)
+{
+	NSBlockOperation * op = reinterpret_cast<NSBlockOperation *>(job_handle);
+
+	// Wait for the outstanding job to finish
+	[op waitUntilFinished];
+
+	// Then release our reference on it
+	[op release];
+}
+
+extern "C"
+void * example_jobify_init_gcd();
+
+#if 0
+int mac_get_core_count()
+{
+	//int	count;
+	//size_t count_len = sizeof(count);
+	//sysctlbyname("hw.logicalcpu", &count, &count_len, NULL, 0);
+	//sysctlbyname("hw.physicalcpu", &count, &count_len, NULL, 0);
+
+    int nm[2];
+    size_t len = 4;
+    uint32_t count;
+
+    nm[0] = CTL_HW; nm[1] = HW_AVAILCPU;
+    sysctl(nm, 2, &count, &len, NULL, 0);
+
+    if(count < 1) {
+        nm[1] = HW_NCPU;
+        sysctl(nm, 2, &count, &len, NULL, 0);
+        if(count < 1) { count = 1; }
+    }
+    return count;
+}
+#endif
+
+void * example_jobify_init_gcd()
+{
+	// Allocate our operation queue
+	NSOperationQueue * queue = [[NSOperationQueue alloc] init];
+
+	int target_parallelism = 0; // physical core count unknown, let Oodle choose
+	// std::thread::hardware_concurrency(); ?
+	// logicalcpucount = sysconf( _SC_NPROCESSORS_ONLN ); ?
+
+	// Install GCD as our job system.
+	//OodleCore_Plugins_SetJobSystemAndCount(gcd_run_job, gcd_wait_job, target_parallelism);
+
+	example_jobify_run_job_fptr = gcd_run_job;
+	example_jobify_wait_job_fptr = gcd_wait_job;
+	example_jobify_target_parallelism = target_parallelism;
+
+	// Pass the NSOperationQueue to use as our user pointer
+	void * jobifyUserPtr = queue;
+
+	// save to global :
+  	g_example_jobify_jobifyUserPtr = jobifyUserPtr;
+  	
+	return jobifyUserPtr;
+}
+
+#endif
+
diff --git a/Engine/Source/Programs/Horde/HordeStorage/Jupiter.Common/Shared/Oodle/2.9.3/linux/examples/example_jobify_linuxtbb.inl b/Engine/Source/Programs/Horde/HordeStorage/Jupiter.Common/Shared/Oodle/2.9.3/linux/examples/example_jobify_linuxtbb.inl
new file mode 100644
--- /dev/null
+++ b/Engine/Source/Programs/Horde/HordeStorage/Jupiter.Common/Shared/Oodle/2.9.3/linux/examples/example_jobify_linuxtbb.inl
@@ -0,0 +1,280 @@
+/*
+
+  Oodle pluggable job system example using Intel TBB on Linux. The TBB logic in here
+  is not Linux-specific and should work on any platform supported by TBB.
+
+  TBB's model is one of fork-join parallelism and doesn't directly support the somewhat
+  more general job graphs we want, but it is relatively straightforward to build the
+  semantics we need on top of TBB tasks. This works the same way as in TBBs "General
+  Acyclic Graphs of Tasks" example, although it takes a bit of plumbing, encapsulated
+  in the JobContext class.
+
+  TBB "wait_for_all()" operations don't actually perform a blocking wait on the calling
+  thread, and instead make it run other work until the task it was waiting for is
+  complete. Therefore this implementation never needs more worker threads than logical
+  cores to fill up the machine and is quite efficient.
+
+*/
+
+#include "example_jobify.h"
+
+#ifdef OOEX_PLATFORM_LINUX
+#define _POSIX_C_SOURCE 200809L
+#include <stdio.h>
+#include <stdlib.h>
+#include <assert.h>
+#include <tbb/spin_mutex.h>
+#include <tbb/task.h>
+#include <tbb/task_scheduler_init.h>
+
+// ---- Implementation of the job callbacks using TBB
+
+class JobContext
+{
+	struct DepLink
+	{
+		JobContext * job;
+		DepLink * next; // singly-linked list
+	};
+
+	class Task : public tbb::task
+	{
+		// The job that owns this task.
+		JobContext * job;
+
+		// The job callback and its (opaque) data pointer
+		t_fp_Oodle_Job * fp_job;
+		void * job_data;
+
+	public:
+		Task(JobContext * job, t_fp_Oodle_Job * fp_job, void * job_data)
+			: job(job), fp_job(fp_job), job_data(job_data)
+		{
+		}
+
+		virtual tbb::task * execute();
+	};
+
+	// Root task for WaitJob (this is the parent of the actual work task)
+	tbb::task * root_task;
+	// The actual child task we spawn
+	tbb::task * child_task;
+
+	// Dependency slots. Not actually our data, it belongs to the jobs we depend
+	// on, which use it to store a linked list of their dependents.
+	DepLink dep_slot[OODLE_JOB_MAX_DEPENDENCIES];
+
+	// Dependency tracking information
+	tbb::spin_mutex dependents_lock;
+	bool job_completed;
+	DepLink * dependents_head;
+	DepLink * * dependents_tail_next;
+	// ---- end data protected by dependents_lock
+
+	// continue_with is the next task to run after this one, or NULL
+	// if we don't have a viable candidate yet.
+	//
+	// Returns the new continue_with.
+	tbb::task * notify_dep_satisfied(tbb::task * continue_with, int num_deps=1)
+	{
+		// Once its oustanding ref count reaches zero, spawn the child task.
+		if (child_task->add_ref_count(-num_deps) == 0)
+		{
+			// If we didn't have a viable candidate for continue_with, now we do.
+			if (!continue_with)
+				return child_task;
+			else
+				tbb::task::spawn(*child_task);
+		}
+
+		return continue_with;
+	}
+
+	// Returns next task to run
+	tbb::task * notify_completion()
+	{
+		// When the task completes, we need to stop accepting new dependencies
+		// and grab our current dependency list. The only thing happening inside
+		// the mutex is grabbing the dependencies; do the actual spawning outside.
+		DepLink * dep;
+		{
+			tbb::spin_mutex::scoped_lock guard(dependents_lock);
+			job_completed = true;
+			dep = dependents_head;
+		}
+
+		// Spawn the dependent jobs. The first dependency, we don't actually
+		// spawn(), but rather return from execute() (so it immediately gets
+		// scheduled on the calling thread).
+		tbb::task * continue_with = 0;
+		while(dep)
+		{
+			DepLink * depnext = dep->next;
+			continue_with = dep->job->notify_dep_satisfied(continue_with);
+			// dep is in child, it could be deleted now
+			dep = depnext;
+		}
+
+		return continue_with;
+	}
+
+	// Returns 1 if the job was already completed before we set up the
+	// dependency edge.
+	int add_dependent(JobContext * which, int slot)
+	{
+		// Set up the dependency record
+		DepLink * dep = &which->dep_slot[slot];
+		dep->job = which;
+		dep->next = 0;
+
+		tbb::spin_mutex::scoped_lock guard(dependents_lock);
+		if (!job_completed)
+		{
+			// Insert the new dependent at the end. (We maintain the list in FIFO order
+			// so we unblock the longest-waiting dependent first.)
+			*dependents_tail_next = dep;
+			dependents_tail_next = &dep->next;
+			return 0;
+		}
+		else
+			return 1;
+	}
+
+public:
+	JobContext(t_fp_Oodle_Job * fp_job, void * job_data, U64 * dependencies, int num_dependencies, tbb::task_group_context *group_ctx)
+	{
+		job_completed = false;
+		dependents_head = 0;
+		dependents_tail_next = &dependents_head;
+
+		// Set up the root TBB task; this task is empty, it's just a handle for us to use
+		// for bookkeeping.
+		if (group_ctx)
+			root_task = new( tbb::task::allocate_root(*group_ctx) ) tbb::empty_task;
+		else
+			root_task = new( tbb::task::allocate_root() ) tbb::empty_task;
+		root_task->set_ref_count(2); // 1 for child, 1 for wait
+		
+		// Set up the child task that does the actual work
+		child_task = new( root_task->allocate_child() ) Task(this, fp_job, job_data);
+
+		// We use the child_task's ref_count to count the number of outstanding dependencies
+		// before it may launch. Note that this is different from the way TBB normally
+		// uses ref counts: they usually count outstanding jobs before tasks may join,
+		// not before they may fork. (It's idiomatic for graphs though, see the "General
+		// Acyclic Graphs of Tasks" example.)
+		//
+		// The extra dependency is on us completing the setup; until the job has finished
+		// setting up, the child task may not launch.
+		child_task->set_ref_count(num_dependencies + 1);
+
+		// Number of dependencies to release when we finish setup. This is initially just 1
+		// but will be higher if we notice during dependency setup that some of the jobs
+		// we depend on have already completed.
+		int release_count = 1;
+
+		// Add our dependencies
+		for (int i = 0; i < num_dependencies; i++)
+			release_count += reinterpret_cast<JobContext *>(dependencies[i])->add_dependent(this, i);
+
+		// Release our setup dependency, along with dependencies on jobs that had already
+		// completed when we ran add_dependent. We pass "child_task" as continue_with to
+		// signal that we don't want to elide tbb::task::spawn calls.
+		notify_dep_satisfied(child_task, release_count);
+	}
+
+	// This has the same semantics as WaitJob(), i.e. you only get to do this once.
+	void wait()
+	{
+		// This waits for the child task (the actual work) to be completed and then frees
+		// the root task. The child task got freed when the work was actually done, whenever
+		// that was.
+		root_task->wait_for_all();
+
+		// The root task was only there for us to do the wait; that's done now!
+		tbb::task::destroy(*root_task);
+
+		child_task = 0;
+		root_task = 0;
+	}
+};
+
+tbb::task * JobContext::Task::execute()
+{
+	fp_job(job_data);
+	return job->notify_completion();
+}
+
+//================================
+
+void * g_example_jobify_jobifyUserPtr = NULL;
+
+static U64 OODLE_CALLBACK tbb_run_job(t_fp_Oodle_Job * fp_job, void * job_data, U64 * dependencies, int num_dependencies, void * user_ptr)
+{
+	assert(num_dependencies <= OODLE_JOB_MAX_DEPENDENCIES);
+
+	if ( user_ptr == NULL )
+	{
+		// someone invoked a Jobify'ed Oodle call but didn't pass the user pointer
+		// in this example, we'll make it still work by using a global pointer
+		// you could also choose to fail or assert here
+		//	if you want to enforce in your code that the correct user pointer is always used
+		printf("example_jobify : user_ptr was not passed through Jobified call; using global.\n");
+
+		user_ptr = g_example_jobify_jobifyUserPtr;
+		
+		#if 0
+		// alternatively you could run synchronously in this case
+		// dependencies must all be done already
+		(*fp_job)(job_data);
+		// no handle :
+		return 0;
+		#endif
+	}
+	
+	// Create the new job
+	tbb::task_group_context * group_ctx = static_cast<tbb::task_group_context *>(user_ptr);
+	assert( group_ctx != NULL );
+		
+	JobContext * ctx = new JobContext(fp_job, job_data, dependencies, num_dependencies, group_ctx);
+
+	// Return our context pointer as the handle
+	return reinterpret_cast<U64>(ctx);
+}
+
+static void OODLE_CALLBACK tbb_wait_job(U64 job_handle, void * /*user_ptr*/)
+{
+	JobContext * job = reinterpret_cast<JobContext *>(job_handle);
+	job->wait();
+	delete job;
+}
+
+void * example_jobify_init()
+{
+	//tbb::task_scheduler_init tbb_init;	// Automatic number of threads
+	int tbb_threads = tbb::task_scheduler_init::default_num_threads();
+  
+	// Run the job system test.
+	//Oodle_JobSystemStressTest(tbb_run_job, tbb_wait_job, NULL, 0, 256, 50000);
+
+	int target_parallelism = tbb_threads; // in tbb this is logical processor count
+	
+	// Install TBB as our job system.
+	//OodleCore_Plugins_SetJobSystemAndCount(tbb_run_job, tbb_wait_job, target_parallelism);
+	
+	example_jobify_run_job_fptr = tbb_run_job;
+	example_jobify_wait_job_fptr = tbb_wait_job;
+	example_jobify_target_parallelism = tbb_threads; // in tbb this is logical processor count
+
+	// Set up our task group context to use for spawning jobs, with default mode
+	// (in particular, no "concurrent wait")
+	tbb::task_group_context * oodle_group_ctx = new tbb::task_group_context;
+	void * jobifyUserPtr = oodle_group_ctx;
+	
+	g_example_jobify_jobifyUserPtr = jobifyUserPtr;
+	
+	return jobifyUserPtr;
+}
+
+#endif
+
diff --git a/Engine/Source/Programs/Horde/HordeStorage/Jupiter.Common/Shared/Oodle/2.9.3/linux/examples/example_jobify_test.cpp b/Engine/Source/Programs/Horde/HordeStorage/Jupiter.Common/Shared/Oodle/2.9.3/linux/examples/example_jobify_test.cpp
new file mode 100644
--- /dev/null
+++ b/Engine/Source/Programs/Horde/HordeStorage/Jupiter.Common/Shared/Oodle/2.9.3/linux/examples/example_jobify_test.cpp
@@ -0,0 +1,145 @@
+/****
+
+example_jobify_test :
+
+test OodleLZ Jobify using the"example_jobify" job system implementations
+
+****/
+
+#include "../include/oodle2.h"
+#include <stdio.h> 
+#include <stdlib.h> 
+#include <string.h> 
+
+//===================================
+
+#include "example_jobify.h"
+
+#include "ooex.h"
+#include "read_whole_file.h"
+ 
+// ---- Job system tester
+
+//#include "job_stress_test.h"
+
+// ---- Main program
+
+#ifdef BUILDING_EXAMPLE_CALLER
+#define main example_jobify_test
+#endif
+
+extern "C" int main(int argc,char *argv[])
+{
+    if (argc < 2)
+    {
+        printf("Usage: example_jobify_test <filename>\n");
+        return 1;
+    }
+    
+    //example_jobify_init_oodleX();
+    void * jobifyUserPtr = example_jobify_init();
+    
+    OodleCore_Plugins_SetJobSystemAndCount(example_jobify_run_job_fptr, example_jobify_wait_job_fptr, example_jobify_target_parallelism);
+    //etc :
+    //OodleNet_Plugins_SetJobSystemAndCount(example_jobify_run_job_fptr, example_jobify_wait_job_fptr, example_jobify_target_parallelism);
+
+    const char *filename = argv[1];
+
+    // Read the provided input file
+    OO_SINTa file_size = 0;
+    char * file_bytes = read_whole_file(filename, &file_size);
+    if (!file_bytes)
+        OOEX_ASSERT_FAILURE_ALWAYS("Error reading input file!\n");
+
+    printf("\"%s\": %lld bytes.\n", filename, (long long)file_size);
+
+    // Determine the required size of the output buffer, then allocate it.
+    OO_SINTa comp_buf_size = OodleLZ_GetCompressedBufferSizeNeeded(OodleLZ_Compressor_Invalid,file_size);
+    unsigned char * comp_buf = new unsigned char[comp_buf_size];
+    if (!comp_buf)
+        OOEX_ASSERT_FAILURE_ALWAYS("Error allocating compressed data buffer!\n");
+
+    // Run the job system test.
+    //Oodle_JobSystemStressTest(thread_pool_run_job, thread_pool_wait_job, NULL, 0, 256, 50000);
+
+    unsigned char * decomp_buf = new unsigned char[file_size];
+    if (!decomp_buf)
+        OOEX_ASSERT_FAILURE_ALWAYS("Error allocating decompressed data buffer!\n");
+
+    for (int iter=0;iter<3;iter++)
+    {
+        // Configure compresion parameters
+        OodleLZ_Compressor compressor = OodleLZ_Compressor_Kraken;
+        // Jobify only helps Optimal level encoders :
+        OodleLZ_CompressionLevel level = OodleLZ_CompressionLevel_Optimal2;
+        
+        OO_SINTa comp_result;
+            
+        if ( iter == 0 ) 
+        {
+            printf("Compressing with Default args:\n");
+            
+            // Default args have Jobify ON (Jobify_Default)
+            // jobifyUserPtr will be NULL
+            // this means if you call OodleLZ_Compress like this with default args
+            //  you will get a callback to RunJob but with NULL user pointer
+            
+            comp_result = OodleLZ_Compress(compressor, file_bytes, file_size, comp_buf, level);
+        
+        }
+        else
+        {       
+            OodleLZ_CompressOptions options = *OodleLZ_CompressOptions_GetDefault(compressor, level);
+            
+            if ( iter == 1 )
+            {
+                printf("Compressing with Jobify Disabled:\n");
+                
+                options.jobify = OodleLZ_Jobify_Disable;
+            }
+            else
+            {
+                // Compress the file contents!
+                printf("Compressing with Jobify Aggressive:\n");
+                
+                // Turn on aggressive jobifying for this example. This tends to use
+                // extra memory; if you're worried about memory usage, stick with the
+                // default.
+                options.jobify = OodleLZ_Jobify_Aggressive;
+            }
+            
+            // pass required pointer to jobify context :
+            options.jobifyUserPtr = jobifyUserPtr;
+
+            comp_result = OodleLZ_Compress(compressor, file_bytes, file_size, comp_buf, level, &options);
+        }
+        
+        if (comp_result == OODLELZ_FAILED)
+            OOEX_ASSERT_FAILURE_ALWAYS("Error occured during compression!\n");
+                
+        printf("Compressed %lld bytes -> %lld bytes.\n", 
+            (long long)file_size, 
+            (long long)comp_result);
+
+        // Decompress to make sure it worked.
+        printf("Decompressing...\n");
+
+        OO_SINTa decomp_result = OodleLZ_Decompress(comp_buf, comp_result, decomp_buf, file_size, OodleLZ_FuzzSafe_Yes);
+        if (decomp_result == OODLELZ_FAILED)
+            OOEX_ASSERT_FAILURE_ALWAYS("Error occured during decompression!\n");
+
+        // Verify that the decompressed results indeed match!
+        if (memcmp(file_bytes, decomp_buf, file_size) != 0)
+            OOEX_ASSERT_FAILURE_ALWAYS("The original and decompressed data disagree!\n");
+
+    }
+
+    printf("All done!\n");
+
+    delete[] decomp_buf;
+    delete[] comp_buf;
+    free(file_bytes);
+    return 0;
+}
+
+
diff --git a/Engine/Source/Programs/Horde/HordeStorage/Jupiter.Common/Shared/Oodle/2.9.3/linux/examples/example_jobify_win32tp.inl b/Engine/Source/Programs/Horde/HordeStorage/Jupiter.Common/Shared/Oodle/2.9.3/linux/examples/example_jobify_win32tp.inl
new file mode 100644
--- /dev/null
+++ b/Engine/Source/Programs/Horde/HordeStorage/Jupiter.Common/Shared/Oodle/2.9.3/linux/examples/example_jobify_win32tp.inl
@@ -0,0 +1,313 @@
+/*
+
+  Oodle pluggable job system example using the Win32 thread pool API (Windows Vista and later).
+
+  The Vista+ thread pool API is part of the regular Windows APIs and is exposed with
+  _WIN32_WINNT >= 0x0600.
+
+  The thread pool doesn't have direct support for DAGs of jobs, but does permit tasks on the
+  pool to perform blocking waits (creating more pool threads when required), which is stronger
+  than what we need.
+
+  The main difficulty in this implementation is implementing wake-up of dependent tasks
+  efficiently. It was written carefully to work even when CreateThreadpoolWork calls
+  fail spuriously (not that this is expected to happen).
+
+*/
+
+#include "example_jobify.h"
+
+#ifdef OOEX_PLATFORM_NT
+#include <stdlib.h>
+#include <assert.h>
+
+#if defined(UE_BUILD_DEBUG) || defined(UE_BUILD_DEVELOPMENT) || defined(UE_BUILD_TEST) || defined(UE_BUILD_SHIPPING)
+
+#include "Windows/AllowWindowsPlatformTypes.h"
+#include "Windows/WindowsPlatformAtomics.h"
+
+#define InterlockedExchangeAdd(ptr,v) FPlatformAtomics::InterlockedAdd((int32 *)ptr,v)
+
+#else
+
+// For the Windows Vista-level thread pool APIs
+#define _WIN32_WINNT 0x0600
+#include <Windows.h>
+
+#endif
+
+#ifdef RTL_SRWLOCK_INIT
+
+//#pragma message("example_jobify_win32tp!")
+
+// ---- Implementation of the jobify callbacks using the Windows thread pool
+
+//#define DEBUG printf
+#define DEBUG(...)
+
+class JobContext
+{
+	struct DepLink
+	{
+		JobContext * job;
+		DepLink * next; // singly-linked list
+	};
+
+	// The job callback and its (opaque) data pointer
+	t_fp_Oodle_Job * fp_job;
+	void * job_data;
+
+	// The thread pool work handle
+	PTP_WORK work_handle;
+
+	// "Inverted" lock; held from creation until we've finished executing the job.
+	SRWLOCK busy_lock;
+
+	// Dependency slots. Not actually our data, it belongs to the jobs
+	// we depend on, which use it to store a linked list of their
+	// dependents.
+	DepLink dep_slot[OODLE_JOB_MAX_DEPENDENCIES];
+
+	// Dependency tracking information (manipulated via Interlocked ops, don't touch directly outside init!)
+	LONG volatile num_unsatisfied_deps;
+
+	SRWLOCK dependents_lock; // Just used as a non-recursive mutex; protects the following fields:
+	bool job_completed;
+	DepLink * dependents_head;
+	DepLink * * dependents_tail_next;
+	// ---- end data protected by dependents_lock
+
+	static VOID CALLBACK work_callback(PTP_CALLBACK_INSTANCE /*instance*/, PVOID context, PTP_WORK /*work*/)
+	{
+		JobContext * me = static_cast<JobContext *>(context);
+
+		// Actually run the function
+		me->fp_job(me->job_data);
+
+		// Flag our job as complete.
+		me->notify_completion();
+	}
+
+	void notify_dep_satisfied(LONG num_deps = 1)
+	{
+		LONG deps_before = InterlockedExchangeAdd(&num_unsatisfied_deps, -num_deps);
+
+		if (deps_before == num_deps)
+		{
+			DEBUG("submit 0x%p\n", work_handle);
+			SubmitThreadpoolWork(work_handle);
+		}
+	}
+
+	void notify_completion()
+	{
+		// Flag the job as complete and grab the dependents list head
+		DepLink * dep;
+
+		AcquireSRWLockExclusive(&dependents_lock);
+		job_completed = true; // prevent anyone from adding further dependencies on us
+		dep = dependents_head;
+		ReleaseSRWLockExclusive(&dependents_lock);
+
+		// This notifies wait() that the job is complete.
+		ReleaseSRWLockExclusive(&busy_lock);
+		// note our "this" can now be deleted!
+
+		// Unblock the dependents we know about; setting job_completed guarantees the
+		// dependency list is now fixed and will not be modified any further.
+		
+		// Notify all dependents we know about that this dependency is now satisfied
+ 		while(dep)
+ 		{
+ 			DepLink * next = dep->next;
+ 			JobContext * job = dep->job;
+ 			// as soon as you call this, "dep" can be deleted :
+			job->notify_dep_satisfied();
+ 			dep = next;
+ 		}
+
+	}
+
+	// Returns 1 if the job was already completed before we set up the
+	// dependency edge.
+	LONG add_dependent(JobContext *which, int slot)
+	{
+		// Set up the dependency record
+		DepLink *dep = &which->dep_slot[slot];
+		dep->job = which;
+		dep->next = 0;
+
+		AcquireSRWLockExclusive(&dependents_lock);
+
+		LONG ret = 1; // assume job is already complete
+		if (!job_completed)
+		{
+			// Insert a new dependent at the end. (We maintain the dependents
+			// list in FIFO order so that we unblock the longest-waiting
+			// dependent jobs first.)
+			*dependents_tail_next = dep;
+			dependents_tail_next = &dep->next;
+
+			// Return 0 to signal the dependency was actually added.
+			ret = 0;
+		}
+
+		ReleaseSRWLockExclusive(&dependents_lock);
+		return ret;
+	}
+
+public:
+
+ 	JobContext(t_fp_Oodle_Job * _fp_job, void * _job_data, OO_U64 * dependencies, int num_dependencies, void * user_ptr)
+   	{
+ 		// @@ CB make sure I'm initialized :
+ 		memset(this,0,sizeof(*this));
+ 
+ 		fp_job = _fp_job;
+ 		job_data = _job_data;
+
+		// One extra dependency, for us completing setup, because the job may
+		// not launch until we're done setting up all dependencies.
+		num_unsatisfied_deps = num_dependencies + 1;
+
+		InitializeSRWLock(&busy_lock);
+		AcquireSRWLockExclusive(&busy_lock); // Only gets released when the job completes
+
+		InitializeSRWLock(&dependents_lock);
+		job_completed = false;
+		dependents_head = 0;
+		dependents_tail_next = &dependents_head;
+
+		// user_ptr is optional in this example (can be NULL)
+		//	PTP_CALLBACK_ENVIRON of NULL means use the process's default thread pool
+		PTP_CALLBACK_ENVIRON tpenviron = (PTP_CALLBACK_ENVIRON) user_ptr;
+
+		work_handle = CreateThreadpoolWork(work_callback, this, tpenviron);
+		DEBUG("create 0x%p\n", work_handle);
+
+		// If the work item was created successfully, set up the dependencies
+		if (work_handle)
+		{
+			// Initially, the only dependency we're going to release on completion
+			// of setup is the explicit setup dependency.
+			LONG release_count = 1;
+
+			// Add our dependencies. If we discover already-completed jobs during
+			// this, our release_count will go up.
+			for (int i = 0; i < num_dependencies; i++)
+				release_count += reinterpret_cast<JobContext *>(dependencies[i])->add_dependent(this, i);
+
+			// We're done setting up. Release our setup dependency as well as
+			// any other already-satisfied dependencies we ran into, and if that
+			// brings our outstanding dependency count down to 0, launch the job.
+			notify_dep_satisfied(release_count);
+		}
+	}
+
+	~JobContext()
+	{
+		DEBUG("close  0x%p\n", work_handle);
+		if (work_handle)
+			CloseThreadpoolWork(work_handle);
+	}
+
+	bool valid() const
+	{
+		return work_handle != NULL;
+	}
+
+	void wait()
+	{
+		// The busy_lock is held until the job completes;
+		// so trying to grab it "parks" us until then.
+		AcquireSRWLockExclusive(&busy_lock);
+
+		// Immediately release it. We use the lock just for the
+		// explicit hand-off. Most jobs only get waited for exactly
+		// once, but we can end up calling this multiple times in
+		// the error-handling case in thread_pool_run_job below.
+		ReleaseSRWLockExclusive(&busy_lock);
+	}
+};
+
+OO_U64 OODLE_CALLBACK example_jobify_win32tp_run_job(t_fp_Oodle_Job * fp_job, void * job_data, OO_U64 * dependencies, int num_dependencies, void * user_ptr)
+{
+	assert(num_dependencies <= OODLE_JOB_MAX_DEPENDENCIES);
+
+	// user_ptr coming back NULL is fine in this example
+
+	// Allocate a context struct to pass to the work callback
+	JobContext * ctx = new JobContext(fp_job, job_data, dependencies, num_dependencies, user_ptr);
+
+	if (!ctx || !ctx->valid())
+	{
+		// If something went wrong, bail on the job context, manually wait for
+		// our dependencies and then run the job directly.
+		if (ctx)
+			delete ctx;
+
+		for (int i = 0; i < num_dependencies; i++)
+			reinterpret_cast<JobContext *>(dependencies[i])->wait();
+
+		fp_job(job_data);
+		return 0;
+	}
+
+	// Return our context pointer as the handle.
+	return reinterpret_cast<OO_U64>(ctx);
+}
+
+void OODLE_CALLBACK example_jobify_win32tp_wait_job(OO_U64 job_handle, void * /*user_ptr*/)
+{
+	JobContext * job = reinterpret_cast<JobContext *>(job_handle);
+
+	// Wait for the outstanding job to finish
+	job->wait();
+
+	// Then dispose of the context and clean up
+	delete job;
+}
+
+void * example_jobify_init()
+{
+	SYSTEM_INFO systeminfo;
+	GetSystemInfo( &systeminfo );
+	int logical_cpu_count = systeminfo.dwNumberOfProcessors;
+		
+	example_jobify_run_job_fptr = example_jobify_win32tp_run_job;
+	example_jobify_wait_job_fptr = example_jobify_win32tp_wait_job;
+	
+	// target_parallelism is desired number of jobs for Oodle to run simultaneously
+	// logical_cpu_count includes hyperthreads
+	// often it's better to only run on physical cores
+	// unfortunately there's no simple API to get that
+	//	so for simplicity in this example code we use logical_cpu_count
+	//example_jobify_target_parallelism = logical_cpu_count/2; 
+	example_jobify_target_parallelism = logical_cpu_count;
+	
+	// The jobifyUserPtr in this implementation optionally points at a
+	// TP_CALLBACK_ENVIRON; if you're OK with using the defaults (this
+	// example is), just pass 0.
+	void * jobifyUserPtr = NULL;
+	
+	return jobifyUserPtr;
+}
+
+#else
+
+#pragma message("windows.h too old, no Jobify!")
+
+void * example_jobify_init()
+{
+	example_jobify_run_job_fptr = NULL;
+	example_jobify_wait_job_fptr = NULL;
+	example_jobify_target_parallelism = 0;
+	
+	return NULL;
+}
+
+#endif
+
+#endif // OOEX_PLATFORM_NT
+
+
