Moved Horde.Storage from git to p4

diff --git a/Engine/Source/Programs/Horde/HordeStorage/Jupiter.Common/Shared/Oodle/2.9.3/linux/examples/example_lz.cpp b/Engine/Source/Programs/Horde/HordeStorage/Jupiter.Common/Shared/Oodle/2.9.3/linux/examples/example_lz.cpp
new file mode 100644
--- /dev/null
+++ b/Engine/Source/Programs/Horde/HordeStorage/Jupiter.Common/Shared/Oodle/2.9.3/linux/examples/example_lz.cpp
@@ -0,0 +1,1171 @@
+//idoc(parent,examples)
+//idoc(tutorial,example_lz,example_lz : Example demonstrating LZ compression and decompression)
+
+/*
+
+Oodle example_lz
+
+Use the various LZ compress/decompress API's
+
+API's demonstrated here :
+
+OodleLZ : low level buffer compress/decompress
+
+OodleLZDecoder : streaming decoder
+
+OodleLZ_Async : high level async helpers
+
+*/
+
+#include "../include/oodle2x.h"
+#include "ooex.h"
+
+#include <stdlib.h>
+#include <stdio.h>
+#include <string.h>
+
+#include "make_example_input.h"
+
+#ifdef BUILDING_EXAMPLE_CALLER
+#define main example_lz
+#endif
+
+//===========================================================
+// file names :
+
+static const char * in_name = "oodle_example_input_file";
+static const char * out_name = "oodle_example_output_file";
+
+// just make the input buffer global here so all the tests can use it :
+static void * in_buffer = NULL;
+static OO_SINTa in_size = 0;
+
+//===========================================================
+
+// protos :
+
+void lz_test_1();
+void lz_test_2();
+void lz_test_4();
+void lz_test_9();
+void lz_test_10();
+void lz_test_11();
+void lz_test_12();
+void lz_test_13();
+
+//===========================================================
+
+extern "C" int main(int argc,char *argv[])
+{
+    OodleXLog_Printf_v1("usage: example_lz [in] [out]\n");
+
+    // Init Oodle systems with default options :
+    OodleXInitOptions opts;
+    if ( ! OodleX_Init_GetDefaults(OODLE_HEADER_VERSION,&opts) )
+    {
+        fprintf(stderr,"Oodle header version mismatch.\n");
+        return 10;
+    }
+    // change opts here if you like
+    if ( ! OodleX_Init(OODLE_HEADER_VERSION,&opts) )
+    {
+        fprintf(stderr,"OodleX_Init failed.\n");
+        return 10;
+    }
+
+    if ( argc >= 2 ) in_name = argv[1];
+    else make_example_input(in_name);
+    
+    if ( argc >= 3 ) out_name = argv[2];
+
+    OodleXLog_Printf_v1("lz test %s to %s\n",in_name,out_name);
+
+    // read the input file to the global buffer :
+    OO_S64 in_size_64;
+    in_buffer = OodleXIOQ_ReadMallocWholeFile_AsyncAndWait(in_name,&in_size_64);
+    
+    if ( ! in_buffer)
+    {
+        OodleXLog_Printf_v0("failed to read %s\n",in_name); 
+        return 10;
+    }
+    
+    in_size = OodleX_S64_to_SINTa_check( in_size_64 );
+    
+    lz_test_1();
+    lz_test_2();
+    lz_test_4();
+    lz_test_9();
+    lz_test_10();
+    lz_test_11();
+    lz_test_12();
+    lz_test_13();
+
+    OodleXLog_Printf_v1("\ndone.\n");
+
+    OodleXFree_IOAligned(in_buffer);
+
+    //OodleX_Shutdown();
+    OodleX_Shutdown(NULL,OodleX_Shutdown_LogLeaks_Yes,0);
+    
+    //OodleXLog_Printf_v1("press a key\n");
+    //fgetc(stdin);
+    
+    return 0;
+}
+
+//=================================================
+
+/*
+
+lz_test_1 :
+
+example of directly calling the simple buffer->buffer compression API's
+
+OodleLZ_Compress
+OodleLZ_Decompress
+
+*/
+
+void lz_test_1()
+{
+    OodleXLog_Printf_v0("lz_test_1\n");
+    // allocate compressed buffer & decoded buffer of the correct sizes :
+
+    OO_SINTa comp_buf_size = OodleLZ_GetCompressedBufferSizeNeeded(OodleLZ_Compressor_Kraken,in_size);
+    void * comp_buf = OodleXMalloc( comp_buf_size );
+    OOEX_ASSERT( comp_buf != NULL );
+
+    void * dec_buf = OodleXMalloc( in_size );
+    OOEX_ASSERT( dec_buf != NULL );
+    
+    //---------------------------------------------------
+    
+    // compress buffer -> buffer :
+    
+    OodleLZ_CompressionLevel level = OodleLZ_CompressionLevel_VeryFast;
+    
+    OO_SINTa comp_len = OodleLZ_Compress(OodleLZ_Compressor_Kraken,in_buffer,in_size,comp_buf,level);
+    
+    OodleXLog_Printf_v1("Kraken compress %d -> %d\n",(int)in_size,(int)comp_len);
+
+    // decompress :
+
+    OO_SINTa dec_len = OodleLZ_Decompress(comp_buf,comp_len,dec_buf,in_size,OodleLZ_FuzzSafe_Yes);
+
+    OOEX_ASSERT_ALWAYS( dec_len == in_size );
+    OOEX_ASSERT_ALWAYS( memcmp(in_buffer,dec_buf,in_size) == 0 );
+
+    //---------------------------------------------------
+    // do it again with another compressor, and custom options :
+    
+    // compress buffer -> buffer :
+    
+    OodleLZ_CompressOptions options = * OodleLZ_CompressOptions_GetDefault(OodleLZ_Compressor_Leviathan,level);
+    options.spaceSpeedTradeoffBytes = OODLELZ_SPACESPEEDTRADEOFFBYTES_DEFAULT/2; // favor size over decode speed
+    
+    comp_len = OodleLZ_Compress(OodleLZ_Compressor_Leviathan,in_buffer,in_size,comp_buf,level,&options);
+    
+    OodleXLog_Printf_v1("Leviathan compress %d -> %d\n",(int)in_size,(int)comp_len);
+
+    // decompress :
+
+    dec_len = OodleLZ_Decompress(comp_buf,comp_len,dec_buf,in_size,OodleLZ_FuzzSafe_Yes);
+
+    OOEX_ASSERT( dec_len == in_size );
+    OOEX_ASSERT( memcmp(in_buffer,dec_buf,in_size) == 0 );
+
+    //-------------------------------------
+    // free buffers :
+
+    OodleXFree(comp_buf);
+    OodleXFree(dec_buf);
+}
+
+/*
+
+ lz_test_2
+    
+example of using the OodleLZ_Async_ async helper functions (eg OodleXLZ_Decompress_Wide_Async)
+this is the simple way to get the best performance
+
+Use the seekChunkReset option on to get a seekable packed stream.
+
+*/
+
+void lz_test_2()
+{
+    OodleXLog_Printf_v0("lz_test_2\n");
+
+    // allocate compressed buffer & decoded buffer of the correct sizes :
+
+    OO_SINTa comp_buf_size = OodleLZ_GetCompressedBufferSizeNeeded(OodleLZ_Compressor_Kraken,in_size);
+    void * comp_buf = OodleXMalloc( comp_buf_size );
+    OOEX_ASSERT( comp_buf != NULL );
+
+    void * dec_buf = OodleXMalloc( in_size );
+    OOEX_ASSERT( dec_buf != NULL );
+    
+    //---------------------------------------------------
+    
+    OodleLZ_CompressOptions options = * OodleLZ_CompressOptions_GetDefault(OodleLZ_Compressor_Kraken,OodleLZ_CompressionLevel_Normal);
+    // turn on seekChunkReset
+    //  this makes chunks independent so they can be decompressed in any order (not just linear)
+    options.seekChunkReset = true;
+    options.seekChunkLen = OodleLZ_MakeSeekChunkLen(in_size,8);
+
+    // with seekChunkReset on, compression will also go in parallel
+    //  (actually compression can *always* run in parallel, but seekChunkReset makes it scale more linearly,
+    //    and parallelize on a smaller granularity) 
+    // use the OodleLZ_Compressor_Kraken compressor
+    
+    OodleXHandle h = OodleXLZ_Compress_Async(OodleXAsyncSelect_Full,OodleLZ_Compressor_Kraken,in_buffer,in_size,comp_buf,OodleLZ_CompressionLevel_Normal,&options);
+
+    // ... do other game work while compression runs ...
+    
+    OO_SINTa comp_len = -1; 
+    OodleXLZ_Compress_Wait_GetResult(h,&comp_len);
+
+    OodleXLog_Printf_v1("LZ compress %d -> %d\n",(int)in_size,(int)comp_len);
+
+    //-----------------------------------------------------
+    // make seek entries :
+    //  seek entries allow parallel decompression
+    
+    OodleLZ_SeekTable * seekTable = OodleLZ_CreateSeekTable(OodleLZSeekTable_Flags_None,options.seekChunkLen,in_buffer,in_size,comp_buf,comp_len);
+    OOEX_ASSERT( seekTable != NULL );
+
+    //-----------------------------------------------------
+
+    OodleXHandle dh = OodleXLZ_Decompress_Wide_Async(OodleXAsyncSelect_Full,seekTable,comp_buf,comp_len,dec_buf,in_size,
+                            OodleLZ_FuzzSafe_Yes,OodleLZ_CheckCRC_No,OodleLZ_Verbosity_None,0,0,OodleLZ_PackedRawOverlap_No,0,0);
+    
+    // ... do other game work while decompression runs ...
+    
+    OodleXStatus st = OodleX_Wait(dh,OodleXHandleDeleteIfDone_Yes);
+    OOEX_ASSERT_ALWAYS( st == OodleXStatus_Done );
+    
+    //-----------------------------------------------------
+    // check :
+        
+    OOEX_ASSERT( memcmp(in_buffer,dec_buf,in_size) == 0 );
+
+    //-------------------------------------
+    // free buffers :
+
+    OodleLZ_FreeSeekTable(seekTable);
+    OodleXFree(comp_buf);
+    OodleXFree(dec_buf);
+}
+
+
+/*
+
+ test_4 :
+
+example of seeking in packed stream
+and firing per-chunk decompression tasks
+
+Sort of like what OodleXLZ_Decompress_Wide_Async does internally.
+
+*/
+
+void lz_test_4()
+{
+    OodleXLog_Printf_v0("lz_test_4\n");
+    // allocate compressed buffer & decoded buffer of the correct sizes :
+
+    OO_SINTa comp_buf_size = OodleLZ_GetCompressedBufferSizeNeeded(OodleLZ_Compressor_Kraken,in_size);
+    void * comp_buf = OodleXMalloc( comp_buf_size );
+    OOEX_ASSERT( comp_buf != NULL );
+
+    void * dec_buf = OodleXMalloc( in_size );
+    OOEX_ASSERT( dec_buf != NULL );
+    
+    // set up compress options for independent seek chunks of the smallest allowed size :
+    OodleLZ_CompressOptions lzOptions = * OodleLZ_CompressOptions_GetDefault(OodleLZ_Compressor_Kraken,OodleLZ_CompressionLevel_VeryFast);
+    lzOptions.seekChunkReset = true;
+    //lzOptions.seekChunkLen = OODLELZ_BLOCK_LEN;
+    // make a seek chunk len to target 32 chunks :
+    lzOptions.seekChunkLen = OodleLZ_MakeSeekChunkLen(in_size,32);
+        
+    //---------------------------------------------------
+    
+    // compress buffer -> buffer :
+        
+    OO_SINTa comp_len = OodleLZ_Compress(OodleLZ_Compressor_Kraken,in_buffer,in_size,comp_buf,OodleLZ_CompressionLevel_VeryFast,&lzOptions);
+    
+    OodleXLog_Printf_v1("LZ compress %d -> %d\n",(int)in_size,(int)comp_len);
+    
+    //---------------------------------------------------
+    // decompress by seeking and firing async decodes
+    
+    OO_SINTa maxNumSeeks = (in_size + lzOptions.seekChunkLen-1)/lzOptions.seekChunkLen;
+    
+    OodleXHandle * handles = (OodleXHandle *) OodleXMalloc( sizeof(OodleXHandle) * maxNumSeeks );
+    OOEX_ASSERT( handles != NULL );
+
+    OO_S32 numHandles = 0;
+    
+    {   
+    OO_SINTa dec_pos = 0;
+    OO_U8 * comp_ptr = (OO_U8 *)comp_buf;
+    OO_SINTa comp_avail = comp_len;
+
+    while(dec_pos < in_size)    
+    {
+        OO_SINTa dec_chunk_len = OOEX_MIN(lzOptions.seekChunkLen, (in_size-dec_pos) );
+
+        OodleXHandle h = OodleXLZ_Decompress_Narrow_Async(OodleXAsyncSelect_Full,comp_ptr,comp_avail,(OO_U8 *)dec_buf+dec_pos,dec_chunk_len);
+        //  OodleLZ_CheckCRC_No,OodleLZ_Verbosity_None, etc
+        handles[numHandles++] = h;
+        
+        OO_SINTa seek_step = OodleLZ_GetCompressedStepForRawStep(comp_ptr,comp_avail,dec_pos,dec_chunk_len);
+        
+        comp_ptr += seek_step;
+        dec_pos += dec_chunk_len;
+        
+        // wait on handles[numHandles-128] to prevent handle count exceeding handle table size :
+        if ( numHandles >= 128 )
+        {
+            OodleXStatus st = OodleX_WaitAndDelete(handles[numHandles-128]);
+            OOEX_ASSERT_ALWAYS( st == OodleXStatus_Done );
+        }
+    }
+    
+    }
+    
+    // ... do other game work while async decomps run ...
+    
+    OodleXStatus st = OodleX_WaitAll(handles,numHandles,OodleXHandleDeleteIfDone_Yes);
+    OOEX_ASSERT_ALWAYS( st == OodleXStatus_Done );
+
+    OodleXFree(handles);
+
+
+    // check it's okay :
+    OOEX_ASSERT_ALWAYS( memcmp(in_buffer,dec_buf,in_size) == 0 );
+
+    //-------------------------------------
+    // free buffers :
+
+    OodleXFree(comp_buf);
+    OodleXFree(dec_buf);
+}
+
+
+/*
+
+test_9 :
+
+Demonstrate separate block compression & decompression
+
+OodleLZ blocks can be concatenated to form a single valid LZ data stream
+That means you can just call OodleLZ_Compress on separate blocks and append the output,
+then decode in one call.
+
+OodleLZ blocks that were made from separate Compress calls will be independent
+unless you specified dictionary backup in the encode, which makes them depend on previous
+data.
+
+
+The rules are :
+
+1. OodleLZ Decompress can be called on invidual blocks (OODLELZ_BLOCK_LEN) if :
+
+    they are seek-chunk-reset points,
+    OR if they were made by separate OodleLZ Compress calls 
+    OR if the compressor does not carry state across blocks (OodleLZ_Compressor_MustDecodeWithoutResets)
+
+2. OodleLZ Decompress must get the same dictionary as OodleLZ Compress saw
+    no previous dictionary is needed if it's a seek-chunk-reset point
+    (the start of an OodleLZ_Compress call is always a seek reset point,
+    if no dictionary backup is provided to the encoder) 
+
+*/
+
+void lz_test_9()
+{
+    OodleXLog_Printf_v0("lz_test_9\n");
+
+    //---------------------------------------------------
+    // split the buffer into two pieces
+    //  such that the split point is a valid seek chunk point :
+    OO_SINTa block_size = OodleLZ_MakeSeekChunkLen(in_size,2);
+
+    if ( block_size >= in_size )
+    {
+        // too small to split at seek chunk
+        return;
+    }
+    
+    char * in1 = (char *)in_buffer;
+    char * in2 = in1 + block_size;
+    OO_SINTa len1 = block_size;
+    OO_SINTa len2 = in_size - len1;
+    
+    OodleXLog_Printf_v1("Chunks : %d + %d\n",(int)len1,(int)len2);
+
+    //---------------------------------------------------
+    
+    // allocate compressed buffer & decoded buffer of the correct sizes :
+
+    OO_SINTa comp_buf_size = OodleLZ_GetCompressedBufferSizeNeeded(OodleLZ_Compressor_Kraken,in_size);
+    void * comp_buf = OodleXMalloc( comp_buf_size );
+    OOEX_ASSERT( comp_buf != NULL );
+
+    void * dec_buf = OodleXMalloc( in_size );
+    OOEX_ASSERT( dec_buf != NULL );
+    
+    //---------------------------------------------------
+
+    OodleLZ_Compressor compressor = OodleLZ_Compressor_Kraken;
+    OodleLZ_CompressionLevel level = OodleLZ_CompressionLevel_Fast; 
+    OodleLZ_CompressOptions options = * OodleLZ_CompressOptions_GetDefault(compressor,level);
+    
+    // options does NOT have seek resets by default
+
+    //-----------------------------
+    // compress as one part :
+    {
+
+    OO_SINTa comp_len = OodleLZ_Compress(compressor,in_buffer,in_size,comp_buf,level,&options);
+
+    OodleXLog_Printf_v1("Whole buffer compress : %d -> %d\n",(int)in_size,(int)comp_len);   
+
+    // normal one part decompression :
+    memset(dec_buf,0xEE,in_size);
+    OodleLZ_Decompress(comp_buf,comp_len,dec_buf,in_size,OodleLZ_FuzzSafe_Yes);
+    OOEX_ASSERT_ALWAYS( memcmp(in_buffer,dec_buf,in_size) == 0 );
+
+    OO_U8 * comp_end = (OO_U8 *)comp_buf + comp_len;
+
+    //===================================================================
+
+    // can decode in two calls with the full dictionary, but only for compressors that don't carry state :
+    if ( ! OodleLZ_Compressor_MustDecodeWithoutResets(compressor) )
+    {
+        //-------------------------------------------------------------
+        // decode as two parts : (len1,len2) :
+        memset(dec_buf,0xEE,in_size);
+
+        OO_SINTa dec_comp_len1 = OodleLZ_GetCompressedStepForRawStep(comp_buf,comp_len,0,len1);
+                
+        OodleLZ_Decompress(comp_buf,comp_len,dec_buf,len1,OodleLZ_FuzzSafe_Yes);
+
+        // decompress second part with dictionary base :
+
+        OodleLZ_Decompress((char *)comp_buf+dec_comp_len1,comp_len-dec_comp_len1,(char *)dec_buf+len1,len2,
+                                    OodleLZ_FuzzSafe_Yes,OodleLZ_CheckCRC_No,OodleLZ_Verbosity_None,dec_buf,in_size);
+
+        OOEX_ASSERT_ALWAYS( memcmp(in_buffer,dec_buf,in_size) == 0 );
+
+        //-------------------------------------------------------------
+        // can also decode block by block :
+        memset(dec_buf,0xEE,in_size);
+        
+        // scan comp_ptr through blocks :
+        OO_U8 * comp_ptr = (OO_U8 *)comp_buf;
+        for(OO_SINTa block_pos=0;block_pos < in_size;block_pos += OODLELZ_BLOCK_LEN)
+        {
+            OO_SINTa block_len = OOEX_MIN(OODLELZ_BLOCK_LEN, (in_size - block_pos) );
+            
+            OO_SINTa block_comp_len = OodleLZ_GetCompressedStepForRawStep(comp_ptr,comp_end-comp_ptr,0,block_len);
+            
+            // decode current block, with window set to whole buffer :
+            OO_SINTa got_pos = OodleLZ_Decompress(comp_ptr,block_comp_len,(char *)dec_buf+block_pos,block_len,
+                                    OodleLZ_FuzzSafe_Yes,OodleLZ_CheckCRC_No,OodleLZ_Verbosity_None,dec_buf,in_size);
+            
+            OOEX_ASSERT_ALWAYS( got_pos == block_pos+block_len );
+            
+            comp_ptr += block_comp_len;
+        }
+
+        OOEX_ASSERT_ALWAYS( memcmp(in_buffer,dec_buf,in_size) == 0 );
+        //-------------------------------------------------------------
+    }
+    
+    }
+    //-----------------------------
+    // two part compression with overlap :
+    // two compress calls, but using the full window, so decompression must use full window as well
+    {
+    
+    OO_SINTa comp_len1 = OodleLZ_Compress(compressor,in1,len1,comp_buf,level,&options,in_buffer);
+    OO_SINTa comp_len2 = OodleLZ_Compress(compressor,in2,len2,(char *)comp_buf + comp_len1,level,&options,in_buffer);
+    OO_SINTa comp_len = comp_len1 + comp_len2;
+
+    OodleXLog_Printf_v1("Two part compress with overlap : %d -> %d\n",(int)in_size,(int)comp_len);  
+
+    // must decode whole buffer, but can do it in two calls :
+
+    // you can always just do a whole buffer decode here :
+    memset(dec_buf,0xEE,in_size);
+    OodleLZ_Decompress(comp_buf,comp_len,dec_buf,in_size,OodleLZ_FuzzSafe_Yes);
+    OOEX_ASSERT_ALWAYS( memcmp(in_buffer,dec_buf,in_size) == 0 );
+
+    // or incremental, but with the whole dictionary :
+    memset(dec_buf,0xEE,in_size);
+
+    OO_SINTa dec_comp_len1 = OodleLZ_GetCompressedStepForRawStep(comp_buf,comp_len,0,len1);
+        
+    OodleLZ_Decompress(comp_buf,comp_len,dec_buf,len1,OodleLZ_FuzzSafe_Yes);
+        
+    OodleLZ_Decompress((char *)comp_buf+dec_comp_len1,comp_len-dec_comp_len1,(char *)dec_buf+len1,len2,
+                                OodleLZ_FuzzSafe_Yes,OodleLZ_CheckCRC_No,OodleLZ_Verbosity_None,dec_buf,in_size);
+
+    OOEX_ASSERT_ALWAYS( memcmp(in_buffer,dec_buf,in_size) == 0 );
+    
+    }
+    //-----------------------------
+    // two part no overlap :
+    // second compress doesn't use earlier dictionary here
+    // so decompression can be done in two pieces with no overlap
+    {
+    
+    OO_SINTa comp_len1 = OodleLZ_Compress(compressor,in1,len1,comp_buf,level,&options);
+    OO_SINTa comp_len2 = OodleLZ_Compress(compressor,in2,len2,(char *)comp_buf + comp_len1,level,&options);
+    OO_SINTa comp_len = comp_len1 + comp_len2;
+
+    OodleXLog_Printf_v1("Two part compress no overlap : %d -> %d\n",(int)in_size,(int)comp_len); 
+
+    // can decode in two parts :
+
+    // you can always just do a whole buffer decode here :
+    memset(dec_buf,0xEE,in_size);
+    OodleLZ_Decompress(comp_buf,comp_len,dec_buf,in_size,OodleLZ_FuzzSafe_Yes);
+    OOEX_ASSERT_ALWAYS( memcmp(in_buffer,dec_buf,in_size) == 0 );
+
+    // or incremental :
+    memset(dec_buf,0xEE,in_size);
+    
+    OO_SINTa dec_comp_len1 = OodleLZ_GetCompressedStepForRawStep(comp_buf,comp_len,0,len1);
+
+    // no dictionary backup needed
+    // decode in reverse order to simulate random access :
+
+    OodleLZ_Decompress((char *)comp_buf+dec_comp_len1,comp_len-dec_comp_len1,(char *)dec_buf+len1,len2,OodleLZ_FuzzSafe_Yes);
+        
+    OodleLZ_Decompress(comp_buf,comp_len,dec_buf,len1,OodleLZ_FuzzSafe_Yes);
+
+    OOEX_ASSERT_ALWAYS( memcmp(in_buffer,dec_buf,in_size) == 0 );
+    
+    }
+    //-----------------------------
+    // two part no overlap via seek reset :
+    // seek reset system is equivalent to splitting Compress calls like the above
+    {   
+    options.seekChunkReset = true;
+    options.seekChunkLen = (OO_S32)block_size;
+    
+    OodleXLog_Printf_v1("seekChunkLen : %d\n",options.seekChunkLen);
+
+    OO_SINTa comp_len = OodleLZ_Compress(compressor,in_buffer,in_size,comp_buf,level,&options);
+
+    OodleXLog_Printf_v1("Whole buffer compress seek reset : %d -> %d\n",(int)in_size,(int)comp_len); 
+
+    // can decode in two parts :
+
+    // you can always just do a whole buffer decode here :
+    memset(dec_buf,0xEE,in_size);
+
+    OodleXLog_Printf_v1("one part : \n");
+
+    OodleLZ_Decompress(comp_buf,comp_len,dec_buf,in_size,OodleLZ_FuzzSafe_Yes);
+    OOEX_ASSERT_ALWAYS( memcmp(in_buffer,dec_buf,in_size) == 0 );
+
+    // or incremental :
+    memset(dec_buf,0xEE,in_size);
+    
+    OodleXLog_Printf_v1("two part : \n");
+
+    OO_SINTa dec_comp_len1 = OodleLZ_GetCompressedStepForRawStep(comp_buf,comp_len,0,len1);
+    
+    OodleXLog_Printf_v1("dec_comp_len1 = %d\n",(int)dec_comp_len1);
+
+    // no dictionary backup needed
+    // decode in reverse order to simulate random access :
+
+    OodleLZ_Decompress((char *)comp_buf+dec_comp_len1,comp_len-dec_comp_len1,(char *)dec_buf+len1,len2,OodleLZ_FuzzSafe_Yes);
+
+    OodleLZ_Decompress(comp_buf,comp_len,dec_buf,len1,OodleLZ_FuzzSafe_Yes);
+
+    OOEX_ASSERT_ALWAYS( memcmp(in_buffer,dec_buf,in_size) == 0 );
+    }
+    
+    //=============================================
+
+    OodleXFree(dec_buf);
+    OodleXFree(comp_buf);
+}
+
+
+/*
+
+lz_test_10 :
+
+example of using the incremental/streaming decoder
+
+OodleLZDecoder_Create, etc.
+
+this example shows decoding *from* a limited window
+
+outputs into a single buffer
+
+this example simulates using a limited IO buffer for compressed data
+
+it decodes quanta from the available compressed data
+
+*/
+
+
+void lz_test_10()
+{
+    OodleXLog_Printf_v0("lz_test_10\n");
+    // allocate compressed buffer & decoded buffer of the correct sizes :
+
+    OodleLZ_Compressor compressor = OodleLZ_Compressor_Kraken;
+    OodleLZ_CompressionLevel level = OodleLZ_CompressionLevel_Fast;
+    
+    OO_SINTa comp_buf_size = OodleLZ_GetCompressedBufferSizeNeeded(OodleLZ_Compressor_Kraken,in_size);
+    OO_U8 * comp_buf = (OO_U8 *) OodleXMalloc( comp_buf_size );
+    OOEX_ASSERT( comp_buf != NULL );
+
+    OO_U8 * dec_buf = (OO_U8 *)OodleXMalloc( in_size );
+    OOEX_ASSERT( dec_buf != NULL );
+    
+    //---------------------------------------------------
+    
+    // compress buffer -> buffer :
+    
+    OO_SINTa comp_len = OodleLZ_Compress(compressor,in_buffer,in_size,comp_buf,level);
+    
+    OodleXLog_Printf_v1("LZ compress %d -> %d\n",(int)in_size,(int)comp_len);
+
+    //---------------------------------------------------
+    // decompress with incremental streaming decoder :
+    
+    // we're now going to pretend that "comp_buf" is in a file
+    //  and we can't read the whole thing
+    
+    // 64k IO buffer to stress the code :
+    //  obviously you would use much larger
+    //  must be at least enough for 1 whole compressed quantum (OODLELZ_BLOCK_MAX_COMPLEN = 256k + 2)
+    //  (for Kraken / whole-block compressors you probably want 512k io_buffer minimum)
+    const OO_SINTa io_buffer_size = OODLELZ_BLOCK_LEN*2;
+    OO_U8 * io_buffer = (OO_U8 *) OodleXMalloc_IOAligned(io_buffer_size);
+    OOEX_ASSERT( io_buffer != NULL );
+    
+    OO_SINTa io_buffer_avail = 0; // starts empty
+    
+    OO_SINTa comp_file_io_pos = 0; // simulated compressed file next read pos
+    
+    {
+    
+    // make the Decoder object using on-stack memory :
+    //OO_S32 memSize = OodleLZDecoder_MemorySizeNeeded(compressor,in_size);
+    
+    OodleLZDecoder * decoder = OodleLZDecoder_Create(compressor,in_size,NULL,0);
+
+    OO_SINTa io_buffer_pos = 0;
+    OO_U8 * dec_buf_ptr = (OO_U8 *) dec_buf;
+    OO_U8 * dec_buf_end = (OO_U8 *) dec_buf + in_size;;
+
+    while(dec_buf_ptr<dec_buf_end)
+    {
+        // see if we can do a "read" into the io_buffer :
+        if ( comp_file_io_pos < comp_len )
+        {
+            // don't bother with an IO unless I have some minimum amount of room :
+            OO_SINTa min_io_size = 16*1024;
+            if ( (io_buffer_size - io_buffer_avail) > min_io_size )
+            {
+                OO_SINTa io_size = OOEX_MIN( (io_buffer_size - io_buffer_avail), (comp_len - comp_file_io_pos) );
+                // stress - limit IO size :
+                //io_size = OOEX_MIN(io_size,64*1024);
+                
+                //OodleXLog_Printf_v1("IO read : %d at %d\n",(int)io_size,(int)comp_file_io_pos);
+                
+                // IO read :
+                memcpy(io_buffer+io_buffer_avail,comp_buf+comp_file_io_pos,io_size);
+                comp_file_io_pos += io_size;
+                io_buffer_avail += io_size;
+            }
+        }
+    
+        // ask the Decoder for a partial decode :
+        OodleLZ_DecodeSome_Out out;
+        
+        OO_BOOL ok = OodleLZDecoder_DecodeSome(decoder,&out,
+            dec_buf,(dec_buf_ptr - dec_buf),in_size,(dec_buf_end - dec_buf_ptr),
+            io_buffer+io_buffer_pos,io_buffer_avail-io_buffer_pos);
+
+        // real usage should check error return conditions
+        OOEX_ASSERT_ALWAYS( ok );
+
+        OO_S32 decoded = out.decodedCount;
+        OO_S32 comp_used = out.compBufUsed;
+                                
+        // advance the decoder :
+        dec_buf_ptr += decoded;
+        io_buffer_pos += comp_used;
+        
+        OOEX_ASSERT( out.curQuantumCompLen < io_buffer_size );
+        
+        //OodleXLog_Printf_v1("decoded : %d using %d\n",decoded,comp_used);
+        
+        if ( decoded == 0 )
+        {
+            // couldn't decode anything
+            // this should only happen because we're near the end of the io buffer
+            // and don't have enough compressed data to do antyhing
+            OOEX_ASSERT( io_buffer_pos > 0 );
+            // slide down the io buffer so it can refill
+            OO_SINTa io_buffer_keep = io_buffer_avail - io_buffer_pos;
+            memmove(io_buffer,io_buffer+io_buffer_pos,io_buffer_keep);
+            io_buffer_pos = 0;
+            io_buffer_avail = io_buffer_keep;
+        }
+    }
+
+    OOEX_ASSERT( comp_file_io_pos == comp_len );
+
+    OodleLZDecoder_Destroy(decoder);
+
+    }
+    
+    // check it's okay :
+    OOEX_ASSERT_ALWAYS( memcmp(in_buffer,dec_buf,in_size) == 0 );
+
+    //-------------------------------------
+    // free buffers :
+
+    OodleXFree(comp_buf);
+    OodleXFree(dec_buf);
+    OodleXFree_IOAligned(io_buffer);
+}
+
+
+/*
+
+lz_test_11 :
+
+example of using the incremental/streaming decoder
+
+OodleLZDecoder_Create, etc.
+
+this example simulates using a limited IO buffer for compressed data (like lz_test_10)
+
+Kraken does not currently have a true "sliding window" decoder; it can't wrap around a circular
+window.  This example shows how to simulate a sliding window with the Kraken decoder by sliding down chunks.
+
+It decodes 256k at a time into a 512k window.  It decodes into the second half of the window, with
+the first window filled by the previous decode.  After each decode, it memcopies down the data to be used
+as dictionary for the next block.
+
+---------
+
+In general this method should not be used if you can just decode directly into the output buffer.
+That's always the best way if possible.
+
+One case where you might want to use this is if your output buffer is in non-cached graphics memory.
+
+---------
+
+The simpler alternative to this is just to reset every 256k block, so there's no dictionary overlap.
+eg. just set :
+
+options.seekChunkReset = true;
+options.seekChunkLen = OODLELZ_BLOCK_LEN;
+
+then you can use a 256k decode output window and don't need to memcpy to slide down the dictionary.
+The disadvantage of resetting is just lower compression.
+
+*/
+
+
+void lz_test_11()
+{
+    OodleXLog_Printf_v0("lz_test_11\n");
+
+    //*
+    // fast encoder:
+    OodleLZ_Compressor compressor = OodleLZ_Compressor_Kraken;
+    OodleLZ_CompressionLevel level = OodleLZ_CompressionLevel_VeryFast;
+    /*/
+    // slowest encoder:
+    OodleLZ_Compressor compressor = OodleLZ_Compressor_Leviathan;
+    OodleLZ_CompressionLevel level = OodleLZ_CompressionLevel_Optimal5;
+    /**/
+
+    //---------------------------------------------------
+    
+    /*
+    // minimum size :
+    // we will decode 256k (one "block") at a time
+    // +256k more for dictionary references to preceding data
+    
+    OO_S32 decode_window_size = 2*OODLELZ_BLOCK_LEN; // OODLELZ_BLOCK_LEN = 256k
+    OO_S32 dictionary_size = OODLELZ_BLOCK_LEN;
+    /*/
+    // more reasonable size :
+    // dictionary limit 2M
+    // decode in a 3M window, so we do a memcpy after every 1M streamed
+    
+    OO_S32 decode_window_size = 3*1024*1024; // OODLELZ_BLOCK_LEN = 256k
+    OO_S32 dictionary_size = 2*1024*1024;
+    /**/
+    
+    OodleXLog_Printf_v1("dictionary_size : %d, decode_window_size : %d\n",dictionary_size,decode_window_size);
+    
+    //---------------------------------------------------
+    
+    OO_SINTa comp_buf_size = OodleLZ_GetCompressedBufferSizeNeeded(compressor,in_size);
+    OO_U8 * comp_buf = (OO_U8 *) OodleXMalloc( comp_buf_size );
+    OOEX_ASSERT( comp_buf != NULL );
+
+    // dec_window is our scratch circular window
+    OO_U8 * dec_window = (OO_U8 *)OodleXMalloc( decode_window_size );
+    OOEX_ASSERT( dec_window != NULL );
+    
+    // dec_out_buf is the final output location (perhaps uncached graphics memory)
+    OO_U8 * dec_out_buf = (OO_U8 *)OodleXMalloc( in_size );
+    OOEX_ASSERT( dec_out_buf != NULL );
+    
+    // decoderMem is used for the OodleLZ decoder object
+    OO_S32 memSize = OodleLZDecoder_MemorySizeNeeded(compressor);
+    OO_U8 * decoderMem = (OO_U8 *) OodleXMalloc(memSize);
+    
+    //---------------------------------------------------
+    // compress buffer -> buffer :
+    
+    // limit dictionarySize so matches can't go out of the decode window :
+    OodleLZ_CompressOptions options = * OodleLZ_CompressOptions_GetDefault(compressor,level);
+    options.dictionarySize = dictionary_size;
+    
+    OO_SINTa comp_len = OodleLZ_Compress(compressor,in_buffer,in_size,comp_buf,level,&options);
+    
+    OodleXLog_Printf_v1("LZ compress %d -> %d\n",(int)in_size,(int)comp_len);
+
+    //---------------------------------------------------
+    // decompress with incremental streaming decoder :
+    
+    // we're now going to pretend that "comp_buf" is in a file
+    //  and we can't read the whole thing
+    
+    // IO buffer must be at least enough for 1 whole quantum (256k + a little bit)
+    //  Kraken uses "large block quantum" (256k) not the old 16k quantum
+    const OO_SINTa io_buffer_size = (256+63)*1024;
+    OO_U8 io_buffer[io_buffer_size];
+    
+    OO_SINTa io_buffer_avail = 0; // starts empty
+    
+    OO_SINTa comp_file_io_pos = 0; // simulated compressed file next read pos
+    
+    {   
+    OodleLZDecoder * decoder = OodleLZDecoder_Create(compressor,in_size,decoderMem,memSize);
+
+    OO_SINTa io_buffer_pos = 0;
+    OO_U8 * dec_out_ptr = (OO_U8 *) dec_out_buf;
+    OO_U8 * dec_out_end = (OO_U8 *) dec_out_buf + in_size;;
+
+    OO_SINTa dec_window_pos = 0;
+
+    while(dec_out_ptr<dec_out_end)
+    {
+        // see if we can do a "read" into the io_buffer :
+        if ( comp_file_io_pos < comp_len )
+        {
+            OO_SINTa min_io_size = 16*1024;
+            if ( (io_buffer_size - io_buffer_avail) > min_io_size )
+            {
+                OO_SINTa io_size = OOEX_MIN( (io_buffer_size - io_buffer_avail), (comp_len - comp_file_io_pos) );
+                
+                //OodleXLog_Printf_v1("IO read : %d at %d\n",(int)io_size,(int)comp_file_io_pos);
+                
+                // IO read :
+                memcpy(io_buffer+io_buffer_avail,comp_buf+comp_file_io_pos,io_size);
+                comp_file_io_pos += io_size;
+                io_buffer_avail += io_size;
+            }
+        }
+    
+        // when dec_window_pos reaches end of window :
+        if ( (dec_window_pos + OODLELZ_BLOCK_LEN) > decode_window_size )
+        {
+            OodleXLog_Printf_v1("slide!\n");
+        
+            // slide down the dictionary for the next block :
+            memmove(dec_window,dec_window + dec_window_pos - dictionary_size,dictionary_size);
+            dec_window_pos = dictionary_size;
+        }
+            
+        OodleXLog_Printf_v1("decode : at %d in window, %d in output\n",(int)dec_window_pos,(int)(dec_out_ptr - dec_out_buf));
+    
+        // ask the Decoder for a partial decode :
+        OodleLZ_DecodeSome_Out out;
+        
+        // no need to truncate dec_avail at the end, the "in_size" passed to LZDecoder_Create does this
+        //OO_SINTa dec_avail = OOEX_MIN(dec_out_remain,decode_window_size - dec_window_pos);        
+        OO_SINTa dec_avail = decode_window_size - dec_window_pos;
+        
+        OO_BOOL ok = OodleLZDecoder_DecodeSome(decoder,&out,
+            dec_window,dec_window_pos,
+            in_size, //decode_window_size, // !!
+            dec_avail,
+            io_buffer+io_buffer_pos,io_buffer_avail-io_buffer_pos,
+            OodleLZ_FuzzSafe_Yes,OodleLZ_CheckCRC_No);
+
+        // !! = this is a bit funny; we lie about the window size here
+        //  if Kraken had a true sliding window decoder (like eg. LZH or LZB16 does)
+        //  then you would pass decode_window_size here and let it do the wrapping
+        //  but Kraken does not, so we pretend that we are decoding the whole file
+        //  so that OodleLZDecoder_DecodeSome doesn't try to use its sliding window path
+        //  (which would fail)
+        //  the "dec_avail" value prevents DecodeSome from going out of the window
+        //  and we simulate the sliding using the memcpy
+
+        // real usage should check error return conditions
+        OOEX_ASSERT_ALWAYS( ok );
+
+        OO_S32 decoded = out.decodedCount;
+        OO_S32 comp_used = out.compBufUsed;
+        
+        io_buffer_pos += comp_used;
+        
+        OOEX_ASSERT( out.curQuantumCompLen < io_buffer_size );
+                
+        //OodleXLog_Printf_v1("decoded : %d using %d\n",decoded,comp_used);
+        
+        if ( decoded == 0 )
+        {
+            // couldn't decode anything
+            // this should only happen because we're near the end of the io buffer
+            // and don't have enough compressed data to do antyhing
+            OOEX_ASSERT( io_buffer_pos > 0 );
+            // slide down the io buffer so it can refill
+            OO_SINTa io_buffer_keep = io_buffer_avail - io_buffer_pos;
+            memmove(io_buffer,io_buffer+io_buffer_pos,io_buffer_keep);
+            io_buffer_pos = 0;
+            io_buffer_avail = io_buffer_keep;
+        }
+        else
+        {
+            // copy out the decoded data :
+            OO_U8 * dec_window_ptr = dec_window + dec_window_pos;
+            // dec_out_ptr is the final output memory ; eg. perhaps uncached graphics memory
+            // "decoded" is always OODLELZ_BLOCK_LEN unless we hit EOF
+            memcpy(dec_out_ptr,dec_window_ptr,decoded);
+            
+            // advance the decoder :
+            dec_out_ptr += decoded;
+            dec_window_pos += decoded;
+        }
+    }
+
+    OOEX_ASSERT( comp_file_io_pos == comp_len );
+    OOEX_ASSERT( dec_out_ptr == dec_out_end );
+
+    OodleLZDecoder_Destroy(decoder);
+    }
+    
+    // check it's okay :
+    OOEX_ASSERT_ALWAYS( memcmp(in_buffer,dec_out_buf,in_size) == 0 );
+
+    //-------------------------------------
+    // free buffers :
+
+    OodleXFree(decoderMem);
+    OodleXFree(comp_buf);
+    OodleXFree(dec_out_buf);
+    OodleXFree(dec_window);
+}
+
+
+/*
+
+lz_test_12 :
+
+example of directly calling the simple buffer->buffer compression API's using an "in place" buffer
+
+OodleLZ_Compress
+OodleLZ_Decompress
+OodleLZ_GetInPlaceDecodeBufferSize
+
+*/
+
+void lz_test_12()
+{
+    OodleXLog_Printf_v0("lz_test_12\n");
+    // allocate compressed buffer & decoded buffer of the correct sizes :
+
+    OO_SINTa comp_buf_size = OodleLZ_GetCompressedBufferSizeNeeded(OodleLZ_Compressor_Kraken,in_size);
+    void * comp_buf = OodleXMalloc( comp_buf_size );
+    OOEX_ASSERT( comp_buf != NULL );
+    
+    //---------------------------------------------------
+    
+    // compress buffer -> buffer :
+    
+    OO_SINTa comp_len;
+    
+    {
+    comp_len = OodleLZ_Compress(OodleLZ_Compressor_Kraken,in_buffer,in_size,comp_buf,OodleLZ_CompressionLevel_Fast);
+    }
+    
+    //---------------------------------------------------
+
+    OO_SINTa inplace_size = OodleLZ_GetInPlaceDecodeBufferSize(OodleLZ_Compressor_Kraken,comp_len,in_size);
+
+    OodleXLog_Printf_v1("Kraken compress %d -> %d ; inplace_size = %d , padding = %d\n",
+        (int)in_size,(int)comp_len,
+        (int)inplace_size,(int)(inplace_size - in_size));
+
+    void * inplace_buf = OodleXMalloc( inplace_size );
+    OOEX_ASSERT( inplace_buf != NULL );
+
+    // in game use, you load the compressed data into the *end* of the inplace buffer 
+    // simulate the loading by doing a memcpy :
+    
+    char * inplace_comp_ptr = (char *)inplace_buf;
+    inplace_comp_ptr += inplace_size - comp_len;
+    
+    memcpy(inplace_comp_ptr,comp_buf,comp_len);
+
+    //---------------------------------------------------
+        
+    // decompress :
+    // note the source (inplace_comp_ptr) and dest (inplace_buf) overlap
+    //  - the compressed data at inplace_comp_ptr is destroyed by this call
+
+    OO_SINTa dec_len = OodleLZ_Decompress(inplace_comp_ptr,comp_len,inplace_buf,in_size,OodleLZ_FuzzSafe_Yes);
+
+    OOEX_ASSERT_ALWAYS( dec_len == in_size );
+    OOEX_ASSERT_ALWAYS( memcmp(in_buffer,inplace_buf,in_size) == 0 );
+
+    //---------------------------------------------------
+    
+    // free buffers :
+
+    OodleXFree(comp_buf);
+    OodleXFree(inplace_buf);
+}
+
+
+//=================================================
+
+/*
+
+lz_test_13 :
+
+example of dictionary-relative compression
+
+This is a technique in which you train a dictionary offline based on typical data, then for
+each packet you wish to compress, the dictionary can be used a reference to improve compression ratio.
+
+Oodle can do dictionary relative compression by putting the packet buffer to compress
+in a contiguous buffer immediately following the dictionary.
+
+Then simply use memcpy to move the active packet to the desired memory location.
+
+NOTE that the work space for {dictionary + packet} must be allocated per thread, or mutex controlled
+(it cannot be shared by simultaneously decoding threads)
+
+For small packets (under 4 KB or so) such as network packets, consider Oodle Network instead.
+
+For large buffers (over 128 KB or so), dictionary-relative compression doesn't help much and isn't recommended.
+
+dictionary-relative compression is most typically useful on data in the 4 - 128 KB range.
+
+*/
+
+void lz_test_13()
+{
+    OodleXLog_Printf_v0("lz_test_13\n");
+    
+    //---------------------------------------------------
+    
+    // pretend that "in_buffer" consists of a trained dictionary + a packet to compress
+    
+    void * dictionary = in_buffer;
+    OO_SINTa dictionary_size = (in_size * 2) /3;
+    // dictionary_size must be a multiple of OODLELZ_BLOCK_LEN :
+    dictionary_size &= ~(OODLELZ_BLOCK_LEN-1);
+    
+    void * packet1 = (char *)dictionary + dictionary_size;
+    OO_SINTa packet1_size = in_size / 4;
+    
+    void * packet2 = (char *)packet1 + packet1_size;
+    OO_SINTa packet2_size = in_size - packet1_size - dictionary_size;
+    
+    OodleXLog_Printf_v1("dictionary_size : %d ; packets : %d + %d\n",dictionary_size,packet1_size,packet2_size);
+    
+    //---------------------------------------------------
+    // allocate compressed buffer & decoded buffer of the correct sizes :
+
+    OO_SINTa max_packet_size = OOEX_MAX(packet1_size,packet2_size);
+
+    // room for dictionary + a packet following :
+    void * dictionary_and_packet_buf = OodleXMalloc( dictionary_size + max_packet_size );
+    OOEX_ASSERT( dictionary_and_packet_buf != NULL );
+
+    // comp buf just for a packet :
+    OO_SINTa comp_buf_size = OodleLZ_GetCompressedBufferSizeNeeded(OodleLZ_Compressor_Kraken, max_packet_size );
+    void * comp_buf = OodleXMalloc( comp_buf_size );
+    OOEX_ASSERT( comp_buf != NULL );
+
+    // room for dictionary + a packet following :
+    void * dec_buf = OodleXMalloc( dictionary_size + max_packet_size );
+    OOEX_ASSERT( dec_buf != NULL );
+    
+    //---------------------------------------------------
+    
+    void * packets[2] = { packet1, packet2 };
+    OO_SINTa packet_sizes[2] = { packet1_size, packet2_size };
+
+    // setup work that's done in advance :
+
+    // put dictionary at head of dictionary_and_packet_buf (for encoder) :
+    memcpy(dictionary_and_packet_buf,dictionary,dictionary_size);
+    void * after_dictionary_ptr = (char *)dictionary_and_packet_buf + dictionary_size;
+    
+    // preload dictionary at head of dec_buf (for decoder) :
+    memcpy(dec_buf,dictionary,dictionary_size);
+    
+    for(int packet_i=0;packet_i<2;packet_i++)
+    {
+        // work that's done per packet :
+    
+        void * packet_ptr = packets[packet_i];
+        OO_SINTa packet_size = packet_sizes[packet_i];
+    
+        // compress packet to comp_buf
+        // preload with dictionary
+        
+        // copy packet to be immediately following dictionary :
+        memcpy(after_dictionary_ptr,packet_ptr,packet_size);
+        
+        OO_SINTa comp_len = OodleLZ_Compress(OodleLZ_Compressor_Kraken,after_dictionary_ptr,packet_size,comp_buf,OodleLZ_CompressionLevel_Fast,NULL,dictionary_and_packet_buf);
+        
+        OodleXLog_Printf_v1("Kraken compress %d -> %d\n",(int)packet_size,(int)comp_len);
+
+        // decompress :
+        // decode into buffer containing dictionary, immediately following dictionary :
+
+        void * dec_packet_ptr = (char *)dec_buf + dictionary_size;
+
+        OO_SINTa dec_len = OodleLZ_Decompress(comp_buf,comp_len,dec_packet_ptr,packet_size,OodleLZ_FuzzSafe_Yes,OodleLZ_CheckCRC_No,OodleLZ_Verbosity_None,dec_buf);
+
+        OOEX_ASSERT_ALWAYS( dec_len == packet_size + dictionary_size );
+        OOEX_ASSERT_ALWAYS( memcmp(packet_ptr,dec_packet_ptr,packet_size) == 0 );
+        
+        // if you need the decoded packet to be in another memory location, memcpy it there now
+    }
+    
+    //---------------------------------------------------
+    // free buffers :
+
+    OodleXFree(dictionary_and_packet_buf);
+    OodleXFree(comp_buf);
+    OodleXFree(dec_buf);
+}
diff --git a/Engine/Source/Programs/Horde/HordeStorage/Jupiter.Common/Shared/Oodle/2.9.3/linux/examples/example_lz_chart.cpp b/Engine/Source/Programs/Horde/HordeStorage/Jupiter.Common/Shared/Oodle/2.9.3/linux/examples/example_lz_chart.cpp
new file mode 100644
--- /dev/null
+++ b/Engine/Source/Programs/Horde/HordeStorage/Jupiter.Common/Shared/Oodle/2.9.3/linux/examples/example_lz_chart.cpp
@@ -0,0 +1,483 @@
+//idoc(parent,examples)
+//idoc(tutorial,example_lz_chart,example_lz_chart : Example that makes a chart of OodleLZ options )
+
+/*
+
+Oodle example_lz_chart
+
+The Oodle SDK comes with a pre-built exe for example_lz_chart in the bin/ directory
+
+usage :
+example_lz_chart <file>
+
+Run with a file name, which will be loaded and used as data to test on.
+
+You can also toggle compile-time options with the define EXAMPLE_LZ_CHART_NUM_LEVELS below.
+
+makes an output like this :
+
+<PRE>
+Oodle 2.6.3 example_lz_chart <file>
+lz_chart loading r:\testsets\lztestset\lzt99...
+file size : 24700820
+------------------------------------------------------------------------------
+Selkie : super fast to encode & decode, least compression
+Mermaid: fast decode with better-than-zlib compression
+Kraken : good compression, fast decoding, great tradeoff!
+Leviathan : very high compression, slowest decode
+------------------------------------------------------------------------------
+chart cell shows | raw/comp ratio : encode MB/s : decode MB/s |
+All compressors run at various encoder effort levels (SuperFast - Optimal).
+Many repetitions are run for accurate timing.
+------------------------------------------------------------------------------
+       |   HyperFast4|   HyperFast3|   HyperFast2|   HyperFast1|   SuperFast |
+Selkie |1.41:675:3895|1.45:622:3888|1.53:465:3696|1.68:369:3785|1.70:342:3759|
+Mermaid|1.66:436:2189|1.66:436:2188|1.79:352:2090|2.01:276:2055|2.04:261:2025|
+Kraken |1.55:588:1839|1.71:419:1136|1.88:331:1087|2.10:279:1093|2.27:167:1010|
+------------------------------------------------------------------------------
+compression ratio (raw/comp):
+       |   HyperFast4|   HyperFast3|   HyperFast2|   HyperFast1|   SuperFast |
+Selkie |    1.412    |    1.447    |    1.526    |    1.678    |    1.698    |
+Mermaid|    1.660    |    1.660    |    1.793    |    2.011    |    2.041    |
+Kraken |    1.548    |    1.711    |    1.877    |    2.103    |    2.268    |
+------------------------------------------------------------------------------
+encode speed (MB/s):
+       |   HyperFast4|   HyperFast3|   HyperFast2|   HyperFast1|   SuperFast |
+Selkie |    674.548  |    621.811  |    464.555  |    369.364  |    341.588  |
+Mermaid|    435.650  |    435.923  |    352.475  |    276.199  |    260.511  |
+Kraken |    588.488  |    418.921  |    331.423  |    279.129  |    167.206  |
+------------------------------------------------------------------------------
+decode speed (MB/s):
+       |   HyperFast4|   HyperFast3|   HyperFast2|   HyperFast1|   SuperFast |
+Selkie |   3894.644  |   3887.820  |   3695.984  |   3785.457  |   3758.594  |
+Mermaid|   2189.030  |   2187.863  |   2090.319  |   2054.897  |   2024.692  |
+Kraken |   1839.091  |   1135.920  |   1086.922  |   1093.407  |   1009.967  |
+------------------------------------------------------------------------------
+       |   VeryFast  |   Fast      |   Normal    |   Optimal1  |   Optimal3  |
+Selkie |1.75:205:3490|1.83:105:3687|1.86: 43:3815|1.93:5.1:3858|1.94:2.6:3856|
+Mermaid|2.12:173:1991|2.19: 84:2177|2.21: 32:2291|2.37:2.8:2058|2.44:1.8:1978|
+Kraken |2.32:112:1104|2.39: 37:1187|2.43: 20:1189|2.55:3.1:1103|2.65:1.2:1038|
+Leviath|2.50: 31: 738|2.57: 17: 787|2.62:9.5: 807|2.71:1.6: 811|2.76:0.9: 776|
+------------------------------------------------------------------------------
+compression ratio (raw/comp):
+       |   VeryFast  |   Fast      |   Normal    |   Optimal1  |   Optimal3  |
+Selkie |    1.748    |    1.833    |    1.863    |    1.933    |    1.943    |
+Mermaid|    2.118    |    2.194    |    2.207    |    2.367    |    2.437    |
+Kraken |    2.320    |    2.390    |    2.434    |    2.551    |    2.646    |
+Leviath|    2.504    |    2.572    |    2.617    |    2.707    |    2.756    |
+------------------------------------------------------------------------------
+encode speed (MB/s):
+       |   VeryFast  |   Fast      |   Normal    |   Optimal1  |   Optimal3  |
+Selkie |    204.621  |    104.758  |     42.504  |      5.102  |      2.554  |
+Mermaid|    172.681  |     84.227  |     32.030  |      2.798  |      1.836  |
+Kraken |    111.858  |     37.126  |     19.859  |      3.091  |      1.204  |
+Leviath|     31.031  |     16.697  |      9.461  |      1.621  |      0.869  |
+------------------------------------------------------------------------------
+decode speed (MB/s):
+       |   VeryFast  |   Fast      |   Normal    |   Optimal1  |   Optimal3  |
+Selkie |   3490.442  |   3686.689  |   3814.655  |   3857.857  |   3856.226  |
+Mermaid|   1991.442  |   2176.725  |   2291.498  |   2057.575  |   1977.721  |
+Kraken |   1104.172  |   1186.638  |   1189.372  |   1103.148  |   1038.352  |
+Leviath|    737.934  |    787.152  |    806.523  |    811.161  |    775.800  |
+------------------------------------------------------------------------------
+</PRE>
+
+*/
+
+#include "../include/oodle2x.h"
+#include "ooex.h"
+
+#ifndef _CRT_SECURE_NO_WARNINGS
+#define _CRT_SECURE_NO_WARNINGS
+#endif
+#include <stdio.h>
+#include <string.h>
+#include <assert.h>
+#include <stdlib.h>
+
+#ifdef BUILDING_EXAMPLE_CALLER
+#define main example_lz_chart
+#endif
+
+//===========================================================
+// NOTE : these timings are hot cache (no cache invalidation between repeats)
+//  that could be significant on very small buffers
+
+struct time_and_len
+{
+    double time;
+    OO_SINTa len;
+};
+
+// if scratch is NULL or insufficient, Oodle will allocate internally
+static void * s_scratch_memory = NULL;
+static OO_SINTa s_scratch_memory_size = 0;
+
+static time_and_len Encode_And_Time(const void * raw_buf, OO_SINTa raw_len, void * comp_buf,
+    OodleLZ_Compressor compressor,
+    OodleLZ_CompressionLevel compression_level,
+    OodleLZ_CompressOptions * compression_options,
+    int min_repeats,
+    double min_total_seconds)
+{
+    double total_seconds = 0;
+    time_and_len ret;
+    ret.time = 999999.9;
+
+    for(;;)
+    {
+        double dt = OodleX_GetSeconds();
+        
+        ret.len = OodleLZ_Compress(compressor, raw_buf, raw_len, comp_buf, compression_level, compression_options,
+            NULL, NULL, s_scratch_memory, s_scratch_memory_size);
+
+        dt = OodleX_GetSeconds() - dt;
+        
+        total_seconds += dt;
+        ret.time = OOEX_MIN(ret.time,dt);
+        min_repeats--;
+        
+        if ( min_repeats <= 0 && total_seconds >= min_total_seconds )
+            break;
+    }
+    
+    return ret;
+}
+
+static double Decode_And_Time( void * comp_buf, OO_SINTa comp_len, void * decode_buffer, OO_SINTa in_size,
+    int min_repeats,
+    double min_total_seconds)
+{
+    double total_seconds = 0;
+    double ret = 999999.9;
+
+    for(;;)
+    {
+        double dt = OodleX_GetSeconds();
+                
+        OO_SINTa decode_len = OodleLZ_Decompress( comp_buf, comp_len, decode_buffer, in_size , OodleLZ_FuzzSafe_Yes,
+            OodleLZ_CheckCRC_No, OodleLZ_Verbosity_None, NULL, 0, NULL,NULL, s_scratch_memory,s_scratch_memory_size );
+
+        dt = OodleX_GetSeconds() - dt;
+        
+        OOEX_ASSERT_ALWAYS( decode_len == in_size );
+        
+        total_seconds += dt;
+        ret = OOEX_MIN(ret,dt);
+        min_repeats--;
+        
+        if ( min_repeats <= 0 && total_seconds >= min_total_seconds )
+            break;
+    }
+    
+    return ret;
+}
+        
+//===========================================================
+
+static void bar()
+{
+    OodleXLog_Printf_v1("------------------------------------------------------------------------------\n");
+}
+
+static void header(const OodleLZ_CompressionLevel * chart_levels, int num_levels)
+{
+    //OodleXLog_Printf_v1("%5s |   %-10s|   %-10s|   %-10s|   %-10s|\n",
+    //  "","VeryFast","Fast","Normal","Optimal1");
+        
+    OodleXLog_Printf_v1("%7s|","");
+    
+    for(int l=0;l<num_levels;l++)
+    {           
+        OodleLZ_CompressionLevel level = chart_levels[l];
+        
+        OodleXLog_Printf_v1("   %-10s|",
+           OodleLZ_CompressionLevel_GetName(level));
+    }
+    OodleXLog_Printf_v1("\n");
+}
+
+static const char * truncated_compressor_name(  OodleLZ_Compressor compressor )
+{
+    const char * n = OodleLZ_Compressor_GetName(compressor);
+    static char buf[10];
+    memset(buf,' ',sizeof(buf));
+    memcpy(buf,n,OOEX_MIN( strlen(n), sizeof(buf) ) );
+    buf[7] = 0;
+    return buf;
+}   
+        
+extern "C" int main(int argc,char *argv[])
+{
+    // Init Oodle systems with default options :
+    OodleXInitOptions opts;
+    if ( ! OodleX_Init_GetDefaults(OODLE_HEADER_VERSION,&opts) )
+    {
+        fprintf(stderr,"Oodle header version mismatch.\n");
+        return 10;
+    }
+    // change opts here if you like
+    // NOTE : default options enable the OodleX thread system
+    //  so encoders will be "Jobified"
+    if ( ! OodleX_Init(OODLE_HEADER_VERSION,&opts) )
+    {
+        fprintf(stderr,"OodleX_Init failed.\n");
+        return 10;
+    }
+
+    OodleXLog_Printf_v1("Oodle %s example_lz_chart <file>\n",OodleVersion);
+    
+    if ( argc < 2 )
+    {
+        fprintf(stderr,"error: specify a sample data file to test on.\n");
+        return 10;
+    }
+    
+    const char * in_name = argv[1];
+    OodleXLog_Printf_v1("lz_chart loading %s...\n",in_name);
+
+    // read the input file to the global buffer :
+    OO_S64 in_size_64;
+    void * in_buffer = OodleXIOQ_ReadMallocWholeFile_AsyncAndWait(in_name,&in_size_64);
+    
+    if ( ! in_buffer)
+    {
+        OodleXLog_Printf_v0("failed to read %s\n",in_name); 
+        return 10;
+    }
+    
+    OodleXLog_Printf_v1("file size : " OOEX_S64_FMT "\n",in_size_64);
+    
+    OO_SINTa in_size = OodleX_S64_to_SINTa_check( in_size_64 );
+        
+    //-----------------------------------------------------
+    
+    // test parameters :
+    //  increase these to get more reliable timing
+    //  decrease these to run faster
+    int min_encode_repeats = 2;
+    int min_decode_repeats = 5;
+    double min_total_seconds = 2.0;
+    
+    #if 1
+    
+    // test set of compressors : 
+    OodleLZ_Compressor chart_compressors[] = {
+        OodleLZ_Compressor_Selkie, OodleLZ_Compressor_Mermaid, OodleLZ_Compressor_Kraken, OodleLZ_Compressor_Leviathan
+    };
+        
+    #define EXAMPLE_LZ_CHART_NUM_COMPRESSORS    ((int)(sizeof(chart_compressors)/sizeof(chart_compressors[0])))
+    OOEX_ASSERT( EXAMPLE_LZ_CHART_NUM_COMPRESSORS == 4 );
+    
+    #define EXAMPLE_LZ_CHART_NUM_LEVELS 5
+    
+    OodleLZ_CompressionLevel chart_levels1[] = {
+        OodleLZ_CompressionLevel_HyperFast4, OodleLZ_CompressionLevel_HyperFast3, OodleLZ_CompressionLevel_HyperFast2, OodleLZ_CompressionLevel_HyperFast1, OodleLZ_CompressionLevel_SuperFast,
+    };
+    OodleLZ_CompressionLevel chart_levels2[] = {
+        OodleLZ_CompressionLevel_VeryFast, OodleLZ_CompressionLevel_Fast, OodleLZ_CompressionLevel_Normal, OodleLZ_CompressionLevel_Optimal1, OodleLZ_CompressionLevel_Optimal3
+    };
+    OOEX_ASSERT( ((int)(sizeof(chart_levels1)/sizeof(chart_levels1[0]))) == EXAMPLE_LZ_CHART_NUM_LEVELS );
+    OOEX_ASSERT( ((int)(sizeof(chart_levels2)/sizeof(chart_levels2[0]))) == EXAMPLE_LZ_CHART_NUM_LEVELS );
+    
+    #else
+    
+    // test just 1 :
+    
+    OodleLZ_Compressor chart_compressors[] = {
+        OodleLZ_Compressor_Kraken
+    };
+        
+    #define EXAMPLE_LZ_CHART_NUM_COMPRESSORS    1
+    
+    #define EXAMPLE_LZ_CHART_NUM_LEVELS 1
+    
+    OodleLZ_CompressionLevel chart_levels1[] = {
+        OodleLZ_CompressionLevel_None
+    };
+    OodleLZ_CompressionLevel chart_levels2[] = {
+        OodleLZ_CompressionLevel_Optimal2
+    };
+    
+    #endif
+    
+    //===========================================================
+
+    OO_SINTa comp_lens[EXAMPLE_LZ_CHART_NUM_COMPRESSORS][EXAMPLE_LZ_CHART_NUM_LEVELS];
+    double decode_speeds[EXAMPLE_LZ_CHART_NUM_COMPRESSORS][EXAMPLE_LZ_CHART_NUM_LEVELS];
+    double encode_speeds[EXAMPLE_LZ_CHART_NUM_COMPRESSORS][EXAMPLE_LZ_CHART_NUM_LEVELS];
+    
+    //-----------------------------------------------------
+    
+    void * comp_buf;
+
+    OO_SINTa comp_buf_size = OodleLZ_GetCompressedBufferSizeNeeded(OodleLZ_Compressor_Invalid,in_size);
+
+    comp_buf = OodleXMallocBig(comp_buf_size);
+        
+    void * decode_buffer = OodleXMallocBig( in_size );
+        
+    //-----------------------------------------------------
+    // allocate scratch 
+
+    // get enough so the decoder won't allocate
+    //  and encoders up to Leviathan-Normal won't either
+    //  but optimals will
+        
+    s_scratch_memory_size = OodleLZ_GetCompressScratchMemBound(OodleLZ_Compressor_Leviathan,OodleLZ_CompressionLevel_Normal,in_size,NULL);
+    
+    s_scratch_memory_size = OOEX_MAX( s_scratch_memory_size , OodleLZDecoder_MemorySizeNeeded() );
+    
+    s_scratch_memory = OodleXMalloc(s_scratch_memory_size);
+
+    //-----------------------------------------------------
+    
+    bar();
+    OodleXLog_Printf_v1("Selkie : super fast to encode & decode, least compression\n");
+    OodleXLog_Printf_v1("Mermaid: fast decode with better-than-zlib compression\n");
+    OodleXLog_Printf_v1("Kraken : good compression, fast decoding, great tradeoff!\n");
+    OodleXLog_Printf_v1("Leviathan : very high compression, slowest decode\n");
+    bar();
+    OodleXLog_Printf_v1("chart cell shows | raw/comp ratio : encode MB/s : decode MB/s | \n");
+    OodleXLog_Printf_v1("All compressors run at various encoder effort levels (SuperFast - Optimal).\n");
+    OodleXLog_Printf_v1("Many repetitions are run for accurate timing.\n");
+    bar();
+    
+    for(int twice=0;twice<2;twice++)
+    {
+        const OodleLZ_CompressionLevel * chart_levels =
+            twice ? chart_levels2 : chart_levels1;
+    
+        int num_compressors = EXAMPLE_LZ_CHART_NUM_COMPRESSORS;
+        // don't bother running Leviathan at the HyperFast modes :
+        if ( ! twice && chart_compressors[num_compressors-1] == OodleLZ_Compressor_Leviathan ) num_compressors--;
+    
+        header(chart_levels,EXAMPLE_LZ_CHART_NUM_LEVELS);
+        
+        for(int c=0;c<num_compressors;c++)
+        {
+            OodleLZ_Compressor compressor = chart_compressors[c];
+            
+            OodleXLog_Printf_v1("%s|",truncated_compressor_name(compressor));
+            
+            // iterate backwards so we start the slowest first
+            for(int l=0;l<EXAMPLE_LZ_CHART_NUM_LEVELS;l++)
+            {           
+                OodleLZ_CompressionLevel level = chart_levels[l];
+            
+                time_and_len tl = 
+                    Encode_And_Time(in_buffer,in_size,comp_buf,compressor,level,NULL,min_encode_repeats,min_total_seconds);
+                
+                OO_SINTa comp_len = tl.len;
+                double encode_seconds = tl.time;
+                
+                double encode_mbps = (in_size/1000000.0) / encode_seconds;      
+                encode_speeds[c][l] = encode_mbps;
+                
+                comp_lens[c][l] = comp_len;
+                
+                double ratio = (double)in_size/comp_len;
+                
+                if ( ratio >= 10.0 )
+                    OodleXLog_Printf_v1("%4.1f",ratio);
+                else
+                    OodleXLog_Printf_v1("%4.2f",ratio);
+                
+                if ( encode_mbps >= 10.0 )
+                    OodleXLog_Printf_v1(":%3d:",(int)(encode_mbps+0.5));
+                else
+                    OodleXLog_Printf_v1(":%3.1f:",encode_mbps);
+                    
+                double decode_seconds = Decode_And_Time( comp_buf, comp_len, decode_buffer, in_size , 
+                    min_decode_repeats,min_total_seconds);
+        
+                OOEX_ASSERT_ALWAYS( memcmp(decode_buffer,in_buffer,in_size) == 0 );
+                        
+                double decode_mbps = (in_size/1000000.0) / decode_seconds;              
+                decode_speeds[c][l] = decode_mbps;
+                
+                OodleXLog_Printf_v1("%4d|",(int)(decode_mbps+0.5));
+            }
+            OodleXLog_Printf_v1("\n");
+            
+        }
+
+        //-----------------------------------------------------
+
+        bar();  
+        OodleXLog_Printf_v1("compression ratio (raw/comp):\n");
+        header(chart_levels,EXAMPLE_LZ_CHART_NUM_LEVELS);
+        
+        for(int c=0;c<num_compressors;c++)
+        {
+            OodleLZ_Compressor compressor = chart_compressors[c];
+            
+            OodleXLog_Printf_v1("%s|",truncated_compressor_name(compressor));
+            
+            // iterate backwards so we start the slowest first
+            for(int l=0;l<EXAMPLE_LZ_CHART_NUM_LEVELS;l++)
+            {           
+                double ratio = (double)in_size/comp_lens[c][l];
+                
+                OodleXLog_Printf_v1("%9.3f    |",ratio);
+            }
+            OodleXLog_Printf_v1("\n");
+        }
+
+        bar();
+        OodleXLog_Printf_v1("encode speed (MB/s):\n");
+        header(chart_levels,EXAMPLE_LZ_CHART_NUM_LEVELS);
+        
+        for(int c=0;c<num_compressors;c++)
+        {
+            OodleLZ_Compressor compressor = chart_compressors[c];
+            
+            OodleXLog_Printf_v1("%s|",truncated_compressor_name(compressor));
+            
+            // iterate backwards so we start the slowest first
+            for(int l=0;l<EXAMPLE_LZ_CHART_NUM_LEVELS;l++)
+            {           
+                OodleXLog_Printf_v1("%11.3f  |",encode_speeds[c][l]);
+            }
+            OodleXLog_Printf_v1("\n");
+        }
+
+        bar();
+        OodleXLog_Printf_v1("decode speed (MB/s):\n");
+        header(chart_levels,EXAMPLE_LZ_CHART_NUM_LEVELS);
+        
+        for(int c=0;c<num_compressors;c++)
+        {
+            OodleLZ_Compressor compressor = chart_compressors[c];
+            
+            OodleXLog_Printf_v1("%s|",truncated_compressor_name(compressor));
+            
+            // iterate backwards so we start the slowest first
+            for(int l=0;l<EXAMPLE_LZ_CHART_NUM_LEVELS;l++)
+            {           
+                OodleXLog_Printf_v1("%11.3f  |",decode_speeds[c][l]);
+            }
+            OodleXLog_Printf_v1("\n");
+        }
+
+        bar();
+
+    } // twice
+
+    //-----------------------------------------------------
+
+    OodleXFree(s_scratch_memory);
+    OodleXFreeBig(decode_buffer);
+    OodleXFreeBig(comp_buf);        
+    OodleXFree_IOAligned(in_buffer);
+
+    //OodleX_Shutdown();
+    OodleX_Shutdown(NULL,OodleX_Shutdown_LogLeaks_Yes,0);
+    
+    //OodleXLog_Printf_v1("press a key\n");
+    //fgetc(stdin);
+    
+    return 0;
+}
+
+//=================================================
diff --git a/Engine/Source/Programs/Horde/HordeStorage/Jupiter.Common/Shared/Oodle/2.9.3/linux/examples/example_lz_noallocs.cpp b/Engine/Source/Programs/Horde/HordeStorage/Jupiter.Common/Shared/Oodle/2.9.3/linux/examples/example_lz_noallocs.cpp
new file mode 100644
--- /dev/null
+++ b/Engine/Source/Programs/Horde/HordeStorage/Jupiter.Common/Shared/Oodle/2.9.3/linux/examples/example_lz_noallocs.cpp
@@ -0,0 +1,218 @@
+//idoc(parent,examples)
+//idoc(tutorial,example_lz_noallocs,example_lz_noallocs : Example demonstrating Oodle compression with no allocations)
+
+/*
+
+Oodle example_lz_noallocs
+
+Very simple example of OodleLZ memory -> memory compression & decompression.
+
+Shows how to use Oodle without any allocations done by Oodle.  All memory needed is passed in by the client.
+
+Uses stdio for file IO to load an input file.
+
+See $example_lz for more advanced OodleLZ usage.
+
+example_lz_noallocs only uses Oodle Core, no Oodle Ext
+
+include oodle2.h and not oodle2x.h
+
+If you want to use Oodle with no allocations, you cannot use OodleX.  OodleX installs its own allocator into Oodle Core.
+Do not use OodlePlugins_SetAllocators with OodleX.
+
+See $Oodle_FAQ_UseOodleWithNoAllocator and $OodleCore_Plugins_SetAllocators
+
+*/
+
+#include "../include/oodle2.h"
+
+#ifndef _CRT_SECURE_NO_WARNINGS
+#define _CRT_SECURE_NO_WARNINGS
+#endif
+#include <stdio.h>
+#include <string.h>
+#include <assert.h>
+#include <stdlib.h>
+
+#ifdef BUILDING_EXAMPLE_CALLER
+#define main example_lz_noallocs
+#endif
+
+#include "read_whole_file.h"
+
+
+void * OODLE_CALLBACK example_noallocs_MallocAligned_Logging(OO_SINTa size,OO_S32 alignment)
+{
+    // malloc should not be called, log an error :
+    printf("ERROR : example_noallocs_MallocAligned_Logging called (size %d)\n",(int)size);
+
+    // use default Oodle MallocAligned as an easy way to get a cross-platform aligned malloc :
+    return OodleCore_Plugin_MallocAligned_Default(size,alignment);
+}
+
+void OODLE_CALLBACK example_noallocs_Free_Logging(void * ptr)
+{
+    // malloc should not be called, log an error :
+    printf("ERROR : example_noallocs_Free_Logging called.\n");
+    
+    OodleCore_Plugin_Free_Default(ptr);
+}
+
+
+extern "C" int main(int argc,char *argv[])
+{
+/* No initialization is needed for Oodle2 Core
+
+    we let Oodle Core use the default system plugins (the C stdlib)
+    To change them, use $OodleAPI_OodleCore_Plugins
+*/
+
+    // optional check to make sure header matches lib :
+    if ( ! Oodle_CheckVersion(OODLE_HEADER_VERSION) )
+    {
+        fprintf(stderr,"Oodle header version mismatch\n");
+        return 10;
+    }
+
+/* Install our own allocator plugins that log an error if called.
+
+    These should never be called.  You could also disable them :
+    OodleCore_Plugins_SetAllocators(NULL,NULL);
+    
+    but that is not recommended because it will cause a hard failure if Oodle ever needs memory.
+*/
+    OodleCore_Plugins_SetAllocators(example_noallocs_MallocAligned_Logging,example_noallocs_Free_Logging);
+
+    // get args :
+    const char * in_name;
+    if ( argc < 2 )
+    {
+        in_name = "r:\\testsets\\lztestset\\lzt02";
+    }
+    else
+    {
+        in_name = argv[1];
+    }
+    
+/* read input file using stdio
+
+*/
+    OO_SINTa length;
+    void * buf = read_whole_file(in_name,&length);
+    if ( ! buf )
+    {
+        fprintf(stderr,"couldn't open : %s\n",in_name);
+        return 10;
+    }
+    
+
+/*
+
+Run OodleLZ_Compress from memory (buf) to memory (compbuf)
+
+Use the OodleLZ_Compressor_Kraken compressor.  Kraken is an amazing balance of good compression and fast decode
+speed.  It should generally be your first choice.
+
+Use OodleLZ_CompressionLevel_Normal level of effort in the encoder.  Normal is a balance of encode speed and compression
+ratio.  Different levels trade off faster or slower encoding for compressed size.  See $OodleLZ_CompressionLevel.
+
+See $OodleLZ_About for information on selection of the compression options.
+
+This call is synchronous and not threaded; see $example_lz for an example using the async compression APIs.
+
+*/
+
+    OodleLZ_Compressor compressor = OodleLZ_Compressor_Kraken;
+    OodleLZ_CompressionLevel level = OodleLZ_CompressionLevel_Normal;
+    //OodleLZ_CompressionLevel level = OodleLZ_CompressionLevel_Optimal2; // optimals are OODLELZ_SCRATCH_MEM_NO_BOUND
+
+    // allocate memory big enough for compressed data :
+    void * compbuf = malloc( OodleLZ_GetCompressedBufferSizeNeeded(compressor,length) + sizeof(length) );
+    if ( compbuf == NULL )
+        return 10;
+
+    // allocate memory for encoder scratch :
+    OO_SINTa enc_scratch_size = OodleLZ_GetCompressScratchMemBound(compressor,level,length,NULL);
+    if ( enc_scratch_size == OODLELZ_SCRATCH_MEM_NO_BOUND )
+    {
+        // scratch cannot be bounded for this choice of compressor/level
+        // the allocator may be used!
+        // go ahead and give it 4 MB of scratch
+        enc_scratch_size = 4*1024*1024;
+    }
+    
+    void * enc_scratch = malloc(enc_scratch_size);
+    if ( enc_scratch == NULL )
+        return 10;  
+        
+    char * compptr = (char *)compbuf;
+    memcpy(compptr,&length,sizeof(length));
+    compptr += sizeof(length);
+
+    // compress :
+    OO_SINTa complen = OodleLZ_Compress(compressor,buf,length,compptr,level,NULL,NULL,NULL,enc_scratch,enc_scratch_size);
+    compptr += complen;
+
+    // log about it :
+    //  full compressed size also includes the header
+    printf("%s compressed %d -> %d (+%d)\n",in_name,(int)length,(int)complen,(int)sizeof(length));
+
+    // can free enc_scratch now
+    // enc_scratch can be reused for further compression
+    //  but must be used by only one thread at a time
+    free(enc_scratch); enc_scratch = NULL;
+
+/*
+
+Run OodleLZ_Decompress from memory (compbuf) to memory (decbuf)
+
+Note that you must provide the exact decompressed size.  OodleLZ data is headerless; store the size in
+your own header.
+
+*/
+
+    OO_SINTa declength;
+    compptr = (char *)compbuf;
+    memcpy(&declength,compptr,sizeof(declength));
+    compptr += sizeof(declength);
+    assert( length == declength );
+
+    // malloc for decompressed buffer  :
+    void * decbuf = malloc( declength );
+    if ( decbuf == NULL )
+        return 10;      
+
+    OO_SINTa decoder_mem_size = OodleLZDecoder_MemorySizeNeeded(compressor,declength);
+    void * decoder_mem = malloc(decoder_mem_size);
+    if ( decoder_mem == NULL )
+        return 10;      
+
+    // do the decompress :
+    OO_SINTa decompress_return = OodleLZ_Decompress(compptr,complen,decbuf,declength,
+            OodleLZ_FuzzSafe_Yes,OodleLZ_CheckCRC_Yes,OodleLZ_Verbosity_None,NULL,0,NULL,NULL,
+            decoder_mem,decoder_mem_size,
+            OodleLZ_Decode_Unthreaded);
+    
+    // check it was successful :
+    assert( decompress_return == length );
+    assert( memcmp(buf,decbuf,length) == 0 );
+
+    if ( decompress_return != length )
+        return 10;
+
+    printf("decompessed successfully.\n");
+        
+/*
+
+And finish up.  No shutdown is needed for Oodle2 Core.
+
+*/
+
+    // free all the memory :
+    free(buf);
+    free(compbuf);
+    free(decbuf);
+    free(decoder_mem);
+
+    return 0;
+}
diff --git a/Engine/Source/Programs/Horde/HordeStorage/Jupiter.Common/Shared/Oodle/2.9.3/linux/examples/example_lz_overlap.cpp b/Engine/Source/Programs/Horde/HordeStorage/Jupiter.Common/Shared/Oodle/2.9.3/linux/examples/example_lz_overlap.cpp
new file mode 100644
--- /dev/null
+++ b/Engine/Source/Programs/Horde/HordeStorage/Jupiter.Common/Shared/Oodle/2.9.3/linux/examples/example_lz_overlap.cpp
@@ -0,0 +1,409 @@
+//idoc(parent,examples)
+//idoc(tutorial,example_lz_overlap,example_lz_overlap : Example demonstrating parallel overlap with OodleLZ )
+
+/*
+
+Oodle example_lz_overlap
+
+Demonstration of the benefit of overlapping IO with CPU work and parallelism in LZ decompression.
+
+This example compresses a file, then repeatedly reads the compressed data and decompresses it, in several different ways.
+
+There are two types of parallelism demonstrated here :
+
+1. IO overlap.  When reading and decompressing large files, you get minimum latency by
+overlapping the IO with the decompress.  This is done by reading the compressed data in smaller
+chunks and decompressing each chunk (in parallel) as it is done.
+
+2. Parallel "wide" decompression.  OodleLZ with seek chunk resets can decompress using many
+threads simultaneously.
+
+Combining IO overlap and wide decompression is the fastest way to load compressed data.
+
+*/
+
+#include "../include/oodle2x.h"
+#include "ooex.h"
+
+#ifdef _MSC_VER
+#pragma warning(disable : 4127) // conditional is constant
+#endif
+
+#ifndef _CRT_SECURE_NO_WARNINGS
+#define _CRT_SECURE_NO_WARNINGS
+#endif
+#include <stdio.h>
+#include <string.h>
+#include <assert.h>
+#include <stdlib.h>
+
+#include "make_example_input.h"
+
+#ifdef BUILDING_EXAMPLE_CALLER
+#define main example_lz_overlap
+#endif
+
+extern "C" int main(int argc,char *argv[])
+{
+    if ( ! OodleX_Init_Default(OODLE_HEADER_VERSION) )
+    //if ( ! OodleX_Init_Default(OODLE_HEADER_VERSION,OodleX_Init_GetDefaults_DebugSystems_Yes,OodleX_Init_GetDefaults_Threads_No) )
+    {
+        fprintf(stderr,"OodleX_Init failed.\n");
+        return 10;
+    }
+
+/*
+
+    Pass in the file name to compress.  If you don't provide one, I'll make one.
+    
+    Ideally give this example something big to work on, 20M - 100M , to get better charts.
+
+*/
+
+    // get args :
+    const char * inName;
+    if ( argc < 2 )
+    {
+        inName = "oodle_example_input_file";
+        make_example_input(inName);
+    }
+    else
+    {
+        inName = argv[1];
+    }
+
+    // we'll write to this file :
+    const char * compName = "oodle_example_overlap_comp";
+     
+    OodleXLog_Printf_v1("compressing %s -> %s ...\n",inName,compName);
+
+/*
+    
+    Set up the LZ compression options so that we make independent chunks
+    (seekChunkReset = true) and make them 1 MB each.
+    
+    This hurts compression ratio (vs not chunking) but allows decompression of each chunk
+    to run in parallel.
+
+*/
+    
+    OodleLZ_CompressionLevel lzLevel = OodleLZ_CompressionLevel_Fast;
+    OodleLZ_Compressor lzCompressor = OodleLZ_Compressor_Kraken;
+    OodleLZ_CompressOptions lzOptions = * OodleLZ_CompressOptions_GetDefault(lzCompressor,lzLevel);
+
+    lzOptions.seekChunkLen = 1<<20; // 1 MB
+    lzOptions.seekChunkReset = true;
+
+/*
+
+    Read the input file, compress it, and write it out :
+
+    In practice for maximum speed, you could read each chunk of the input and compress them
+    independently as each chunk async IO read finishes.
+
+*/
+
+    OO_SINTa inSize;
+    void * inBuf;
+    OO_SINTa compSize;
+    OodleLZ_SeekTable * seekTable;
+    OO_S64 inSize64;
+    
+    inBuf = OodleXIOQ_ReadMallocWholeFile_AsyncAndWait(inName,&inSize64);
+    if ( inBuf == NULL )
+    {
+        OodleXLog_Printf_v0("ERROR couldn't load : %s\n",inName);
+        return 10;
+    }
+
+    inSize = OodleX_S64_to_SINTa_check(inSize64);
+
+    { // scope for compBuf
+        void * compBuf = OodleXMalloc_IOAligned( OodleLZ_GetCompressedBufferSizeNeeded(lzCompressor,inSize) );
+        
+        OO_U32 asyncSelect = OodleXAsyncSelect_All;
+        compSize = OodleXLZ_Compress_AsyncAndWait(asyncSelect,lzCompressor,inBuf,inSize,compBuf,lzLevel,&lzOptions);
+        
+        if ( compSize <= 0 )
+        {
+            OodleXLog_Printf_v0("ERROR failed to compress\n");
+            return 10;
+        }
+        
+        if ( ! OodleXIOQ_WriteWholeFile_AsyncAndWait(compName,compBuf,compSize) )
+        {
+            OodleXLog_Printf_v0("ERROR couldn't write : %s\n",compName);
+            return 10;
+        }
+        
+        // log about it :
+        OodleXLog_Printf_v1("%s compressed %d -> %d\n",inName,(int)inSize,(int)compSize);
+
+    /*
+
+        Make an OodleLZ_SeekTable on the compressed data for later use with parallel decompression.
+        
+        You should store the seektable to disk with any header information, and load it before loading the
+        compressed bulk data.  Oodle can do this for you with an "OOZ" file.
+        
+        For this example we'll just keep the seekTable in memory.  In real use you would write the
+        seek table with the compressed data.
+
+        Once we make the seekTable, we free compBuf - to get the compressed data again we will
+        have to read it from disk.
+
+    */
+        
+        seekTable = OodleLZ_CreateSeekTable(OodleLZSeekTable_Flags_None,
+                                                        lzOptions.seekChunkLen,NULL,inSize,compBuf,compSize);
+        
+        OodleXFree_IOAligned(compBuf);
+        compBuf = NULL;
+    }
+    
+    // we'll decompress into decompBuf
+    void * decompBuf = OodleXMalloc(inSize);
+    
+    //===================================================================================
+
+    // prevent cheating :
+    memset(decompBuf,(int)rand(),inSize);
+
+/*
+
+    Read the files with unbuffered IO
+    
+    We're mainly doing this here for consistent benchmarking so we can see IO times
+    otherwise the files would just come from the OS buffers and no show real IO time.
+    In practice you should usually use buffered IO. (OodleXFileOpenFlags_Default)
+
+*/
+
+    OodleXFileOpenFlags fileOpenFlags = OodleXFileOpenFlags_NotBuffered;
+
+/*
+
+    Now begins various ways to read the compressed file and decompress it.
+
+    Read-Decomp 1 :
+    
+    Read the whole compressed file, synchronously.
+    When that's done, decompress the whole buffer, synchronously.
+    
+    Simple, but stalls the main thread and gets no overlap of IO or parallelism in the decompress.
+
+*/
+    
+    OodleXLog_Printf_v1("Doing read then decomp, synchronously on the main thread :\n");
+    
+    {
+    void * compBuf = OodleXIOQ_ReadMallocWholeFile_AsyncAndWait(compName,NULL,fileOpenFlags);
+    if ( compBuf == NULL )
+    {
+        OodleXLog_Printf_v0("ReadMallocWholeFile failed on %s\n",compName);
+        return 10;
+    }
+    
+    OodleLZ_Decompress(compBuf,compSize,decompBuf,inSize,OodleLZ_FuzzSafe_Yes);
+    
+    // check it :
+    OOEX_ASSERT( memcmp(inBuf,decompBuf,inSize) == 0 );
+    
+    OodleXFree_IOAligned(compBuf);
+    }
+    
+    //===================================================================================
+    
+/*
+
+    Read-Decomp 2 :
+    
+    Read the whole compressed file, and fire off a full-buffer decompress, scheduled
+    to run automatically when the read is done.
+
+    This has the same latency as method 1, but is a single async operation so the main
+    thread can do something else the whole time.
+
+*/
+
+    OodleXLog_Printf_v1("Doing read then decomp, through a job chain :\n");
+    
+    // prevent cheating :
+    memset(decompBuf,(int)rand(),inSize);
+    
+    {
+    
+    // for simplicity, we'll just malloc compBuf with our known compSize
+    // more generally if you didn't know compSize, you would have to use an OodleWork coroutine
+    //  to first open the file, get the size, do the malloc, do the read, then the decompress
+    void * compBuf = OodleXMalloc_IOAligned(compSize);
+
+    // make an IO request to open and read the whole file : 
+    OodleXIOQFile compFile;
+    OodleXHandle openAndReadH = OodleXIOQ_OpenAndRead_Async(&compFile,compName,compBuf,OodleX_IOAlignUpSINTa(compSize),0,fileOpenFlags,0,OodleXHandleAutoDelete_Yes);
+
+    // go ahead and enqueue a Close to follow the OpenAndRead :
+    OodleXIOQ_CloseFile_Async(compFile,OODLEX_FILE_CLOSE_NO_TRUNCATE_SIZE,OodleXHandleAutoDelete_Yes);
+
+    // the decompress depends on the OpenAndRead - it will run only when that is done
+    //  dependencies are an array, but we have only one, so just point at it :
+    const OodleXHandle * deps = &openAndReadH;
+    int num_deps = 1;
+    
+    OodleXHandle decompH = OodleXLZ_Decompress_Narrow_Async(OodleXAsyncSelect_Full,compBuf,compSize,decompBuf,inSize,
+                                        OodleLZ_FuzzSafe_No,OodleLZ_CheckCRC_No,OodleLZ_Verbosity_None,
+                                        0,0,0,0,0,0,OodleLZ_Decode_Unthreaded,0,0,0,0,
+                                        OodleXHandleAutoDelete_No,
+                                        deps,num_deps);
+    
+    // ... main thread can do other work here ...
+    
+    OodleXStatus st = OodleX_WaitAndDelete(decompH);
+    if ( st != OodleXStatus_Done )
+    {
+        OodleXLog_Printf_v0("OodleXLZ_Decompress_Narrow_Async failed!\n");
+        return 10;
+    }
+    
+    // check it :
+    OOEX_ASSERT( memcmp(inBuf,decompBuf,inSize) == 0 );
+    
+    OodleXFree_IOAligned(compBuf);
+    }
+    
+    //===================================================================================
+    
+/*
+
+    Read-Decomp 3 :
+    
+    Read and decompress with IO overlap, but only using a single thread (not "wide").
+        
+    OodleXLZ_ReadAndDecompress_Stream_Async is an API provided to do IO overlap with decompression.
+    It's something you could easily write yourself in Oodle.  It uses a coroutine to do IO on chunks
+    and then decompress the chunks as they arive.  It tries to always be doing the IO for the next chunk
+    while decompressing the current chunk.
+    
+*/
+    
+    OodleXLog_Printf_v1("Doing read and decomp simultaneously :\n");
+    
+    // prevent cheating :
+    memset(decompBuf,(int)rand(),inSize);
+    
+    {
+    void * compBuf = OodleXMalloc_IOAligned(compSize);
+
+    /*
+    
+    Enqueue an open request + read an initial small chunk into compBuf
+    
+    We don't wait on the Open, but immediately start the OodleXLZ_ReadAndDecompress_Stream_Async operation,
+    and pass it the openAndRead handle.
+    
+    */  
+    
+    OodleXIOQFile compFile;
+    OO_SINTa initialReadSize = OodleX_IOAlignUpSINTa( OOEX_MIN(compSize,512*1024) );
+    OodleXHandle openAndReadH = OodleXIOQ_OpenAndRead_Async(&compFile,compName,compBuf,initialReadSize,0,fileOpenFlags);
+    
+    OodleXHandle readAndDecomp = OodleXLZ_ReadAndDecompress_Stream_Async(OodleXAsyncSelect_Full,
+                                            compBuf,compSize,decompBuf,inSize,
+                                            OodleLZ_FuzzSafe_No,OodleLZ_CheckCRC_No,OodleLZ_Verbosity_None,0,0,
+                                            compFile,compBuf,0,openAndReadH,initialReadSize);
+                                            
+    // enqueue a CloseFile with a dependency on the ReadAndDecomp operation :
+    OodleXIOQ_CloseFile_Async(compFile,OODLEX_FILE_CLOSE_NO_TRUNCATE_SIZE,OodleXHandleAutoDelete_Yes,OodleXPriority_Normal,&readAndDecomp,1);
+    
+    // ... main thread can do other work ...
+    
+    OodleXStatus st = OodleX_WaitAndDelete(readAndDecomp);
+    if ( st != OodleXStatus_Done )
+    {
+        OodleXLog_Printf_v0("ReadAndDecomp failed\n");
+        return 10;
+    }
+    
+    // check it :
+    OOEX_ASSERT( memcmp(inBuf,decompBuf,inSize) == 0 );
+    
+    OodleXFree_IOAligned(compBuf);
+    }
+    
+    //===================================================================================
+        
+/*
+
+    Read-Decomp 4 :
+    
+    Read and decompress with IO overlap, and using all worker threads ("wide").
+    
+    ReadAndDecompress_Wide needs the seekTable to find the compressed block boundaries.  Normally you
+    would have to send that in a file (see the OOZ APIs if you want Oodle to do it for you).  Here
+    we just use the seekTable we made earlier and kept in memory.
+        
+*/
+
+    OodleXLog_Printf_v1("Doing read and decomp wide :\n");
+    
+    // prevent cheating :
+    memset(decompBuf,(int)rand(),inSize);
+    
+    {
+    void * compBuf = OodleXMalloc_IOAligned(compSize);
+    
+    // open compressed file and do an initial read :
+    OodleXIOQFile compFile;
+    OO_SINTa initialReadSize = OodleX_IOAlignUpSINTa( OOEX_MIN(compSize,256*1024) );
+    OodleXHandle openAndReadH = OodleXIOQ_OpenAndRead_Async(&compFile,compName,compBuf,initialReadSize,0,fileOpenFlags);
+
+    // instead of waiting on openAndReadH here, you could pass it as a dependency
+    //  to the OodleXLZ_ReadAndDecompress_Wide_Async function to make the whole sequence async
+    // but we'll just stall here for simplicity in this example :
+    OodleXStatus st = OodleX_WaitAndDelete(openAndReadH);
+    
+    if ( st != OodleXStatus_Done )
+    {
+        OodleXLog_Printf_v0("OpenAndRead failed\n");
+        return 10;
+    }
+    
+    // fire off the read-and-decomp job :   
+    OodleXHandle readAndDecomp = OodleXLZ_ReadAndDecompress_Wide_Async(OodleXAsyncSelect_Full,
+                                            seekTable,compBuf,compSize,initialReadSize,compFile,0,
+                                            decompBuf,inSize,
+                                            OodleLZ_FuzzSafe_No,OodleLZ_CheckCRC_No,OodleLZ_Verbosity_None);
+
+    // enqueue a CloseFile with a dependency on the ReadAndDecomp operation :
+    OodleXIOQ_CloseFile_Async(compFile,OODLEX_FILE_CLOSE_NO_TRUNCATE_SIZE,OodleXHandleAutoDelete_Yes,OodleXPriority_Normal,&readAndDecomp,1);
+
+    // ... main thread can do other work ...
+        
+    st = OodleX_WaitAndDelete(readAndDecomp);
+    
+    if ( st != OodleXStatus_Done )
+    {
+        OodleXLog_Printf_v0("ReadAndDecomp failed\n");
+        return 10;
+    }
+    
+    // check it :
+    OOEX_ASSERT( memcmp(inBuf,decompBuf,inSize) == 0 );
+    
+    OodleXFree_IOAligned(compBuf);
+    }
+    
+    //===================================================================================
+    // all done, clean up :
+    
+    OodleLZ_FreeSeekTable(seekTable);
+    
+    OodleXFree(decompBuf);
+    OodleXFree_IOAligned(inBuf);
+
+    OodleX_Shutdown();
+
+    return 0;
+}
+
+
diff --git a/Engine/Source/Programs/Horde/HordeStorage/Jupiter.Common/Shared/Oodle/2.9.3/linux/examples/example_lz_simple.cpp b/Engine/Source/Programs/Horde/HordeStorage/Jupiter.Common/Shared/Oodle/2.9.3/linux/examples/example_lz_simple.cpp
new file mode 100644
--- /dev/null
+++ b/Engine/Source/Programs/Horde/HordeStorage/Jupiter.Common/Shared/Oodle/2.9.3/linux/examples/example_lz_simple.cpp
@@ -0,0 +1,169 @@
+//idoc(parent,examples)
+//idoc(tutorial,example_lz_simple,example_lz_simple : Example demonstrating very simple LZ memory->memory compression using only Oodle Core)
+
+/*
+
+Oodle example_lz_simple
+
+Very simple example of OodleLZ memory -> memory compression & decompression.
+
+Uses stdio for file IO to load an input file.
+
+See $example_lz for more advanced OodleLZ usage.
+
+example_lz_simple only uses Oodle Core, no Oodle Ext
+
+include oodle2.h and not oodle2x.h
+
+*/
+
+#include "../include/oodle2.h"
+
+#ifndef _CRT_SECURE_NO_WARNINGS
+#define _CRT_SECURE_NO_WARNINGS
+#endif
+#include <stdio.h>
+#include <string.h>
+#include <assert.h>
+#include <stdlib.h>
+
+#ifdef BUILDING_EXAMPLE_CALLER
+#define main example_lz_simple
+#endif
+
+#include "read_whole_file.h"
+
+extern "C" int main(int argc,char *argv[])
+{
+/* No initialization is needed for Oodle2 Core
+
+    we let Oodle Core use the default system plugins (the C stdlib)
+    To change them, use $OodleAPI_OodleCore_Plugins
+*/
+
+    // optional check to make sure header matches lib :
+    if ( ! Oodle_CheckVersion(OODLE_HEADER_VERSION) )
+    {
+        fprintf(stderr,"Oodle header version mismatch\n");
+        return 10;
+    }
+
+    // get args :
+    const char * in_name;
+    if ( argc < 2 )
+    {
+        in_name = "r:\\testsets\\lztestset\\lzt02";
+    }
+    else
+    {
+        in_name = argv[1];
+    }
+    
+/* read input file using stdio
+
+*/
+    OO_SINTa length;
+    void * buf = read_whole_file(in_name,&length);
+    if ( ! buf )
+    {
+        fprintf(stderr,"couldn't open : %s\n",in_name);
+        return 10;
+    }
+    
+
+/*
+
+Run OodleLZ_Compress from memory (buf) to memory (compbuf)
+
+Use the OodleLZ_Compressor_Kraken compressor.  Kraken is an amazing balance of good compression and fast decode
+speed.  It should generally be your first choice, then try Mermaid or Leviathan if you want faster decodes or more
+compression.
+
+Use OodleLZ_CompressionLevel_Normal level of effort in the encoder.  Normal is a balance of encode speed and compression
+ratio.  Different levels trade off faster or slower encoding for compressed size.  See $OodleLZ_CompressionLevel.
+
+See $OodleLZ_About for information on selection of the compression options.
+
+This call is synchronous and not threaded; see $example_lz for an example using the async compression APIs.
+
+*/
+
+    OodleLZ_Compressor compressor = OodleLZ_Compressor_Kraken;
+    OodleLZ_CompressionLevel level = OodleLZ_CompressionLevel_Normal;
+    //OodleLZ_CompressionLevel level = OodleLZ_CompressionLevel_Optimal; // for high compression, slower encode
+
+    // allocate memory big enough for compressed data :
+    void * compbuf = malloc( OodleLZ_GetCompressedBufferSizeNeeded(compressor,length) + sizeof(length) );
+    if ( compbuf == NULL )
+        return 10;
+        
+    char * compptr = (char *)compbuf;
+    memcpy(compptr,&length,sizeof(length));
+    compptr += sizeof(length);
+
+    // compress :
+    OO_SINTa complen = OodleLZ_Compress(compressor,buf,length,compptr,level);
+    compptr += complen;
+
+    // log about it :
+    //  full compressed size also includes the header +(int)sizeof(length)
+    printf("%s compressed %d -> %d\n",in_name,(int)length,(int)complen);
+
+/*
+
+Run OodleLZ_Decompress from memory (compbuf) to memory (decbuf)
+
+Note that you must provide the exact decompressed size.  OodleLZ data is headerless; store the size in
+your own header.
+
+We will allocate the needed decoder scratch mem to pass in, then OodleLZ_Decompress will do no allocations
+internally.  In real use you might want to keep the scratch mem allocated across calls so it doesn't need to
+be allocated and freed each time.  The scratch mem can be reused, but cannot be used by multiple threads at the
+same time.  See also example_lz_noallocs
+
+*/
+
+    OO_SINTa declength;
+    compptr = (char *)compbuf;
+    memcpy(&declength,compptr,sizeof(declength));
+    compptr += sizeof(declength);
+    assert( length == declength );
+
+    // malloc for decompressed buffer  :
+    void * decbuf = malloc( declength );
+
+    // allocate the decoder scratch memory needed
+    //  if we pass NULL for these (the default argument)
+    //  then OodleLZ_Decompress will allocate them internally
+    OO_SINTa decoderScratchMemSize = OodleLZDecoder_MemorySizeNeeded(compressor,-1);
+    void * decoderScratchMem = malloc( decoderScratchMemSize );
+
+    // do the decompress :
+    OO_SINTa decompress_return = OodleLZ_Decompress(compptr,complen,decbuf,declength,
+        OodleLZ_FuzzSafe_Yes,OodleLZ_CheckCRC_No,OodleLZ_Verbosity_None,NULL,0,NULL,NULL,
+        decoderScratchMem,decoderScratchMemSize);
+
+    // check it was successful :
+    assert( decompress_return == length );
+    assert( memcmp(buf,decbuf,length) == 0 );
+
+    if ( decompress_return != length )
+        return 10;
+
+    printf("decompressed successfully.\n");
+    
+/*
+
+And finish up.  No shutdown is needed for Oodle2 Core.
+
+*/
+
+    // free all the memory :
+    free(buf);
+    free(compbuf);
+    free(decbuf);
+    free(decoderScratchMem);
+
+    return 0;
+}
+
