Moved Horde.Storage from git to p4

diff --git a/Engine/Source/Programs/Horde/HordeStorage/Jupiter.Common/Shared/Oodle/2.9.3/linux/examples/example_lz_threadphased.cpp b/Engine/Source/Programs/Horde/HordeStorage/Jupiter.Common/Shared/Oodle/2.9.3/linux/examples/example_lz_threadphased.cpp
new file mode 100644
--- /dev/null
+++ b/Engine/Source/Programs/Horde/HordeStorage/Jupiter.Common/Shared/Oodle/2.9.3/linux/examples/example_lz_threadphased.cpp
@@ -0,0 +1,420 @@
+//idoc(parent,examples)
+//idoc(tutorial,example_lz_threadphased,example_lz_threadphased : Example of 2-thread ThreadPhased decoding)
+
+/*
+
+Oodle example_lz_threadphased
+
+This example demonstrates the ability of Kraken to decode using "ThreadPhased" parallelism
+
+See $OodleLZ_About_ThreadPhasedDecode for details.
+
+ThreadPhased decoding provides a 1X-2X speedup, typically around 33%-50%
+
+This example implements an entire ThreadPhased decoder, in example_lz_threadphased_decompress.
+
+The idea is you may take this example code and modify it to run in your own threading or job system,
+so that you may implement the threading however you look.
+
+example_lz_threadphased_decompress can be used without OodleX or the Oodle Worker system.
+
+
+In this example I use OodleX_CreateThread and OodleX_Semaphore ; these are just handy cross-platform
+implementations for me.  The intention is that you replace them with your own threading system calls.
+
+For the semaphore, it's important that it tries to avoid going into an OS Wait (thread sleep) when the
+two threads are nearly synchronized.  To avoid this you want a user-space spin-backoff loop to try to keep
+the two threads awake together.  One option is to use something like "fastsemaphore" as a wrapper to your
+underlying OS semaphore :
+
+http://cbloomrants.blogspot.com/2011/12/12-08-11-some-semaphores.html
+
+If you use that "fastsemaphore" make sure to set
+        int spin_count = 1; // ! set this for your system
+something like 100 is usually reasonable, and probably add a backoff.
+
+*/
+
+#include "../include/oodle2x.h"
+#include "ooex.h"
+
+#include <stdlib.h>
+#include <stdio.h>
+#include <string.h>
+
+#include "make_example_input.h"
+
+#ifdef BUILDING_EXAMPLE_CALLER
+#define main example_lz_threadphased
+#endif
+
+//===========================================================
+// file names :
+
+static const char * in_name = "oodle_example_input_file";
+
+//===========================================================
+
+static bool example_lz_threadphased_decompress(void * comp_buf,OO_SINTa comp_len,void * dec_buf,OO_SINTa dec_size,bool async);
+    
+//===========================================================
+
+extern "C" int main(int argc,char *argv[])
+{
+    // Init Oodle systems with default options :
+    if ( ! OodleX_Init_Default(OODLE_HEADER_VERSION) )
+    {
+        fprintf(stderr,"OodleX_Init failed.\n");
+        return 10;
+    }   
+
+    if ( argc >= 2 ) in_name = argv[1];
+    else make_example_input(in_name);
+
+    OodleXLog_Printf_v1("example_lz_threadphased : %s\n",in_name);
+
+    // read the input file to the global buffer :
+    OO_S64 in_size_64;
+    void * in_buffer = OodleXIOQ_ReadMallocWholeFile_AsyncAndWait(in_name,&in_size_64);
+    
+    if ( ! in_buffer)
+    {
+        OodleXLog_Printf_v0("failed to read %s\n",in_name); 
+        return 10;
+    }
+    
+    OO_SINTa in_size = OodleX_S64_to_SINTa_check( in_size_64 );
+    
+    //=========================================
+    // select options :
+    
+    OodleLZ_Compressor compressor = OodleLZ_Compressor_Kraken;
+    //OodleLZ_CompressionLevel level = OodleLZ_CompressionLevel_Normal;
+    OodleLZ_CompressionLevel level = OodleLZ_CompressionLevel_Fast;
+
+    // this example is only valid for compressors that support ThreadPhased decode : (eg. Kraken)
+    OOEX_ASSERT( OodleLZ_Compressor_CanDecodeThreadPhased(compressor) );
+
+    //=========================================
+    
+    // allocate memory big enough for compressed data :
+    void * comp_buf = OodleXMalloc( OodleLZ_GetCompressedBufferSizeNeeded(compressor,in_size) );    
+    // memory to decode to :
+    void * dec_buf = OodleXMalloc( in_size );
+    
+    //=========================================
+    // compress the input :
+    
+    // compress :
+    // this is just a normal whole-block compress, no special parallel mode is needed
+    OO_SINTa comp_len = OodleLZ_Compress(compressor,in_buffer,in_size,comp_buf,level);
+    
+    OodleXLog_Printf_v1("Compressed : %d -> %d\n",(int)in_size,(int)comp_len);
+    
+    //=======================================
+    /*
+
+    We can decompress asynchronously using the "ThreadPhased" helper in OodleX.
+    
+    The Narrow helper uses 2 threads and frees up the calling thread to do other work.  
+    
+    OodleXLZ_Decompress_ThreadPhased_Narrow_Async uses the OodleX Worker thread system, which is started by
+    default in OodleX_Init.
+    
+    NOTE for maximum speed you should pass in the scratch space needed by OodleXLZ_Decompress_ThreadPhased_Narrow_Async
+    pre-allocated, so it doesn't have to do the allocation internally.
+    
+    */
+    
+    OodleXLog_Printf_v1("OodleXLZ_Decompress_ThreadPhased_Narrow_Async...\n");
+    
+    OodleXHandle decomp_handle = OodleXLZ_Decompress_ThreadPhased_Narrow_Async(comp_buf,comp_len,dec_buf,in_size);
+    
+    // .. can do other work on the main thread now ..
+    
+    OodleXStatus decomp_status = OodleX_WaitAndDelete(decomp_handle);
+    if ( decomp_status == OodleXStatus_Error )
+    {
+        OodleXLog_Printf_v1("Error!\n");
+    }
+    
+    // check it :
+    OOEX_ASSERT( memcmp(in_buffer,dec_buf,in_size) == 0 );
+    
+    //=======================================
+    // do our own thread-phased decode :
+        
+    OodleXLog_Printf_v1("example_lz_threadphased_decompress ");
+
+    // run a few reps to stress test :  
+    for(int rep=0;rep<10;rep++)
+    {
+        // run async and sync options :
+        for(int async = 0;async<=1;async++)
+        {
+            OodleXLog_Printf_v1(async ? "+" : "-");
+        
+            // wipe out dec_buf to make sure we decode correctly :
+            memset(dec_buf,0xEE,in_size);
+        
+            if ( !  example_lz_threadphased_decompress(comp_buf,comp_len,dec_buf,in_size,!!async) )
+            {
+                OodleXLog_Printf_v1("Error!\n");
+            }
+        
+            // check it :
+            OOEX_ASSERT( memcmp(in_buffer,dec_buf,in_size) == 0 );
+        }
+    }
+    OodleXLog_Printf_v1("\n");
+    
+    //=======================================
+    
+    OodleXFree(comp_buf);
+    OodleXFree(dec_buf);
+    OodleXFree_IOAligned(in_buffer);
+
+    OodleX_Shutdown(NULL,OodleX_Shutdown_LogLeaks_Yes,0);
+    
+    return 0;
+}
+
+//===========================================================
+/***
+
+example_lz_threadphased_threadfunc
+
+Example of how to run a ThreadPhased decode yourself.
+
+The basic idea of ThreadPhased decoding is that the OodleLZ_Decompress work on each BLOCK can be
+split into two phases.  This can be invoked by just calling OodleLZ_Decompress twice on the same
+block, first with OodleLZ_Decode_ThreadPhase1, then with OodleLZ_Decode_ThreadPhase2.
+
+To get parallelism, we can run the two phases on two separate threads.
+
+The rule is that you must run the Phase2 on each block after the Phase1 for that block is done,
+and with the same "decoderMem" pointer.  The Phase2 decodes on all blocks must be done in
+sequential order (unless they are Seek Resets).
+The decoder memory used for OodleLZ_Decompress here must be larger than normal, of size
+OodleLZ_ThreadPhased_BlockDecoderMemorySizeNeeded().
+
+Our thread model here is a two-thread circular buffer scan with semaphore signalling.
+
+Thread 1 does :
+
+    For each block
+    Wait on sem_blocksavail to get an available circular buffer slot
+    Do Phase1 Decompress into a slot
+    Post sem_phase1done
+
+Thread 2 does :
+    
+    For each block
+    Wait on sem_phase1done to get a slot with phase1 decode done
+    Do Phase2 Decompress on that slot
+    Post sem_blocksavail to signal Thread 1 that this slot may be reused
+
+    
+***/
+
+struct example_lz_threadphased_threaddata
+{
+    volatile OO_U32 * error_cancel;
+    OodleX_Semaphore * sem_consume;
+    OodleX_Semaphore * sem_produce;
+    
+    OO_SINTa num_blocks;
+    OO_SINTa num_scratch_blocks;
+    OO_U8 * scratch_mem;
+    OO_SINTa scratch_block_size;
+    
+    OO_U8 * rawBuf; OO_SINTa rawSize;
+    const OO_U8 * compBuf; OO_SINTa compSize;
+    OodleLZ_Decode_ThreadPhase threadPhase;
+    
+    bool success;
+};
+
+#define THREAD_ERROR    0
+#define THREAD_SUCCESS  1
+
+// thread_function_DecodePhase
+//  can be used for both Phase1 and Phase2 !
+static OO_U32 OODLE_CALLBACK example_lz_threadphased_threadfunc( void * user_data )
+{
+    example_lz_threadphased_threaddata * data = (example_lz_threadphased_threaddata *)user_data;
+    
+    const OO_U8 * compPtr = (OO_U8 *)(data->compBuf);
+    const OO_U8 * compEnd = compPtr + data->compSize;
+    
+    OO_SINTa decoderMemSize = data->scratch_block_size;
+        
+    OO_SINTa scratch_i = 0;
+    for(OO_SINTa block_pos = 0;block_pos < data->rawSize;block_pos += OODLELZ_BLOCK_LEN, scratch_i++)
+    {
+        // consume one :
+        OodleX_Semaphore_Wait(data->sem_consume);
+    
+        // relaxed load of shared error_cancel variable ; Sem Wait acts as Acquire barrier
+        if ( *(data->error_cancel) ) return THREAD_ERROR;
+    
+        if ( scratch_i == data->num_scratch_blocks ) scratch_i = 0;
+        
+        OO_U8 * decoderMem = data->scratch_mem + scratch_i * decoderMemSize;
+        
+        OO_U8 * chunk_ptr = (OO_U8 *)(data->rawBuf) + block_pos;
+                    
+        OO_SINTa block_len = OOEX_MIN((data->rawSize - block_pos),OODLELZ_BLOCK_LEN);
+                    
+        OO_BOOL indy;
+        OO_SINTa block_complen = OodleLZ_GetCompressedStepForRawStep(compPtr,compEnd-compPtr,block_pos,block_len,NULL,&indy);
+        if ( block_complen <= 0 )
+        {
+            // handle error
+            // relaxed store of shared variable; Sem Post acts as release barrier
+            *(data->error_cancel) = 1;
+            OodleX_Semaphore_Post(data->sem_produce);
+            return THREAD_ERROR;
+        }
+                    
+        OO_SINTa gotLen = OodleLZ_Decompress(compPtr,block_complen,chunk_ptr,block_len,
+            OodleLZ_FuzzSafe_Yes,OodleLZ_CheckCRC_No,OodleLZ_Verbosity_None,
+            data->rawBuf,data->rawSize,
+            NULL,NULL,
+            decoderMem,decoderMemSize,
+            data->threadPhase
+            );
+        OOEX_ASSERT( gotLen == block_len+block_pos );
+        if ( gotLen != block_len+block_pos )
+        {
+            // handle error
+            // relaxed store of shared variable; Sem Post acts as release barrier
+            *(data->error_cancel) = 1;
+            OodleX_Semaphore_Post(data->sem_produce);
+            return THREAD_ERROR;
+        }
+        
+        compPtr += block_complen;
+        
+        OodleX_Semaphore_Post(data->sem_produce);
+    }
+
+    data->success = true;   
+    return THREAD_SUCCESS;
+}
+
+static bool example_lz_threadphased_decompress(void * comp_buf,OO_SINTa comp_len,void * dec_buf,OO_SINTa dec_size,bool async)
+{
+    /**
+    
+    circularBufferBlockCount is the number of circular buffer slots for the two threads to
+    communicate through
+    
+    circularBufferBlockCount >= 2
+    higher is faster because it allows less synchronization of the two threads
+    lower means less memory required
+    circularBufferBlockCount >= 4 is reasonable
+    circularBufferBlockCount >= 6 is close to full speed
+    
+    **/
+    OO_SINTa circularBufferBlockCount = 6; // parameter
+
+    // async is an option for whether the whole operation is run async off this thread or not
+    //  (eg. with 2 additional threads or 1 additional thread)
+
+    // check that the data contains a valid ThreadPhased compressor :
+    OodleLZ_Compressor compressor = OodleLZ_GetAllChunksCompressor(comp_buf,comp_len,dec_size);
+    if (! OodleLZ_Compressor_CanDecodeThreadPhased(compressor) )
+    {
+        OodleXLog_Printf_v1("Asked for ThreadPhase decode but ! OodleLZ_Compressor_CanDecodeThreadPhased\n");
+        return false;
+    }
+
+    // count the number of OODLELZ_BLOCK_LEN in the total size :
+    OO_SINTa nBlocks = (dec_size + OODLELZ_BLOCK_LEN-1)/OODLELZ_BLOCK_LEN;
+
+    // don't need more than nBlocks :
+    circularBufferBlockCount = OOEX_MIN(circularBufferBlockCount,nBlocks);
+
+    // allocate space for the scratch circular buffer :
+    OO_SINTa scratchBlockSize = OodleLZ_ThreadPhased_BlockDecoderMemorySizeNeeded();
+    OO_SINTa scratchBufSize = scratchBlockSize * circularBufferBlockCount;
+    
+    // NOTE in production you may wish to preallocate this memory
+    void * scratchBuf = OodleXMalloc(scratchBufSize);
+    
+    //=========================================================
+    // set up the data needed for the thread phases
+    // NOTE if you want to make the whole decode asynchronous,
+    //   you can't put this on the stack, you need to package it up in memory
+    
+    // OodleX_Semaphore just initialize with 0 :
+    OodleX_Semaphore sem_blocksavail = 0;
+    OodleX_Semaphore sem_phase1done = 0;
+    OO_U32 shared_error_cancel = 0; // shared atomic variable
+    
+    // starting state is that all circular buffer slots are available :
+    OodleX_Semaphore_Post(&sem_blocksavail,(OO_S32)circularBufferBlockCount);
+    
+    example_lz_threadphased_threaddata td1 = { 0 };
+    
+    td1.error_cancel = &shared_error_cancel;
+    td1.success = false;
+    td1.compBuf = (OO_U8 *)(comp_buf);
+    td1.compSize = comp_len;
+    td1.rawBuf = (OO_U8 *)(dec_buf);
+    td1.rawSize = dec_size;
+    td1.num_blocks = nBlocks;
+    td1.num_scratch_blocks = circularBufferBlockCount;
+    td1.scratch_mem = (OO_U8 *)(scratchBuf);
+    td1.scratch_block_size = scratchBlockSize;
+        
+    // thread 1 waits for blocks to be available in the circular buffer
+    //  and posts that phase1 is done
+    td1.threadPhase = OodleLZ_Decode_ThreadPhase1;
+    td1.sem_consume = &sem_blocksavail;
+    td1.sem_produce = &sem_phase1done;
+    
+    // thread 2 waits for each block to reach phase1done
+    //  and then posts that the block is reusable
+    //  same as thread 1, just swap the semaphores :
+    example_lz_threadphased_threaddata td2;
+    td2 = td1;
+    td2.threadPhase = OodleLZ_Decode_ThreadPhase2;
+    td2.sem_consume = &sem_phase1done;
+    td2.sem_produce = &sem_blocksavail;
+    
+    // create a thread to run Phase1 :
+    // NOTE : in production you probably don't want to create a thread every time you make
+    //  this decompress call.  Rather use an idle thread that's already created.
+    OodleX_Thread thread1 = OodleX_CreateThread(example_lz_threadphased_threadfunc,&td1);
+    
+    // either run Phase2 asynchronously (on another thread) or synchronously on this thread :
+    if ( async )
+    {
+        OodleX_Thread thread2 = OodleX_CreateThread(example_lz_threadphased_threadfunc,&td2);
+    
+        // ... current thread is now available while decompress runs on 2 other threads ...
+        // ... return and do other work ...
+    
+        OodleX_WaitAndDestroyThread(thread2);
+    }
+    else
+    {
+        // synchronous version - just run Phase2 on this thread :
+        example_lz_threadphased_threadfunc(&td2);
+    }
+    
+    OodleX_WaitAndDestroyThread(thread1);
+    
+    // OodleX_Semaphore doesn't need cleanup
+    
+    //===========================================================
+    
+    OodleXFree(scratchBuf);
+    
+    bool ok = td1.success && td2.success;
+    
+    return ok;
+}
+
diff --git a/Engine/Source/Programs/Horde/HordeStorage/Jupiter.Common/Shared/Oodle/2.9.3/linux/examples/job_stress_test.h b/Engine/Source/Programs/Horde/HordeStorage/Jupiter.Common/Shared/Oodle/2.9.3/linux/examples/job_stress_test.h
new file mode 100644
--- /dev/null
+++ b/Engine/Source/Programs/Horde/HordeStorage/Jupiter.Common/Shared/Oodle/2.9.3/linux/examples/job_stress_test.h
@@ -0,0 +1,539 @@
+/*********
+
+Stress tester for Oodle Job system implementations.
+
+Use to help verify that the jobify plugin functions you provide to Oodle are working.
+
+Use like :
+
+Oodle_JobSystemStressTest(my_CorePlugin_RunJob, my_CorePlugin_WaitJob, NULL, 0, 256, 50000);
+
+*********/
+
+#include <stdio.h>
+#include <string.h>
+#include <assert.h>
+
+
+namespace JobTester
+{
+
+static inline U32 rotl(U32 x, int k)
+{
+    return (x << k) | (x >> (32 - k));
+}
+
+// Our hash function
+// this is Bob Jenkin's loookup3 for just two 32-bit integers
+static U32 hash_pair(U32 x, U32 y)
+{
+    U32 a, b, c;
+    a = b = c = 0xdeadbeef + 8;
+
+    a += x;
+    b += y;
+
+    // final mix
+    c ^= b; c -= rotl(b,14);
+    a ^= c; a -= rotl(c,11);
+    b ^= a; b -= rotl(a,25);
+    c ^= b; c -= rotl(b,16);
+    a ^= c; a -= rotl(c,4);
+    b ^= a; b -= rotl(a,14);
+    c ^= b; c -= rotl(b,24);
+
+    return c;
+}
+
+struct Context;
+
+struct Step
+{
+    // Type of work step to perform
+    enum StepType
+    {
+        FORK,
+        JOIN,
+    };
+
+    StepType type;
+    int index;
+};
+
+struct Job
+{
+    enum
+    {
+        MAX_FORKS = 32,
+        MAX_STEPS = MAX_FORKS*2, // one join per fork
+        MAX_DEPS = OODLE_JOB_MAX_DEPENDENCIES,
+    };
+
+    Context * ctx;          // work context to use
+
+    int id;                 // my global job ID (index into jobs, results arrays)
+    int parent_id;          // ID of parent job
+    int first_child_id;     // ID of first child job (other children are sequential after that)
+
+    U32 seed;               // seed value for computation performed
+    int ndeps;              // number of input deps
+    int deps[MAX_DEPS];     // indices of sibling jobs we depend on, local to parent (add first_sibling_id to get global ID)
+
+    int nsteps;             // number of work steps to perform
+    Step steps[MAX_STEPS];  // description of the computation to perform
+
+    static void OODLE_CALLBACK run_thunk(void * context)
+    {
+        static_cast<Job*>(context)->run();
+    }
+
+    void run() const;
+};
+
+struct Context
+{
+    enum
+    {
+        MAX_JOBS = 256, // no need to go overboard with this
+    };
+
+    t_fp_OodleCore_Plugin_RunJob * fp_RunJob;
+    t_fp_OodleCore_Plugin_WaitJob * fp_WaitJob;
+    void * user_ptr;
+
+    int next_job_id;
+    int num_jobs_total;
+
+    Job jobs[MAX_JOBS];
+    U32 ref_results[MAX_JOBS];
+    U32 job_results[MAX_JOBS];
+};
+
+void Job::run() const
+{
+    U64 handles[MAX_FORKS];
+
+    U32 x = seed;
+    for (int i = 0; i < ndeps; i++)
+    {
+        // Determine global job ID for the job we depend on
+        int dep_job_id = ctx->jobs[parent_id].first_child_id + deps[i];
+        x = hash_pair(x, ctx->job_results[dep_job_id]);
+    }
+
+    // Run the desired sequence of work steps
+    for (int i = 0; i < nsteps; i++)
+    {
+        int index = steps[i].index;
+        int subjob_id = first_child_id + index;
+        assert(subjob_id >= 0 && subjob_id < ctx->num_jobs_total);
+        assert(index < MAX_FORKS);
+
+        if (steps[i].type == Step::FORK)
+        {
+            Job * subjob = &ctx->jobs[subjob_id];
+            U64 runjob_deps[MAX_DEPS];
+            int nrunjob_deps = 0;
+
+            // Work out our dependency list; jobs that we've already waited for,
+            // or that ran synchronously (handle value of 0) don't count.
+            for (int j = 0; j < subjob->ndeps; j++)
+            {
+                int dep_index = subjob->deps[j];
+                if (handles[dep_index])
+                    runjob_deps[nrunjob_deps++] = handles[dep_index];
+            }
+
+            // Run the job!
+            handles[index] = ctx->fp_RunJob(run_thunk, subjob, runjob_deps, nrunjob_deps, ctx->user_ptr);
+        }
+        else // join
+        {
+            // If we didn't get a 0 handle from RunJob, wait for job completion!
+            if (handles[index])
+            {
+                ctx->fp_WaitJob(handles[index], ctx->user_ptr);
+                handles[index] = 0; // mark as done
+            }
+
+            // Immediately incorporate the result of the computation
+            x = hash_pair(x, ctx->job_results[subjob_id]);
+        }
+    }
+
+    // Could do a random delay here, but it seems to be working great and
+    // finding tricky interleavings without.
+
+    // Save the result of the computation and we're done
+    ctx->job_results[id] = x;
+}
+
+// ---- Randomness utilities
+
+// The RNG here is PCG XSH RR 64/32 MCG; this has a moderate state size
+// and decent quality.
+class Rng
+{
+    static const U64 MCG_MUL = 6364136223846793005ULL; 
+    U64 state;
+
+public:
+    static Rng seed(U32 seed)
+    {
+        // State may not be 0 (MCG); that's why we flip the low bits
+        // also do one multiply step in case the input is a small integer (which it often is).
+        Rng r;
+        r.state = ~(U64)seed | (U64(seed) << 32);
+        return r;
+    }
+
+    // Random 32-bit uint
+    U32 random()
+    {
+        U64 oldstate = state;
+        U32 rot_input = (U32) (((oldstate >> 18) ^ oldstate) >> 27);
+        U32 rot_amount = (U32) (oldstate >> 59);
+        U32 output = (rot_input >> rot_amount) | (rot_input << ((0u - rot_amount) & 31)); // rotr(rot_input, rot_amount)
+
+        // Advance multiplicative congruential generator
+        // Constant from PCG reference impl.
+        state = oldstate * MCG_MUL;
+
+        return output;
+    }
+
+    // Random uint x with 0 <= x <= max
+    U32 random_with_max(U32 max)
+    {
+        // Smear right to determine next-larger pow2 mask
+        U32 mask = max;
+        mask |= mask >> 16;
+        mask |= mask >>  8;
+        mask |= mask >>  4;
+        mask |= mask >>  2;
+        mask |= mask >>  1;
+
+        // Expected number of iterations is below 2
+        for(;;)
+        {
+            U32 x = random() & mask;
+            if (x <= max)
+                return x;
+        }
+    }
+
+    U32 random_min_max(U32 min, U32 max)
+    {
+        return min + random_with_max(max - min);
+    }
+};
+
+// Generate a random balanced sequence of n +1s and n -1s corresponding to
+// correctly nested brackets - that, is, all of the partial sums are
+// nonnegative.
+static void random_bracket_sequence(S8 * result, int n, Rng &rng)
+{
+    // Algorithm follows Atkinson & Sack, "Generating binary trees at random", 1992
+    int partial_sum = 0;
+    int left = 0;
+    int right = 2*n;
+    int word_start = 0;
+    int npos_left = n;
+
+    for (int i = 0; i < 2*n; i++)
+    {
+        S8 sel;
+
+        // Steps 1 and 2: Selection sampling: of the 2*n-i remaining items, npos_left
+        // need to be +1s.
+        if (int(rng.random_with_max(2*n - i - 1)) < npos_left)
+        {
+            sel = 1;
+            --npos_left;
+        }
+        else
+            sel = -1;
+        result[left++] = sel;
+
+        // Step 3: make it well-formed
+        partial_sum += sel;
+        if (partial_sum == 0) // at end of an irreducible balanced word
+        {
+            if (sel == -1)
+                word_start = left; // it was well-formed! we're all good.
+            else
+            {
+                // not well-formed; fix it!
+                // move the middle of the current word to the right side and flip all the signs
+                // (skip the first and last items, which we know are -1 and +1, respectively)
+                for (int j = left - 2; j > word_start; --j)
+                    result[--right] = -result[j];
+                // other half of the fixup
+                result[word_start++] = 1;
+                result[--right] = -1;
+                left = word_start;
+            }
+        }
+    }
+
+    assert(npos_left == 0);
+    assert(left == right);
+    assert(partial_sum == 0);
+}
+
+// ---- Random job generator
+
+static void generate_random_job_rec(Context * ctx, int id, int parent_id, const int * deps, int ndeps, Rng &rng, int rec_level, int max_rec_level)
+{
+    Job * job = &ctx->jobs[id];
+    assert(ndeps <= Job::MAX_DEPS);
+
+    job->ctx = ctx;
+    job->id = id;
+    job->parent_id = parent_id;
+    job->first_child_id = ctx->next_job_id;
+
+    job->seed = rng.random();
+    job->ndeps = ndeps;
+
+    // We also compute the reference result as we're setting up the job
+    U32 ref_result = job->seed;
+
+    for (int i = 0; i < ndeps; i++)
+    {
+        job->deps[i] = deps[i];
+        int dep_job_id = ctx->jobs[parent_id].first_child_id + deps[i];
+        assert(dep_job_id < id);
+        ref_result = hash_pair(ref_result, ctx->ref_results[dep_job_id]);
+    }
+
+    // Figure out how many jobs to fork
+    int max_jobs = ctx->num_jobs_total - ctx->next_job_id;
+    if (max_jobs > Job::MAX_FORKS)
+        max_jobs = Job::MAX_FORKS;
+
+    // Deepest level may not spawn further jobs
+    if (rec_level >= max_rec_level)
+        max_jobs = 0;
+
+    int fork_count = rng.random_with_max(max_jobs);
+
+    // Avoid root jobs with a fork count of 0 because that would be pointless
+    if (parent_id == -1 && fork_count == 0 && max_jobs != 0)
+        fork_count = 1;
+
+    ctx->next_job_id += fork_count;
+    job->nsteps = fork_count * 2;
+
+    S8 spawn_sequence[Job::MAX_STEPS];
+    int child_inds[Job::MAX_FORKS];
+    int nlive_children = 0;
+    int fork_index = 0;
+
+    // Generate a random bracket sequence corresponding to the forks/joins
+    random_bracket_sequence(spawn_sequence, fork_count, rng);
+
+    for (int i = 0; i < job->nsteps; i++)
+    {
+        Step * step = &job->steps[i];
+        step->type = (spawn_sequence[i] > 0) ? Step::FORK : Step::JOIN;
+        if (step->type == Step::FORK)
+        {
+            step->index = fork_index;
+
+            // Randomly select dependencies from among the already-spawned children
+            // We're OK with putting dependencies on already waited-for jobs
+            // (i.e. we still want the data dependency on some things we already waited
+            // for); the driver code handles this, as does the equivalent logic in Oodle
+            // itself, and we want to test it.
+            int ndep_max = fork_index < Job::MAX_DEPS ? fork_index : Job::MAX_DEPS;
+            int jndeps = rng.random_with_max(ndep_max);
+            int jdeps[Job::MAX_DEPS];
+            int nleft_to_select = jndeps;
+            for (int j = 0; j < fork_index; j++)
+            {
+                // of the fork_index - j remaining children, we want to select nleft_to_select
+                if (int(rng.random_with_max(fork_index - j - 1)) < nleft_to_select)
+                    jdeps[--nleft_to_select] = j;
+            }
+            assert(nleft_to_select == 0);
+
+            // Set up the job recursively
+            generate_random_job_rec(ctx, job->first_child_id + fork_index, id, jdeps, jndeps, rng, rec_level + 1, max_rec_level);
+
+            // Log as one of the live children
+            child_inds[nlive_children++] = fork_index;
+            fork_index++;
+        }
+        else // JOIN
+        {
+            // Randomly pick one of the live children and wait for it
+            assert(nlive_children > 0);
+            int slot = rng.random_with_max(nlive_children - 1);
+
+            step->index = child_inds[slot];
+
+            // Update our reference result
+            ref_result = hash_pair(ref_result, ctx->ref_results[job->first_child_id + step->index]);
+
+            // Remove that item from the list of live children
+            child_inds[slot] = child_inds[--nlive_children];
+        }
+    }
+
+    assert(fork_index == fork_count);
+    assert(nlive_children == 0);
+
+    // Store reference result
+    ctx->ref_results[id] = ref_result;
+}
+
+static int generate_random_jobs(Context * ctx, int num_jobs_total, Rng &rng, int max_rec_level)
+{
+    assert(num_jobs_total >= 1 && num_jobs_total <= Context::MAX_JOBS);
+
+    ctx->next_job_id = 1;
+    ctx->num_jobs_total = num_jobs_total;
+    for (int i = 0; i < Context::MAX_JOBS; i++)
+    {
+        ctx->ref_results[i] = 0;
+        ctx->job_results[i] = 0;
+    }
+
+    generate_random_job_rec(ctx, 0, -1, NULL, 0, rng, 0, max_rec_level);
+
+    // ensure we didn't generate too many jobs
+    assert(ctx->next_job_id <= ctx->num_jobs_total);
+
+    // but we might not have used all of them! so return real count.
+    ctx->num_jobs_total = ctx->next_job_id;
+    return ctx->num_jobs_total;
+}
+
+static void describe_job(FILE * dst, const Context * ctx, int job_id)
+{
+    const Job * job = ctx->jobs + job_id;
+
+    fprintf(dst, "job %d:\n", job_id);
+    fprintf(dst, "  result expected=0x%08x, actual=0x%08x\n", ctx->ref_results[job_id], ctx->job_results[job_id]);
+    fprintf(dst, "  parent=%d\n", job->parent_id);
+    fprintf(dst, "  seed=0x%08x\n", job->seed);
+    fprintf(dst, "  depends = {");
+    for (int i = 0; i < job->ndeps; i++)
+        fprintf(dst, "%c%d", (i == 0) ? ' ' : ',', ctx->jobs[job->parent_id].first_child_id + job->deps[i]);
+    fprintf(dst, " }\n");
+    fprintf(dst, "  work sequence:\n");
+    for (int i = 0; i < job->nsteps; i++)
+    {
+        const Step *step = job->steps + i;
+        fprintf(dst, "  [%2d] %s job %d\n", i, (step->type == Step::FORK) ? "fork" : "join", job->first_child_id + step->index);
+    }
+    fprintf(dst, "  [%2d] done\n", job->nsteps);
+    fprintf(dst, "\n");
+}
+
+} // namespace JobTester
+
+void Oodle_JobSystemStressTest(
+    t_fp_OodleCore_Plugin_RunJob * fp_RunJob,
+    t_fp_OodleCore_Plugin_WaitJob * fp_WaitJob,
+    void * user_ptr,
+    U32 random_seed, int max_jobs_per_test, int num_tests_to_run)
+{
+    using namespace JobTester;
+
+    if (max_jobs_per_test < 2)
+        max_jobs_per_test = 2;
+    if (max_jobs_per_test > Context::MAX_JOBS)
+        max_jobs_per_test = Context::MAX_JOBS;
+
+    Context * ctx = new Context;
+
+    ctx->fp_RunJob = fp_RunJob;
+    ctx->fp_WaitJob = fp_WaitJob;
+    ctx->user_ptr = user_ptr;
+
+    U32 cur_seed = random_seed;
+    int total_num_jobs = 0;
+    bool errors = false;
+
+    for (int i = 0; i < num_tests_to_run; i++)
+    {
+        Rng rng = Rng::seed(cur_seed);
+        
+        // The root job runs on this thread, so always use at least 2 jobs
+        // so something interesting happens.
+        int job_count = rng.random_min_max(2, max_jobs_per_test);
+
+        // Generate a random set of jobs
+        // we don't let any but the first level spawn (and wait for) child jobs
+        job_count = generate_random_jobs(ctx, job_count, rng, 1);
+
+        // Run the root job (which spawns everything else!)
+        ctx->jobs[0].run();
+
+        if (memcmp(ctx->ref_results, ctx->job_results, job_count * sizeof(U32)) != 0)
+        {
+            fprintf(stderr, "JobSystemStressTest: ERROR in test run %d/%d\n", i, num_tests_to_run);
+            fprintf(stderr, "repro with seed=0x%08x max_jobs_per_test=%d num_tests_to_run=1\n", cur_seed, max_jobs_per_test);
+            fprintf(stderr, "%d jobs in this test instance.\n", job_count);
+
+            // Identify first failed test, which is the _largest_ index with a difference since parent tasks
+            // have results that depend on their child tasks.
+            int first_mismatch = job_count - 1;
+            while (ctx->ref_results[first_mismatch] == ctx->job_results[first_mismatch])
+                --first_mismatch;
+
+            // If this job has a parent, it might have siblings it depends on that started earlier
+            // and were also wrong; see if that's the case.
+            int parent_id = ctx->jobs[first_mismatch].parent_id;
+            int first_sibling_id = -1;
+
+            if (parent_id != -1)
+            {
+                // Go to the first sibling with a mismatch
+                first_sibling_id = ctx->jobs[parent_id].first_child_id;
+                first_mismatch = first_sibling_id;
+                while (ctx->ref_results[first_mismatch] == ctx->job_results[first_mismatch])
+                    ++first_mismatch;
+            }
+
+            fprintf(stderr, "First mismatch is on job %d:\n\n", first_mismatch);
+            describe_job(stderr, ctx, first_mismatch);
+
+            if (parent_id != -1)
+            {
+                fprintf(stderr, "Parent is ");
+                describe_job(stderr, ctx, parent_id);
+
+                for (int isib = first_sibling_id; isib < first_mismatch; isib++)
+                {
+                    fprintf(stderr, "Sibling %d is ", isib - first_sibling_id);
+                    describe_job(stderr, ctx, isib);
+                }
+            }
+
+            errors = true;
+            break;
+        }
+
+        // Generate use another random number to re-seed the next iteration
+        // we re-seed every iteration so we have a 32-bit number we can print
+        // that describes the failing test.
+        cur_seed = rng.random();
+        total_num_jobs += job_count;
+
+        if (i % 100 == 99)
+        {
+            putchar('.');
+            fflush(stdout);
+        }
+    }
+
+    if (!errors)
+    {
+        printf("\nSUCCESS! JobSystemStress test completed %d runs correctly, %d jobs executed total.\n", num_tests_to_run, total_num_jobs);
+    }
+
+    delete ctx;
+}
diff --git a/Engine/Source/Programs/Horde/HordeStorage/Jupiter.Common/Shared/Oodle/2.9.3/linux/examples/make_example_input.h b/Engine/Source/Programs/Horde/HordeStorage/Jupiter.Common/Shared/Oodle/2.9.3/linux/examples/make_example_input.h
new file mode 100644
--- /dev/null
+++ b/Engine/Source/Programs/Horde/HordeStorage/Jupiter.Common/Shared/Oodle/2.9.3/linux/examples/make_example_input.h
@@ -0,0 +1,57 @@
+
+/*
+
+Oodle example :
+
+create "oodle_example_input_file" if not found
+
+*/
+
+static void make_example_input(const char * in_name)
+{
+    OO_SINTa size = 3*1024*1024;
+    
+    OodleXFileInfo info;
+    if ( OodleXIOQ_GetInfoByName_AsyncAndWait(in_name,&info,OodleFileNotFoundIsAnError_No) )
+    {
+        if ( info.size != size )
+        {
+            fprintf(stderr,"make_example_input (%s) : file exists but not expected size!?\n",in_name);          
+        }
+        return;
+    }
+    
+    OO_U8 * buffer = (OO_U8 *) OodleXMalloc_IOAligned(size);
+
+    OO_U64 state = 0x0102030405060708ULL;
+
+    OO_SINTa count = size/8;
+    OO_U64 * buf64 = (OO_U64 *)buffer;
+
+    for(OO_SINTa i=0;i<count;i++)
+    {
+        OO_SINTa rand = i * 2147001325 + 715136305;
+        rand = 0x31415926 ^ ((rand >> 16) + (rand << 16));
+
+        OO_S32 r = (((OO_S32)rand) % 257); // LZ compress 3145728 -> 1539959
+        //r = (r*r)/4;
+        
+        state += r;
+
+        if ( (i&31) == 0 ) state = 0x0102030405060708ULL * (i>>5);
+
+        state ^= ((state >> (rand & 31)) & 0x0F0F0F);
+
+        if ( (i&3) == 0 ) state &= 0x003FFFFFFFFFFFFFULL;
+        //if ( (i&3) == 0 ) state &= 0xFFFFFFFFFFFF0000ULL;
+
+        buf64[i] = state;
+    }
+        
+    OodleXHandle h = OodleXIOQ_OpenWriteWholeFileClose_Async(in_name,buffer,size);
+    // can't do anything funny here with leaving the write pending
+    // because some of the examples read using stdio , so we must flush
+    OodleX_Wait(h,OodleXHandleDeleteIfDone_Yes);
+    
+    OodleXFree_IOAligned(buffer);
+}
diff --git a/Engine/Source/Programs/Horde/HordeStorage/Jupiter.Common/Shared/Oodle/2.9.3/linux/examples/ooex.h b/Engine/Source/Programs/Horde/HordeStorage/Jupiter.Common/Shared/Oodle/2.9.3/linux/examples/ooex.h
new file mode 100644
--- /dev/null
+++ b/Engine/Source/Programs/Horde/HordeStorage/Jupiter.Common/Shared/Oodle/2.9.3/linux/examples/ooex.h
@@ -0,0 +1,204 @@
+#ifndef __OOEX_H__
+#define __OOEX_H__
+
+/*******
+
+OOEX - some helpers for me to use in the examples
+
+*******/
+
+#ifdef _MSC_VER
+#pragma warning(disable: 4127) // conditional is constant
+#pragma warning(disable: 4996) // secure deprecated
+#endif
+
+// include oodle2x.h before me if you want it
+//  if not, I stub in simple implementations here :
+#ifndef __OODLE2X_H_INCLUDED__
+#define OodleX_DisplayAssertion(file,line,func,str)  ( fprintf(stderr,"ASSERT: %s (%d) : %s\n",file,line,str) , 1 )
+#define OodleXLog_Printf_v1 printf
+#define OodleXLog_Printf_v0 printf
+#endif
+
+//===========================================================
+
+#ifndef OOEX_BREAK
+#ifdef _MSC_VER
+#define OOEX_BREAK()    __debugbreak()
+#elif defined(__GNUC__)
+#define OOEX_BREAK()    __builtin_trap()
+#else
+// crappy fallback BREAK
+#include <stdlib.h>
+#define OOEX_BREAK()    assert(0)
+#endif
+#endif
+
+//===========================================================
+
+// Minimalist platform detection for desktop platforms
+
+#ifdef ANDROID
+    #define OOEX_PLATFORM_ANDROID
+#endif
+
+#if defined(__linux__) && !defined(ANDROID)
+    #define OOEX_PLATFORM_LINUX
+#endif
+
+#if defined(_WIN32) || defined(WINAPI_FAMILY)
+    #ifdef WINAPI_FAMILY
+        // If this is #defined, we might be in a Windows Store App. But
+        // VC++ by default #defines this to a symbolic name, not an integer
+        // value, and those names are defined in "winapifamily.h". So if
+        // WINAPI_FAMILY is #defined, #include the header so we can parse it.
+        #include <winapifamily.h>
+        #if WINAPI_FAMILY_PARTITION(WINAPI_PARTITION_DESKTOP)
+            #define OOEX_PLATFORM_NT
+        #endif
+    #else // assume Windows-style platform without WINAPI_FAMILY defined is an old desktop Windows SDK
+        #define OOEX_PLATFORM_NT
+    #endif
+#endif
+
+#if defined(__APPLE__)
+    #include "TargetConditionals.h"
+    // TARGET_OS_MAC is always set on Darwin-derived platforms, even for iOS-derived targets
+    #if TARGET_OS_SIMULATOR
+        #define OOEX_PLATFORM_IOS
+        #define OOEX_PLATFORM_IOS_SIM
+    #elif TARGET_OS_IPHONE
+        #define OOEX_PLATFORM_IOS
+    #else
+        #define OOEX_PLATFORM_MAC
+    #endif
+#endif
+
+//=================================================================
+// OOEX_ASSERT :
+
+// toggle by setting OOEX_DO_ASSERTS
+
+#ifndef OOEX_DO_ASSERTS
+
+#if (defined(_DEBUG) && !defined(NDEBUG)) || defined(ASSERT_IN_RELEASE)
+  #define OOEX_DO_ASSERTS 1
+#else
+  #define OOEX_DO_ASSERTS 0
+#endif
+
+#endif
+
+//-----------------------------------------------------------
+
+#if defined(__GNUG__) || defined(__GNUC__) || (defined(_MSC_VER) && _MSC_VER > 1200)
+  #define OOEX_FUNCTION_NAME __FUNCTION__
+#else
+  #define OOEX_FUNCTION_NAME 0
+
+  // __func__ is in the C99 standard
+#endif
+
+//-----------------------------------------------------------
+      
+#define OOEX_ASSERT_BREAK()   OOEX_BREAK()
+
+#define OOEX_ASSERT_FAILURE_ALWAYS(str)           do { if ( OodleX_DisplayAssertion(__FILE__,__LINE__,OOEX_FUNCTION_NAME,str) ) OOEX_ASSERT_BREAK(); } while(0)
+
+#define OOEX_ASSERT_ALWAYS(exp)      do { if ( ! (exp) ) { OOEX_ASSERT_FAILURE_ALWAYS(#exp); } } while(0)
+
+//-----------------------------------
+#if OOEX_DO_ASSERTS 
+
+#define OOEX_ASSERT(exp)           OOEX_ASSERT_ALWAYS(exp)
+// OOEX_DURING_ASSERT is to set up expressions or declare variables that are only used in asserts
+#define OOEX_DURING_ASSERT(exp)   exp
+
+#define OOEX_ASSERT_FAILURE(str)  OOEX_ASSERT_FAILURE_ALWAYS(str)
+
+#else // OOEX_DO_ASSERTS //-----------------------------------
+
+#define OOEX_ASSERT(exp)           (void)0
+
+#define OOEX_DURING_ASSERT(exp)    (void)0
+
+#define OOEX_ASSERT_FAILURE(str)   (void)0
+
+#endif // OOEX_DO_ASSERTS //-----------------------------------
+
+//===========================================================
+
+// Must be placed after variable declarations for code compiled as .c
+#if defined(_MSC_VER) && _MSC_VER >= 1600 // in 2010 aka 10.0 and later 
+#   define OOEX_UNUSED_VARIABLE(x) x
+#else
+#   define OOEX_UNUSED_VARIABLE(x) (void)(sizeof(x))
+#endif
+
+//===========================================================
+
+#ifndef OOEX_MIN
+#define OOEX_MIN(a,b)    ( (a) < (b) ? (a) : (b) )
+#endif
+
+#ifndef OOEX_MAX
+#define OOEX_MAX(a,b)    ( (a) > (b) ? (a) : (b) )
+#endif
+
+#ifndef OOEX_ABS
+#define OOEX_ABS(a)      ( ((a) < 0) ? -(a) : (a) )
+#endif
+
+#ifndef OOEX_CLAMP
+#define OOEX_CLAMP(val,lo,hi) OOEX_MAX( OOEX_MIN(val,hi), lo )
+#endif
+
+//===========================================================
+
+#define OOEX_STRINGIZE( L )         #L
+
+#define OOEX_STRING_JOIN    OO_STRING_JOIN
+
+//===========================================================
+
+#ifndef OOEX_GET32_LE
+
+// Little-Endian 32-bit integer read/write
+static inline OO_U32 ooex_get32_le(const OO_U8 * ptr)
+{
+    OO_U32 x = ptr[0];
+    x |= (OO_U32)ptr[1] << 8;
+    x |= (OO_U32)ptr[2] << 16;
+    x |= (OO_U32)ptr[3] << 24;
+    return x;
+}
+
+static inline void ooex_put32_le(OO_U8 * ptr, OO_U32 val)
+{
+    ptr[0] = (OO_U8)((val >>  0) & 0xff);
+    ptr[1] = (OO_U8)((val >>  8) & 0xff);
+    ptr[2] = (OO_U8)((val >> 16) & 0xff);
+    ptr[3] = (OO_U8)((val >> 24) & 0xff);
+}
+
+#define OOEX_GET32_LE(ptr) ooex_get32_le((const OO_U8 *)(ptr))
+#define OOEX_PUT32_LE(ptr,val) ooex_put32_le((OO_U8 *)(ptr), (val))
+
+#endif // OOEX_GET32_LE
+
+//===========================================================
+
+#ifdef _MSC_VER
+#define OOEX_64_FMT "%I64"
+#define OOEX_64_FMT_SIZE    "I64"
+#else
+#define OOEX_64_FMT "%ll"
+#define OOEX_64_FMT_SIZE    "ll"
+#endif
+
+#define OOEX_S64_FMT    OOEX_64_FMT "d"
+#define OOEX_U64_FMT    OOEX_64_FMT "u"
+
+//===========================================================
+
+#endif // __OOEX_H__
diff --git a/Engine/Source/Programs/Horde/HordeStorage/Jupiter.Common/Shared/Oodle/2.9.3/linux/examples/read_whole_file.h b/Engine/Source/Programs/Horde/HordeStorage/Jupiter.Common/Shared/Oodle/2.9.3/linux/examples/read_whole_file.h
new file mode 100644
--- /dev/null
+++ b/Engine/Source/Programs/Horde/HordeStorage/Jupiter.Common/Shared/Oodle/2.9.3/linux/examples/read_whole_file.h
@@ -0,0 +1,54 @@
+
+#include <sys/stat.h>
+
+static char * read_whole_file(const char *name,OO_SINTa * pLength);
+
+
+/*
+
+read_whole_file : helper to read a file with stdio
+
+*/
+
+static char * read_whole_file(const char *name,OO_SINTa * pLength)
+{
+    #ifdef _MSC_VER
+    struct __stat64 st;
+    if ( _stat64(name,&st) != 0 ) return NULL;
+
+    #else
+    // stat64 and all that in posix is deprecated, just use stat :
+    struct stat st;
+    if ( stat(name,&st) != 0 ) return NULL;
+    #endif
+    
+    OO_SINTa length = (OO_SINTa)st.st_size;
+    if ( (OO_S64)length != (OO_S64)st.st_size ) return NULL;
+    
+    if ( length <= 0 ) { return NULL; }
+    
+    #ifdef _MSC_VER
+    FILE * fp;
+    errno_t err = fopen_s(&fp, name, "rb");
+    if (err != 0)
+        return NULL;
+    #else
+    FILE * fp = fopen(name,"rb");
+    if ( ! fp )
+        return NULL;
+    #endif
+    
+    *pLength = length;  
+        
+    char * data = (char *) malloc( (size_t)length );
+
+    if ( data )
+    {
+        fread(data,1,(size_t)length,fp);
+    }
+        
+    fclose(fp);
+    
+    return data;
+}
+
