Moved Horde.Storage from git to p4

diff --git a/Engine/Source/Programs/Horde/HordeStorage/Jupiter.Common/Shared/Oodle/2.9.3/linux/include/oodle2x.h b/Engine/Source/Programs/Horde/HordeStorage/Jupiter.Common/Shared/Oodle/2.9.3/linux/include/oodle2x.h
new file mode 100644
--- /dev/null
+++ b/Engine/Source/Programs/Horde/HordeStorage/Jupiter.Common/Shared/Oodle/2.9.3/linux/include/oodle2x.h
@@ -0,0 +1,2621 @@
+
+//===================================================
+// Oodle2 Ext header
+// (C) Copyright 1994-2021 Epic Games Tools LLC
+//===================================================
+
+#ifndef __OODLE2X_H_INCLUDED__
+#define __OODLE2X_H_INCLUDED__
+
+#include "oodle2.h"
+
+#define OODLEX_MAX_PATH (256)   /* Occasionally used for paths held in structs or on the stack.
+*/
+
+typedef OO_U64  OodleXIOQFile;
+/* Opaque weak reference handle to an IOQ File */
+
+#define OODLEX_IO_MAX_ALIGNMENT (4096)  /* Oodle low level offsets and sizes are aligned to OODLEX_IO_MAX_ALIGNMENT 
+
+    Unbuffered IO (as in $OodleXAPI_IOQ) require alignment to OODLEX_IO_MAX_ALIGNMENT.
+    Pointers returned by $OodleXMalloc_IOAligned are so aligned.
+    You can also use the utility functions such as $OodleX_IOAlignUpS32 to align values.
+*/
+
+#define OODLEX_BUFFER_SIZE_DEFAULT  (-1)    /* Pass to functions that want a buffer size to indicate the default should be used.
+
+    The buffer size used comes from $OodleXConfigValues
+*/
+
+typedef OO_U64 OodleXHandle;
+/* Opaque weak reference to Oodle asynchronous objects 
+    
+    Any op which returns an OodleXHandle can be used in $OodleX_Wait or as a dependency for other ops.
+    
+    See $OodleXAPI_Handle
+*/
+
+typedef enum OodleXHandleEnum
+{
+    OodleXHandle_Null = 0
+} OodleXHandleEnum;
+
+typedef enum OodleXPriority
+{
+    OodleXPriority_Normal = 1, // default priority
+    OodleXPriority_Force32 = 0x40000000
+} OodleXPriority;
+/* Priority for async tasks.  DEPRECATED
+
+    Use OodleXPriority_Normal only.
+        
+    Async work is (on average) FIFO.
+*/
+
+//-----------------------------------------------------------------------
+// OODLEX_EXTENSION_KEY is a U32 for an extension for quick compares or switches
+//  lower case by convention
+//  does not include the "."
+
+// FOURCC LE and BE : if these four bytes were in memory and you did *((U32 *)) you would get :
+#define OODLEX_FOURCC_LE(a,b,c,d)   ( ((U32)(a)) | ((U32)(b)<<8) | ((U32)(c)<<16) | ((U32)(d)<<24) )
+#define OODLEX_FOURCC_BE(a,b,c,d)   ( ((U32)(a)<<24) | ((U32)(b)<<16) | ((U32)(c)<<8) | ((U32)(d)) )
+
+// OODLEX_MACRO_TOLOWER should be safe and is nice for constants
+#define OODLEX_MACRO_TOLOWER(c) ( ( (c) >= 'A' && (c) <= 'Z' ) ? ( (c) | 0x20 ) : (c) )
+
+#define OODLEX_EXTENSION_KEY4(a,b,c,d)  ((U32)( (OODLEX_MACRO_TOLOWER(a)<<24) + (OODLEX_MACRO_TOLOWER(b)<<16) + (OODLEX_MACRO_TOLOWER(c)<<8) + OODLEX_MACRO_TOLOWER(d) ))
+#define OODLEX_EXTENSION_KEY(a,b,c)     ((U32)( (OODLEX_MACRO_TOLOWER(a)<<16) + (OODLEX_MACRO_TOLOWER(b)<<8) + (OODLEX_MACRO_TOLOWER(c)) ))
+#define OODLEX_EXTENSION_KEY_NONE   (0) // no extension
+
+//===================================================================
+
+typedef void *  OodleXOSFile;
+
+typedef void *  OodleXOSFileListing;
+
+typedef OOSTRUCT OodleXFileInfo
+{
+    OO_U32  flags;      // logical or of $OODLEX_FILEINFO_FLAGS
+    OO_U32  pad;
+    OO_S64  size;       // file size; $OODLEX_FILE_SIZE_INVALID if unknown
+    OO_U64  modTime;    // modTime on different platforms doesn't necessarilly mean anything, but it should be comparable with integer < and == (on the same platform, not vs. other platforms)
+} OodleXFileInfo;
+
+// FileInfo flags :
+typedef enum OODLEX_FILEINFO_FLAGS
+{
+    OODLEX_FILEINFO_FLAG_DIR         = (1<<0),  // queried name is a directory
+    OODLEX_FILEINFO_FLAG_READONLY = (1<<1), // you do not have write permission for this file
+    OODLEX_FILEINFO_FLAG_HIDDEN  = (1<<2),  // file is marked hidden
+    OODLEX_FILEINFO_FLAG_SYMLINK  = (1<<3),  // file or dir is a symlink or reparse point
+    OODLEX_FILEINFO_FLAG_TEMPORARY = (1<<4),  // file is marked temporary
+    OODLEX_FILEINFO_FLAG_OFFLINE = (1<<5),
+    OODLEX_FILEINFO_FLAG_Force32 = 0x40000000
+} OODLEX_FILEINFO_FLAGS;
+/* Flags for $(OodleXFileInfo:flags) 
+*/
+
+#define OODLEX_FILEINFO_FLAG_INVALID    ((OO_U32)-1)  /* Invalid value for $(OodleXFileInfo:flags) */
+
+#define OODLEX_FILEINFO_MODTIME_INVALID ((OO_U64)-1)  /* Invalid value for $(OodleXFileInfo:modTime) */
+
+#define OODLEX_FILE_SIZE_INVALID    ((OO_S64)-1)    /* Unknown or failure retreiving file size */
+
+#define OODLEX_FILE_OPEN_NO_RESERVE_SIZE    (0) /* Pass for _reserveSize_ to OpenFile calls if you don't want it to reserve any space */
+
+
+#define OODLEX_FILE_CLOSE_NO_TRUNCATE_SIZE  (-1)    /* Pass for _truncateFileSize_ to $OodleXIOQ_CloseFile_Async if you don't want it to truncate */
+
+OOINLINEFUNC void OodleXFileInfo_Reset( OodleXFileInfo * info )
+{
+    info->flags = 0;
+    info->size = OODLEX_FILE_SIZE_INVALID;
+    info->modTime = 0;
+}
+
+//-------------------------------------------------------------------------------------
+typedef enum OodleXError
+{
+    OodleXError_Ok=0,           // no error
+    OodleXError_InvalidHandle=1,    // null pointer, not open file, etc
+    OodleXError_FileNotFound=2, // file not found
+    OodleXError_NoAccess=3,     // attrib or sharing violation
+    OodleXError_BadParameters=4,    // usually unaligned or out of bounds file pointers
+    OodleXError_Corrupt=5,      // scratch or ejected media, damaged bits
+    OodleXError_Alignment=6,        // wrong alignment
+    OodleXError_Malloc=7,       // alloc failed
+    OodleXError_Compressor=8,   // a compressor or decompressor failed
+    OodleXError_UnexpectedEOF=9, // eof where I needed data
+    OodleXError_PreviousAsyncFailed=10, // dependent async failed, so I can't run
+    OodleXError_Close=11,       // error in close or object deletion, so I can't get more info
+
+    OodleXError_Unknown,        // error that doesn't match any of the other enums
+    OodleXError_Count,
+    OodleXError_Force32 = 0x40000000
+} OodleXError;
+/* oodle error enum to get a platform independent simple error code */
+
+OO_COMPILER_ASSERT( sizeof(OodleXError) == 4 );
+
+//-------------------------------------------------------------------------------------
+
+typedef enum OodleXFileMode
+{
+    OodleXFileMode_Invalid  =0,  // file mode not set
+    OodleXFileMode_Read     =1,  // open existing, shared
+    OodleXFileMode_WriteCreate=2,  // open new (create/trunc), exclusive
+    OodleXFileMode_Write    =2,  // alias for OodleXFileMode_WriteCreate
+    OodleXFileMode_ReadWrite = OodleXFileMode_Read|OodleXFileMode_Write, // open existing or create if new, exclusive, read/write
+    
+    // not modes on their own, but flags that can be added :
+    //eOodleFile_Truncate   =4, // flag - create new
+    //eOodleFile_Buffered   =8  // flag - use OS buffers
+    OodleXFileMode_Force32 = 0x40000000
+} OodleXFileMode;
+/* FileMode used by OodleFile and such.
+Not all OodleFile types support OodleXFileMode_ReadWrite 
+*/
+
+OO_COMPILER_ASSERT( sizeof(OodleXFileMode) == 4 );
+
+//-------------------------------------------------------------------------------------
+
+typedef enum OodleXFileOpenFlags
+{
+    OodleXFileOpenFlags_Default = 0,        // use Oodle's default for this platform
+    OodleXFileOpenFlags_Buffered = 1,   // use an OS-buffered file
+    OodleXFileOpenFlags_NotBuffered = 2,    // use a non-OS-buffered file, when possible
+    OodleXFileOpenFlags_WriteCreateDontStomp = 4, // Open for WriteCreate by default stomps existing; this prevents it
+    // OodleXFileOpenFlags_Temporary ?
+    OodleXFileOpenFlags_Force32 = 0x40000000
+} OodleXFileOpenFlags;
+/* OodleXFileOpenFlags specify options when opening files
+
+    OodleXFileOpenFlags_Default lets Oodle select buffered or unbuffered based on the system and global
+    settings.
+
+    OodleXFileOpenFlags_Buffered files are guaranteed to work with unaligned IO.
+    
+    OodleXFileOpenFlags_Default and OodleXFileOpenFlags_NotBuffered files require aligned IO on some platforms.
+
+    Flags may be combined with logical OR
+*/
+
+typedef enum OodleXCopyFileFlags
+{
+    OodleXCopyFileFlags_Overwrite = 0,              // always overwrite existing
+    OodleXCopyFileFlags_DontOverwriteExisting = 1,  // never overwrite existing
+    OodleXCopyFileFlags_OverwriteOnlyIfNewer = 2,   // overwrite only if source modtime is >= dest modtime
+    OodleXCopyFileFlags_OverwriteOnlyIfDifferentSize = 4,   // overwrite if source size != dest size
+    OodleXCopyFileFlags_OverwriteOnlyIfNewerOrDifferentSize = 2|4,  // common combo of overwrite options
+    OodleXCopyFileFlags_Mask = 7,
+    OodleXCopyFileFlags_Default = 0,  // default action Overwrite
+    OodleXCopyFileFlags_Force32 = 0x40000000
+} OodleXCopyFileFlags;
+/* Flags for Oodle CopyFile operations.
+
+    Combine with logical OR.
+    
+*/
+
+typedef OO_U32 OodleXLosslessFilterCode;
+
+#define OODLEX_LOSSLESSFILTER_NONE      (0)
+#define OODLEX_LOSSLESSFILTER_HEURISTIC ((OodleXLosslessFilterCode)-1)
+
+// must match lf_handle_table
+typedef enum OodleXStatus
+{
+    OodleXStatus_Invalid = 0, // indicates that a handle is not a live object (possibly previously deleted)
+    OodleXStatus_Pending = 1, // handle is alive and pending
+    OodleXStatus_Done    = 2, // handle completed succesfully
+    OodleXStatus_Error   = 3, // handle completed in error state
+    OodleXStatus_Count   = 4,
+    OodleXStatus_Force32 = 0x40000000
+} OodleXStatus;
+/* OodleXStatus indicates the status of asynchronous weak reference handles.
+
+    The OodleXStatus generally increases in numeric value during its autoDelete.
+    Check status >= OodleXStatus_Done to test for completion (possibly error).
+    
+    \<PRE>
+    Not yet allocated : OodleXStatus_Invalid = 0
+    Fired off and still pending : OodleXStatus_Pending = 1
+    Completed (possibly in error) : OodleXStatus_Done = 2 or _Error = 3
+    \</PRE>
+*/
+
+OO_COMPILER_ASSERT( sizeof(OodleXStatus) == 4 );
+
+typedef enum OodleXHandleAutoDelete
+{
+    OodleXHandleAutoDelete_No  = 0, // (default) handle lifetime will be managed by the client
+    OodleXHandleAutoDelete_Yes = 1,  // handle will deleted itself when done
+    OodleXHandleAutoDelete_Force32 = 0x40000000
+} OodleXHandleAutoDelete;
+/* When you spawn an async task and get an OodleXHandle back to track the task, with a normal 
+OodleXHandleAutoDelete_No handle you have to ensure that the handle is deleted at some point 
+(typically by calling $OodleX_Wait with $OodleXHandleDeleteIfDone_Yes).
+
+Alternative you can make the handle self-deleting by creating it with the OodleXHandleAutoDelete_Yes option.
+In that case you can still inspect the handle status with $OodleX_GetStatus and $OodleX_Wait, but when the handle completes
+and deletes itself, you will get $OodleXStatus_Invalid.  You cannot detect Done vs. Error cases with an OodleXHandleAutoDelete_Yes
+handle.
+*/
+
+typedef enum OodleXHandleKickDelayed
+{
+    OodleXHandleKickDelayed_No  = 0,     // (default) run async immediately
+    OodleXHandleKickDelayed_Yes = 1,   // wait until manually kicked
+    OodleXHandleKickDelayed_Force32 = 0x40000000
+} OodleXHandleKickDelayed;
+/* Normally async tasks are run as soon as possible; sometimes when spawning many tasks, you might not
+ want to let the thread switch immediately, so it can be better to fire several tasks with OodleXHandle_KickDelayed
+ and then kick them all together.  ("kick" means activate worker threads to do the tasks) 
+*/
+
+typedef enum OodleXHandleDeleteIfDone
+{
+    OodleXHandleDeleteIfDone_No  = 0, // (default) do not delete the handle
+    OodleXHandleDeleteIfDone_Yes = 1,  // delete the handle if it's Status is Done or Error
+    OodleXHandleDeleteIfDone_Force32 = 0x40000000   
+} OodleXHandleDeleteIfDone;
+/* Pass OodleXHandleDeleteIfDone_Yes to handle status checks to delete the handle if it's done.
+    This is the main way to free an $OodleXHandle
+*/
+
+typedef struct OodleXFileOpsVTable OodleXFileOpsVTable;
+
+//===========================================
+
+typedef OOSTRUCT OodleXMallocVTable
+{
+    void * m_context;   // provided context pointer will be passed to the function pointers
+
+    void * (OODLE_CALLBACK *m_pMalloc)( void * context, OO_SINTa bytes ); // must return OODLE_MALLOC_MINIMUM_ALIGNMENT aligned memory
+    void * (OODLE_CALLBACK *m_pMallocAligned)( void * context, OO_SINTa bytes , OO_S32 alignment ); // alignment will always be power of two
+    void   (OODLE_CALLBACK *m_pFree)(void * context, void * ptr); // must be able to free pointers from m_pMalloc or m_pMallocAligned
+    void   (OODLE_CALLBACK *m_pFreeSized)(void * context, void * ptr, OO_SINTa bytes); // use size to make free faster
+    
+    OO_S32 m_bigAlignment;  // indicates the alignment provided by MallocBig ; must be a multiple of $OODLEX_IO_MAX_ALIGNMENT for OodleX
+    void * (OODLE_CALLBACK *m_pMallocBig)(void * context, OO_SINTa bytes); // must return memory aligned to m_bigAlignment
+    void   (OODLE_CALLBACK *m_pFreeBig)  (void * context, void * ptr); // free a pointer allocated by m_pMallocBig
+    
+    OO_BOOL (OODLE_CALLBACK *m_pValidatePointer)(void * context, void * ptr, OO_SINTa bytes); // check on an allocation
+    
+} OodleXMallocVTable;
+/*
+    Function pointer table used to install the OodleX memory allocation functions
+    
+    Use $OodleXMalloc_InstallVTable to register a vtable as the one you want OodleX to use.
+    More commonly let $OodleX_Init set one for you.
+*/
+
+//===========================================
+
+typedef enum OodleXMalloc_OS_Options
+{
+    OodleXMalloc_OS_Options_None = 0, // default
+    OodleXMalloc_OS_Options_GuardBig = 1, // guard page for big allocs
+    OodleXMalloc_OS_Options_GuardBoth = 2, // guard page for big and small allocs
+    OodleXMalloc_OS_Options_GuardFrees = 3, // GuardBoth + leak frees and make them NOACCESS
+    OodleXMalloc_OS_Options_Count = 4,
+    OodleXMalloc_OS_Options_Force32 = 0x40000000
+} OodleXMalloc_OS_Options;
+
+//-------------------------------------------------------------------------------------
+// generic operations on an OodleXHandle :
+
+
+OOFUNC1 OodleXStatus OOFUNC2 OodleX_GetStatus( OodleXHandle h , OodleXHandleDeleteIfDone deleteIfDone OODEFAULT(OodleXHandleDeleteIfDone_No));
+/* Get the Status of an async handle
+
+  $:h   OodleXHandle weak reference
+  $:deleteIfDone    if $OodleXHandleDeleteIfDone_Yes and handle is not pending, it is deleted
+  $:return          handle status
+  
+  This function does not block.  Returns OodleXStatus_Invalid if the handle was already deleted or does not exist.
+  Test status for done by checking >= OodleXStatus_Done, because that also includes Error.
+*/
+
+OOFUNC1 OodleXStatus OOFUNC2 OodleX_Wait(  OodleXHandle h , OodleXHandleDeleteIfDone deleteIfDone OODEFAULT(OodleXHandleDeleteIfDone_No));
+/* Block the calling thread until handle is not Pending
+
+  $:h   OodleXHandle weak reference
+  $:deleteIfDone    if $OodleXHandleDeleteIfDone_Yes, handle will be deleted
+  $:return          handle status
+  
+  Will not return OodleXStatus_Pending.
+  
+  OodleX_WaitNoDelete and OodleX_WaitAndDelete are macros that are provided as short-hands for OodleX_Wait.  They are :
+
+\<div class=prototype>
+\<PRE>  
+  #define OodleX_WaitNoDelete (h)    OodleX_Wait(h,OodleXHandleDeleteIfDone_No)
+  #define OodleX_WaitAndDelete(h)    OodleX_Wait(h,OodleXHandleDeleteIfDone_Yes)
+\</PRE>
+\</div>
+
+*/
+
+#define OodleX_WaitNoDelete(h)  OodleX_Wait(h,OodleXHandleDeleteIfDone_No)
+#define OodleX_WaitAndDelete(h) OodleX_Wait(h,OodleXHandleDeleteIfDone_Yes)
+
+OOFUNC1 OodleXStatus OOFUNC2 OodleX_WaitAll(const OodleXHandle * handles, OO_S32 count , OodleXHandleDeleteIfDone deleteIfDone OODEFAULT(OodleXHandleDeleteIfDone_No));
+/* Block the calling thread until none of the provided handles are Pending
+
+  $:handles         array of OodleXHandle weak reference
+  $:count           number of handles in array
+  $:deleteIfDone    if $OodleXHandleDeleteIfDone_Yes, all handle will be deleted
+  $:return          handle status
+
+    Blocks until *ALL* handles are done.
+    Returns OodleXStatus_Error if any of the handles in the array is done with status OodleXStatus_Error.
+*/
+
+OOFUNC1 void OOFUNC2 OodleX_WaitDoneAllPending();
+/* Block on all pending operations being completed
+
+    FlushAllAsync kills all parallelism and should generally only be used at shutdown or error handling.
+    
+    FlushAllAsync is only guaranteed to stop pending handles that were fired before this call starts.
+    If new operations are created by other threads (or by existing pending operations) they may still be
+    pending when this call returns.
+*/
+
+OOFUNC1 OodleXStatus OOFUNC2 OodleX_SetHandleAutoDelete(OodleXHandle h,OodleXHandleAutoDelete autoDelete);
+/* change handle lifetime management 
+
+  $:h   OodleXHandle weak reference
+  $:autoDelete  if OodleXHandleAutoDelete_Yes, the handle deletes itself when not pending
+  
+  Handles that are $OodleXHandleAutoDelete_No must be deleted or they will leak.  The normal way to
+  delete them is by calling $OodleX_Wait with $OodleXHandleDeleteIfDone_Yes .
+  
+  A handle that deletes itself when done will then report $OodleXStatus_Invalid to queries, because it no longer exists.
+  
+  If you change a handle to OodleXHandleAutoDelete_Yes and it is already done, this function
+  will delete it immediately, and the returned Status will not be $OodleXStatus_Pending.
+*/
+
+OOFUNC1 OO_U32 OOFUNC2 OodleX_GetAvailableAsyncSelect();
+/* Get the currently available async systems
+
+    $:return    a bitwise OR of $OodleXAsyncSelect flags
+
+  The $OodleXAsyncSelect_Wide bit is set if there is more than one runner available.
+*/
+
+
+//-------------------------------------------------------------------------------------
+
+//-------------------------------------------------------------------------------------
+
+// OodleXAsyncSelect_ bit flags to specify which async systems you would like to run on :
+
+// flags for how I'm allowed to run async :
+typedef enum OodleXAsyncSelect
+{
+    OodleXAsyncSelect_None   = 0,     // run synchronously
+    OodleXAsyncSelect_Workers = 0x100, // run async on the Workers worker threads
+    OodleXAsyncSelect_NoFlagsMask = 0xFFF, // mask for all ways to run async ops
+    OodleXAsyncSelect_Wide = 0x1000,   // flag : run async wide, use all possible runners
+    OodleXAsyncSelect_Full = 0xFFFF,    // full speed : just turn on all bits
+    OodleXAsyncSelect_All = OodleXAsyncSelect_Full, // synonym
+    OodleXAsyncSelect_Force32 = 0x40000000
+} OodleXAsyncSelect;
+/* OodleXAsyncSelect are bit masks that can be combined to form an async selector.
+
+    The async selector tells an async operation like $OodleXLZ_Decompress_Narrow_Async where it should run its decompress.
+    
+    OodleXAsyncSelect_Wide means break the task into many smaller pieces that can be run simultaneously, and consume all
+    available runners to make the task complete as quickly as possible.  If WIDE is not specified, then the default is "narrow",
+    that is run async but don't split the task for minimum latency.  Mainly used with OodleXAsyncSelect_Workers ; WIDE means
+    create several smaller Worklet, while narrow creates just one Worklet that does the whole task.
+        
+    OodleXAsyncSelect_Full provides the quickest completion of any one call, but perhaps more contention with other operations.
+*/
+
+//-------------------------------------------------------------------------------------
+
+// public but not documented :
+
+#define OODLEX_ASYNC_SYSTEM_SPECIAL     (0)
+#define OODLEX_ASYNC_SYSTEM_IOQUEUE     (1)
+#define OODLEX_ASYNC_SYSTEM_WORKMGR     (2)
+#define OODLEX_ASYNC_SYSTEM_COPROC      (3) // SPU,GPU,etc.
+#define OODLEX_ASYNC_SYSTEM_GROUP       (4) // OodleAsyncGroup 
+#define OODLEX_ASYNC_SYSTEM_EVENT       (5) // simple event or countdown
+#define OODLEX_ASYNC_SYSTEM_DATA        (7) // not an operation, but a piece of data like an OodleXIOQFile or a OodleIOQStream
+
+#define OODLEX_ASYNC_SYSTEM_SHIFT       (61) // 3 bits for system
+#define OODLEX_ASYNC_SYSTEM_MASK            (((OO_U64)7)<<OODLEX_ASYNC_SYSTEM_SHIFT)
+
+//#define OODLEX_ASYNC_GUID(system,counter)     (((OodleXHandle)(system)<<OODLEX_ASYNC_SYSTEM_SHIFT)|(OodleXHandle)(counter))
+#define OODLEX_ASYNC_HANDLE_GET_SYSTEM(handle)      ((handle)>>OODLEX_ASYNC_SYSTEM_SHIFT)
+#define OODLEX_ASYNC_HANDLE_REMOVE_SYSTEM(handle)   ((handle)&(~OODLEX_ASYNC_SYSTEM_MASK))
+
+//-------------------------------------------------------------------------------------
+// Special handle values that always return the same status :
+
+//#define OODLEX_ASYNC_HANDLE_INVALID       OODLEX_ASYNC_GUID(OODLEX_ASYNC_SYSTEM_SPECIAL, 0 )
+
+#define OODLEX_ASYNC_HANDLE_INVALID     ((OodleXHandle)0) /* OodleXHandle for an invalid handle.
+Calls to $OodleX_GetStatus on this handle value will return &OodleXStatus_Invalid.
+*/
+
+#define OODLEX_ASYNC_HANDLE_PENDING     ((OodleXHandle)0x0000000000000001ULL) /* OodleXHandle to a special always-pending handle.  This is for Oodle internal use only.
+Calls to $OodleX_GetStatus on this handle value will return &OodleXStatus_Pending.
+This is designed for use with OodleAsyncGroup.  See OodleAsyncGroup_ChangePending.
+Calling $OodleX_Wait on this handle is a deadlock.
+This handle must not be deleted!  Do not call $OodleX_Wait on it with deleteIfDone = true.
+*/
+
+#define OODLEX_ASYNC_HANDLE_DONE            ((OodleXHandle)0x0000000100000001ULL) /* OodleXHandle to a special always-done handle.
+Calls to $OodleX_GetStatus on this handle value will return &OodleXStatus_Done.
+This handle must not be deleted!  Do not call $OodleX_Wait on it with deleteIfDone = true.
+*/
+
+#define OODLEX_ASYNC_HANDLE_ERROR       ((OodleXHandle)0x0000000200000001ULL) /* OodleXHandle to a special always-error handle.
+Calls to $OodleX_GetStatus on this handle value will return &OodleXStatus_Error.
+This handle must not be deleted!  Do not call $OodleX_Wait on it with deleteIfDone = true.
+*/
+
+//-------------------------------------------------------------------------------------
+
+OOFUNC1 OO_U32 OOFUNC2 OodleX_GetExtensionKey(const char * filename);
+
+OOFUNC1 OO_U32 OOFUNC2 OodleX_MakeExtensionKey(const char * extension);
+
+//--------------------------------------------------------------------------
+
+
+//-------------------------------------------------------------------------------------
+// OodleInit
+
+typedef OOSTRUCT OodleXInitOptions
+{
+    //---------------------------------
+    // Phase 1:
+    
+    const OodleXMallocVTable * m_pBaseVTable;  // vtable for OodleMalloc to use [OodleXMalloc_GetVTable_OS]
+    OO_BOOL m_OodleInit_DebugAllocator; // option : put a debug allocator layer on top of m_pBaseVTable [false]
+
+    OO_S32 m_num_handles_log2; // log2 number of the number of handles for the OodleXHandleTable [13]
+
+    OO_BOOL m_OodleInit_ThreadLog; // option : enable the ThreadLog [true]
+    OO_BOOL m_OodleInit_Log;  // option : enable the Log [true]
+    OO_BOOL m_OodleInit_Log_Header;  // option : write a header to the Log at startup [true]
+    const char * m_OodleInit_Log_FileName; // set the log file name (NULL for default, which is described in $Oodle_About_Platforms)
+    OO_BOOL m_OodleInit_Log_FlushEachWrite; // option : flush the log after each write [false]
+    
+    //-----------------------------------
+    // Phase 2:
+
+    OO_BOOL m_OodleInit_BreakOnLogError; // option : debug break when Oodle logs an error
+    OO_BOOL m_OodleInit_Telemetry;      // option : make the Telemetry connection for tracking Oodle [false]
+    void * m_OodleInit_Telemetry_Context; // the telemetry context, NULL means I will make it (if m_OodleInit_Telemetry is true)
+    OO_BOOL m_OodleInit_StackTrace;     // option : enable stack tracing in Oodle [true]
+    OO_BOOL m_OodleInit_LeakTrack;      // option : enable LeakTrack in Oodle
+    OO_BOOL m_OodleInit_SimpleProf;     // option : enable simple profiler (this is mainly for me)
+    OO_BOOL m_OodleInit_FuzzTest;       // deprecated , does nothing
+    t_OodleFPVoidVoid * m_OodleInit_ThreadProfiler_funcptr; // option : enable the thread profiler
+    OO_BOOL m_OodleInit_IOQ;                // option : enable the IOQ
+    OO_BOOL m_OodleInit_IOQ_Log;            // option : enable logging operations on the IOQ
+    OO_BOOL m_OodleInit_IOQ_BreakOnError;   // option : make the IOQ issue a debug break on any error (for debugging)
+    OO_BOOL m_OodleInit_IOQ_Threaded;       // option : enable threading on IOQ (turn off for debugging)
+    OO_BOOL m_OodleInit_IOQ_CheckAlignment;  // option : should IOQ check alignment of parameters?
+    //OO_BOOL m_OodleInit_SPU;              // option : enable Oodle on the SPU (does not load the SPU ELF, you do that later)
+    //OO_BOOL m_OodleInit_SPU_Synchronous; // option : for debugging, run SPU ops synchronously
+    OO_BOOL m_OodleInit_Workers;            // option : enable the worker thread system
+    OO_S32 m_OodleInit_Workers_Count;       // number of worker threads to start (default is $OODLE_WORKERS_COUNT_ALL_PHYSICAL_CORES)
+
+} OodleXInitOptions;
+/* Options struct for $OodleX_Init
+    can be filled with $OodleX_Init_GetDefaults
+*/
+
+#define OODLE_WORKERS_COUNT_ALL_PHYSICAL_CORES (-1) /* Make workers for every physical core
+    eg. in a 6-physical core, 12-hyper-thread system, would make 6 threads
+    this is usually best for Oodle Data LZ compression work
+   See also $OODLE_WORKERS_COUNT_ALL_HYPER_CORES
+*/
+
+#define OODLE_WORKERS_COUNT_ALL_HYPER_CORES (-2) /* Make workers for every hyper-thread
+    eg. in a 6-physical core, 12-hyper-thread system, would make 12 threads
+    this is usually best for Oodle Texture work
+   See also $OODLE_WORKERS_COUNT_ALL_PHYSICAL_CORES
+*/
+
+OOFUNC1 t_OodleFPVoidVoid * OOFUNC2 OodleX_Init_ThreadProfilerInit(void);
+/* Get the func pointer for m_OodleInit_ThreadProfiler_funcptr
+*/
+
+typedef enum OodleX_Init_GetDefaults_DebugSystems
+{
+    OodleX_Init_GetDefaults_DebugSystems_No = 0,
+    OodleX_Init_GetDefaults_DebugSystems_Yes = 1,
+    OodleX_Init_GetDefaults_DebugSystems_Force32 = 0x40000000
+} OodleX_Init_GetDefaults_DebugSystems;
+/* Should GetDefaults enable debugging systems?
+*/
+
+typedef enum OodleX_Init_GetDefaults_Threads
+{
+    OodleX_Init_GetDefaults_Threads_No = 0,
+    OodleX_Init_GetDefaults_Threads_Yes = 1,
+    OodleX_Init_GetDefaults_Threads_Force32 = 0x40000000
+} OodleX_Init_GetDefaults_Threads;
+/* Should GetDefaults enable any threads?
+*/
+
+OOFUNC1 OO_BOOL OOFUNC2 OodleX_Init_GetDefaults(OO_U32 oodle_header_version,OodleXInitOptions * pOptions,
+    OodleX_Init_GetDefaults_DebugSystems debugSystems OODEFAULT(OodleX_Init_GetDefaults_DebugSystems_Yes),
+    OodleX_Init_GetDefaults_Threads threads OODEFAULT(OodleX_Init_GetDefaults_Threads_Yes));
+/* Get defaults for $OodleXInitOptions
+
+    $:oodle_header_version  pass $OODLE_HEADER_VERSION here
+    $:pOptions  filled with default $OodleXInitOptions
+    $:debugSystems  should OodleX_Init enable any debug systems (leaktrack, log, etc) ?
+    $:threads       should OodleX_Init start any threads?
+    $:return    false if $OODLE_HEADER_VERSION is not compatible with this lib
+    
+    The _debugSystems_ and _threads_ options are just easy ways of getting pOptions filled out for common
+    use cases.  For fine control of individual settings, you can always set the values in OodleXInitOptions yourself.
+
+    NOTE : do not use this if you want minimal linkage.  See $OodleX_Init_GetDefaults_Minimal.
+*/
+
+OOFUNC1 OO_BOOL OOFUNC2 OodleX_Init_GetDefaults_Minimal(OO_U32 oodle_header_version,OodleXInitOptions * pOptions);
+/*  Get minimal defaults for $OodleXInitOptions , enabling only necessary Oodle systems
+
+    $:oodle_header_version  pass $OODLE_HEADER_VERSION here
+    $:pOptions  filled with default $OodleXInitOptions
+    $:return    false if $OODLE_HEADER_VERSION is not compatible with this lib
+    
+    Fill options such that a minimal part of the Oodle library is imported.
+    
+    All memory->memory compressors will work.
+    
+    IO and Threading will be disabled.
+
+    Can be used with $OodleX_Init_NoThreads or $OodleX_Init
+*/  
+    
+OOFUNC1 OO_BOOL OOFUNC2 OodleX_Init(OO_U32 oodle_header_version,const OodleXInitOptions * pOptions);
+/* Initialize Oodle
+
+    $:oodle_header_version  pass $OODLE_HEADER_VERSION here
+    $:pOptions              options for Init; must not be NULL; use $OodleX_Init_Default if you don't want to set up options
+    $:return                false if $OODLE_HEADER_VERSION is not compatible with this lib
+    
+    You must call $OodleX_Init or $OodleX_Init_NoThreads before any other Oodle function that you expect to work.
+
+    Pair with $OodleX_Shutdown.
+    
+    For minimal linkage, use $OodleX_Init_NoThreads
+*/
+
+OOFUNC1 OO_BOOL OOFUNC2 OodleX_Init_Default(OO_U32 oodle_header_version,
+    OodleX_Init_GetDefaults_DebugSystems debugSystems OODEFAULT(OodleX_Init_GetDefaults_DebugSystems_Yes),
+    OodleX_Init_GetDefaults_Threads threads OODEFAULT(OodleX_Init_GetDefaults_Threads_Yes));
+/* Initialize Oodle, without options struct
+
+    $:oodle_header_version  pass $OODLE_HEADER_VERSION here
+    $:debugSystems  should OodleX_Init enable any debug systems (leaktrack, log, etc) ?
+    $:threads       should OodleX_Init start any threads?
+    $:return    false if $OODLE_HEADER_VERSION is not compatible with this lib
+    
+    The _debugSystems_ and _threads_ options are just easy ways of getting pOptions filled out for common
+    use cases.  For fine control of individual settings, you can always set the values in OodleXInitOptions yourself.
+
+    This is just a shortcut to $OodleX_Init_GetDefaults then $OodleX_Init
+
+    NOTE : do not use this if you want minimal linkage. 
+*/
+    
+OOFUNC1 void OOFUNC2 OodleX_LogSystemInfo();
+/* Log some info about the platform
+
+    This function should be called after $OodleX_Init.
+
+    It prints some info to the Oodle Log about the Oodle build and your system.
+    This is a helpful thing to include in debug reports sent to RAD.
+*/
+
+// pub but no doc :
+OOFUNC1 OO_BOOL OOFUNC2 OodleX_Init_Phase1(OO_U32 oodle_header_version,const OodleXInitOptions * pOptions);
+OOFUNC1 void OOFUNC2 OodleX_Init_Phase2(const OodleXInitOptions * pOptions);
+
+typedef enum OodleX_Shutdown_LogLeaks
+{
+    OodleX_Shutdown_LogLeaks_No = 0,
+    OodleX_Shutdown_LogLeaks_Yes= 1,
+    OodleX_Shutdown_LogLeaks_Force32 = 0x40000000
+} OodleX_Shutdown_LogLeaks;
+/* bool enum
+*/
+
+typedef enum OodleX_Shutdown_DebugBreakOnLeaks
+{
+    OodleX_Shutdown_DebugBreakOnLeaks_No = 0,
+    OodleX_Shutdown_DebugBreakOnLeaks_Yes= 1,
+    OodleX_Shutdown_DebugBreakOnLeaks_Force32 = 0x40000000
+} OodleX_Shutdown_DebugBreakOnLeaks;
+/* bool enum
+*/
+
+// OodleX_Shutdown will release all singletons in the right order
+//  if threadProfileLogName it will write the thread log after killing the threads
+OOFUNC1 void OOFUNC2 OodleX_Shutdown(const char * threadProfileLogName OODEFAULT(NULL), 
+                            OodleX_Shutdown_LogLeaks logLeaks OODEFAULT(OodleX_Shutdown_LogLeaks_Yes),
+                            OO_U64 allocStartCounter OODEFAULT(0),
+                            OodleX_Shutdown_DebugBreakOnLeaks debugBreakOnLeaks OODEFAULT(OodleX_Shutdown_DebugBreakOnLeaks_No));
+/* Shut down Oodle at app exit time.
+
+    $:threadProfileLogName  (optional) if not NULL, and the ThreadProfiler is enabled, writes the threadprofiler output to this file name
+    $:logLeaks              (optional) if true and the LeakTracker is enabled, logs any leaks or memory or handles
+    $:allocStartCounter     (optional) initial counter for the LeakTrack log
+    $:debugBreakOnLeaks     (optional) if there are any leaks, do a debug break
+    
+    Pair with $OodleX_Init.  No Oodle functions should be called after Shutdown.
+    
+    Call Shutdown from the same thread that called Init.
+    
+    Do not shutdown Oodle then init again.  Only call Init and Shutdown once per run.
+*/
+
+OOFUNC1 OO_BOOL OOFUNC2 OodleX_Init_NoThreads(OO_U32 oodle_header_version,const OodleXInitOptions * pOptions);
+/* Initialize Oodle with no threads and minimal systems
+
+    $:oodle_header_version  pass $OODLE_HEADER_VERSION here
+    $:pOptions              options for Init; must not be NULL; use $OodleX_Init_GetDefaults_Minimal to fill out
+    $:return                false if $OODLE_HEADER_VERSION is not compatible with this lib
+    
+    You must call $OodleX_Init or $OodleX_Init_NoThreads before any other Oodle function that you expect to work.
+
+    Pair with $OodleX_Shutdown_NoThreads.
+
+    This function does not enable the Oodle IOQ or WorkMgr.  No async jobs or IO will work.
+    
+    All memory->memory compressors will work.
+
+    Pair with $OodleX_Shutdown_NoThreads
+*/
+
+OOFUNC1 void OOFUNC2 OodleX_Shutdown_NoThreads(const char * threadProfileLogName OODEFAULT(NULL), 
+                            OodleX_Shutdown_LogLeaks logLeaks OODEFAULT(OodleX_Shutdown_LogLeaks_Yes),
+                            OO_U64 allocStartCounter OODEFAULT(0),
+                            OodleX_Shutdown_DebugBreakOnLeaks debugBreakOnLeaks OODEFAULT(OodleX_Shutdown_DebugBreakOnLeaks_No));
+/* Shut down Oodle at app exit time.
+
+    $:threadProfileLogName  (optional) if not NULL, and the ThreadProfiler is enabled, writes the threadprofiler output to this file name
+    $:logLeaks              (optional) if true and the LeakTracker is enabled, logs any leaks or memory or handles
+    $:allocStartCounter     (optional) initial counter for the LeakTrack log
+    $:debugBreakOnLeaks     (optional) if there are any leaks, do a debug break
+    
+    Pair with $OodleX_Init_NoThreads.  No Oodle functions should be called after Shutdown.
+    
+    Call Shutdown from the same thread that called Init.
+    
+    Do not shutdown Oodle then init again.  Only call Init and Shutdown once per run.
+*/
+
+//-------------------------------------------------------------------------------------
+
+//---------------------------------------------------------------------------------
+// OodleIOQ :
+//  low level async IO
+//  every request creates an async handle
+//  handles can be autodelete or manual delete
+//  if you pass in a NULL handle pointer, an auto delete handle will be created
+//  all requests on the same file are done in the order requested
+//      that is, there is an implicit order of handles on each file
+//  errors are per request but also aggregated on the file
+//  so you can check the status of an auto-delete handle by looking at the file
+//  note : ioq requests generally must be at OODLEX_IO_MAX_ALIGNMENT
+//      both size and position
+
+// BAD USER : FlushAllPending is a big hard block ; do not use this !
+OOFUNC1 void OOFUNC2 OodleXIOQ_WaitDoneAllPending();
+/* Block the calling thread until all pending IOQ operations are complete
+
+    Should generally only be used for errors or shutdown.
+    $OodleX_WaitDoneAllPending does this and more.
+*/
+
+// KickOff : after you make requests they may not get going until a thread switch fires them
+//  this manually makes sure the IO thread gets going immediately if possible
+OOFUNC1 void OOFUNC2 OodleXIOQ_KickAnyDelayed();
+/* Fire any requests which have not previously been started
+
+    If requests were enqueued with kick = false (don't start immediately), then they can be
+    started this way.  Disabling auto-kick is good for performance when a very large number of
+    request are being created in a short period of time.
+*/
+
+// returns old state
+// no no not allowed
+// OOFUNC1 OodleXHandleKickDelayed OOFUNC2 OodleXIOQ_SetKickMode( OodleXHandleKickDelayed kickMode );
+/* Set the kick mode
+
+    $:kickMode          new kick mode to set
+    $:return            the previous kick state
+
+    If kick is OodleXHandleKickDelayed_No (default), then each request is fired as it is made.  If not, then requests are
+    queued up until $OodleXIOQ_KickAnyDelayed is called.  Disabling immediate kick is good for performance when a very large number of
+    request are being created in a short period of time.
+    
+    NOTE : calling SetKickMode(Immediate) does NOT immediately kick any delayed requests.
+    You should call $OodleXIOQ_KickAnyDelayed.
+*/
+
+// GetStatus on a request :
+//  this is also the main way that requests are deleted when you're done with them!
+OOFUNC1 OodleXStatus OOFUNC2 OodleXIOQ_GetStatus(OodleXHandle req,OodleXHandleDeleteIfDone andDeleteIfDone,OO_U32 * pErrorCode OODEFAULT(NULL), OO_S32 * pReturnValue OODEFAULT(NULL));
+/* Get the Status of a request, and optionally delete if done
+    
+    $:req       the IOQ operation handle to work on
+    $:andDeleteIfDone   if true and the returned status is >= Done the handle will be deleted
+    $:pErrorCode    (optional) the OS error code, if any
+    $:pReturnValue  (optional) the operation return value
+    $:return        the status of the request
+    
+    This function is similar to $OodleX_GetStatus, but for IOQ operation handles only, and it provides
+    more information (optionally).
+    
+    The error code returned can be processed with $OodleXIOQ_GetErrorEnum or $OodleXIOQ_GetErrorDetails.
+    
+    The return value depends on the operation type.  For example if the operation is a Read, it returns the
+    number of bytes successfully read.
+*/
+
+// OodleXIOQ_GetErrorDetails if you're told you had an error; pMessage can be null
+OOFUNC1 OO_BOOL OOFUNC2 OodleXIOQ_GetErrorDetails(OO_U32 code,OodleXIOQFile file,char * pMessage,int messageSize);
+/* Convert an OS error code into a text message
+
+    $:code  the error code, eg. from $OodleXIOQ_GetStatus
+    $:file  the file that the error occurred on (or 0 for unknown)
+    $:pMessage  pointer to a text buffer that will be filled out
+    $:messageSize   number of bytes in the text buffer
+    $:return    bool for success/failure
+
+    fills out pMessage with a text description of the error (if available).
+*/
+
+// map error to simple enum :
+OOFUNC1 OodleXError OOFUNC2 OodleXIOQ_GetErrorEnum(OO_U32 code,OodleXIOQFile file);
+/* Convert an OS error code into a text message
+
+    $:code  the error code, eg. from $OodleXIOQ_GetStatus
+    $:file  the file that the error occurred on (or 0 for unknown)
+    $:return    an OS-neutral $OodleXError
+
+    Converts an OS-specific error code into a platform agnostic error enum.  Useful for
+    recognizing common error cases like OodleXError_FileNotFound.  Any unusual or platform-specific
+    codes will return OodleXError_Unknown.
+*/
+
+OOFUNC1 void OOFUNC2 OodleXIOQ_LogError(OO_U32 code,OodleXIOQFile file,const char * pName OODEFAULT(NULL));
+/* Logs an OS error code with a detailed text message
+
+    $:code  the error code, eg. from $OodleXIOQ_GetStatus
+    $:file  the file that the error occurred on (or 0 for unknown)
+    $:pName (optional) a tag to log with the error
+    
+    Calls $OodleXLog_Printf to output a detailed error, as created by $OodleXIOQ_GetErrorDetails.
+*/
+
+// stall the thread until request is done
+//  use this for example when deleting the buffers you did a Read or Write with
+//  try to avoid blocking because it ruins asynchronicity
+// NOTE : generally you should call OodleX_Wait() instead
+//  use this just to get error code
+OOFUNC1 OodleXStatus OOFUNC2 OodleXIOQ_Wait(OodleXHandle req,OodleXHandleDeleteIfDone andDelete,OO_U32 * pErrorCode OODEFAULT(NULL));
+/* Block the calling thread until request is not pending
+
+    $:req           the IOQ operation handle to work on
+    $:andDelete     if true, delete the request
+    $:pErrorCode    (optional) filled with the os error code, if any
+    $:return        the status
+
+    The status returned will not be OodleAsync_Pending.
+    Similar to $OodleX_Wait , but only works on IOQ requests, and can return the IOQ error code.
+    Generally you should just call OodleX_Wait in most cases.
+*/
+
+// OodleXIOQ_GetInfo will fail & return false if the open request is not yet done
+OOFUNC1 OO_BOOL OOFUNC2 OodleXIOQ_GetInfo(OodleXIOQFile file,OodleXFileInfo * pInto, OO_S32 * pAlignmentRequired OODEFAULT(NULL));
+/* Get Info about a file
+
+    $:file  the IOQFile to query
+    $:pInto filled with $OodleXFileInfo
+    $:pAlignmentRequired    (optional) filled with alignment required
+    $:return true if successful; if GetInfo returns false, pInto and pAlignmentRequired are untouched.
+    
+    If the file is not yet open, GetInfo will fail and return false.  eg. if $OodleXIOQ_OpenForRead_Async has been done
+    but the request is still pending.
+    
+    If the file size can not be queried it is set to $OODLEX_FILE_SIZE_INVALID.
+
+    If _pAlignmentRequired_ is given, it is filled with the alignment required to use this file.
+    $OODLEX_IO_MAX_ALIGNMENT is guaranteed to always be okay, so if you align to that then you are fine.
+    See $OodleXIOQ_About for more about alignment.
+
+*/
+
+// OodleXIOQ_GetInfo will fail & return false if the open request is not yet done
+OOFUNC1 OO_BOOL OOFUNC2 OodleXIOQ_Wait_GetInfo(OodleXIOQFile file,OodleXFileInfo * pInto, OO_S32 * pAlignmentRequired OODEFAULT(NULL));
+/* Get Info about a file ; if the file is not open yet, wait for it
+
+    $:file  the IOQFile to query
+    $:pInto filled with $OodleXFileInfo
+    $:pAlignmentRequired    (optional) filled with alignment required
+    $:return true if successful; if GetInfo returns false, pInto and pAlignmentRequired are untouched.
+        
+    This function is like $OodleXIOQ_GetInfo , but will not return false if the Open operation is still pending; instead it
+    will block the calling thread until the Open is done so that info is available.
+
+*/
+
+OOFUNC1 OodleXHandle OOFUNC2 OodleXIOQ_GetLastPendingOpOnFile(OodleXIOQFile file);
+/* Get an operation on this file, if any
+
+    $:file  the IOQFile to query
+    $:return the operation found, or 0 if none
+        
+    The operation returned may no longer be pending (nor the last) by the time you check it.
+*/
+
+OOFUNC1 OO_BOOL OOFUNC2 OodleXIOQ_GetName(OodleXIOQFile file,char * pInto,OO_S32 intoSize);
+/* Get the file name (OS name)
+
+    $:file  the IOQFile to query
+    $:pInto filled with the file's OS name (UTF8)
+    $:intoSize  number of bytes Oodle can write to pInto
+    $:return true if successful
+
+    Copies the OS name (UTF8) into pInto.  This may not be the same as the name used when opening
+    the file, if that was a VFS name.
+*/
+
+// error on file is last error that occured from any request on that file
+OOFUNC1 OO_U32 OOFUNC2 OodleXIOQ_GetLastError(OodleXIOQFile file);
+/* Get the last error on a file
+
+    $:file  the IOQFile to query
+    $:return    the last error on the file (0 for none)
+
+    IO operation errors are tracked on the file to simplify error tracking.
+    Individual operation errors can be queried with $OodleXIOQ_GetStatus.
+    The error code returned can be processed with $OodleXIOQ_GetErrorEnum or $OodleXIOQ_GetErrorDetails.
+*/
+
+OOFUNC1 void OOFUNC2 OodleXIOQ_ClearError(OodleXIOQFile file);
+/* Clear any errors on the file
+
+    $:file  the IOQFile to query
+    
+    Wipe out any previous errors recorded on the file, so that $OodleXIOQ_GetLastError now returns zero.
+*/
+
+ // combination of OodleXIOQ_GetLastError and LogError
+OOFUNC1 OO_BOOL OOFUNC2 OodleXIOQ_LogLastError(OodleXIOQFile file);
+/* Log the last error on a file
+
+    $:file  the IOQFile to query
+    $:return true if any error was logged
+    
+    Calls $OodleXIOQ_GetLastError and $OodleXIOQ_LogLastError
+*/
+
+// GetHandle gives you the CURRENT HANDLE if the open is pending it may not be set yet !
+OOFUNC1 void * OOFUNC2 OodleXIOQ_GetOSHandle(OodleXIOQFile file);
+/* Get the OS file handle for this OodleXIOQFile
+
+    $:file  the IOQFile to query
+    $:return the OS file handle
+    
+    If the file is not yet open (eg. $OodleXIOQ_OpenForRead_Async was started but is still pending), this
+    returns NULL.
+*/
+
+// Set the VTable for the file :
+//  returns the old vtable ; warning : vtables are not mutex protected !
+//  WARNING : changing the file's VTable while there are ops on that file in the Queue has undefined results !!
+OOFUNC1 const OodleXFileOpsVTable * OOFUNC2 OodleXIOQ_SetVTable(OodleXIOQFile file,const OodleXFileOpsVTable * vtable);
+/* Set the VTable used for ops the file.
+
+    $:file  the IOQFile to query
+    $:return the previous vtable
+    
+    Change the VTable used for ops the file after opening.  This is discouraged, generally try to set the right
+    vtable in the $OodleXIOQ_OpenForRead_Async call and then don't change it.
+    Warning : vtables are not themselves internally mutex protected !
+    WARNING : changing the file's VTable while there are ops on that file in the Queue has undefined results !!
+
+*/
+
+//---------------------------------------------------------------------------------
+// IO calls :
+//  these are all queued and not executed immediately !
+//  any pointers you pass in must be kept live until the request is done!
+//  check the OodleXHandle status to see when they're done
+//
+// if vtable is null, the global default vtable is used
+//  WARNING : vtables are not protected from thread access with mutexes, they are assumed to be const
+
+// fence is a queue'd NOP operation that you can wait on
+//  (fileRef is optional)
+OOFUNC1 OodleXHandle OOFUNC2 OodleXIOQ_Fence_Async(OodleXIOQFile fileRef OODEFAULT(0),OodleXHandleAutoDelete autoDelete OODEFAULT(OodleXHandleAutoDelete_No),OodleXPriority priority OODEFAULT(OodleXPriority_Normal),const OodleXHandle * dependencies OODEFAULT(NULL),OO_S32 numDependencies OODEFAULT(0));
+/* Add a "fence" to the operation queue
+
+    $:fileRef           (optional) the file to associate the request with
+    $:autoDelete        (optional) lifetime of the operation handle ; see $OodleXHandleAutoDelete
+    $:priority          (optional) priority of the operation ; see $OodleXPriority
+    $:dependencies      (optional) dependencies; the async op won't start until these are all complete; note : these are not freed, they must be autodelete or you must free them some other way.
+    $:numDependencies   (optional)  number of handles in _deps_ array
+    $:return    handle to the operation, or 0 if it could not be started (usually due to invalid args)
+
+    A fence is a NOP which can be used to schedule against other operations.
+    eg. if you have an OodleXIOQFile and want to block on any (unknown) operations on that file completing,
+    you can add a Fence op to the file and block on it; earlier requests will flush first, so when the fence is
+    done you know all previous requests are done.
+*/
+
+// Opens return a File ref right away for your convenience, but the file is not actually open for a little while
+//  you can however go ahead and queue more requests on that file reference
+// OpenForRead is always shared access
+// name is a VFS name and is mapped automatically
+
+OOFUNC1 OodleXHandle OOFUNC2 OodleXIOQ_OpenForRead_Async(OodleXIOQFile * pFile,const char *name, 
+                        OodleXFileOpenFlags fileOpenFlags OODEFAULT(OodleXFileOpenFlags_Default),
+                        const OodleXFileOpsVTable * vtable OODEFAULT(NULL),OodleXHandleAutoDelete autoDelete OODEFAULT(OodleXHandleAutoDelete_No),OodleXPriority priority OODEFAULT(OodleXPriority_Normal),const OodleXHandle * dependencies OODEFAULT(NULL),OO_S32 numDependencies OODEFAULT(0));
+/* Start opening a file for read
+
+    $:pFile     filled with a handle to the file which will be opened
+    $:name      name of the file to open (VFS, UTF-8)
+    $:fileOpenFlags     (optional) flags for the os file open (see $OodleXFileOpenFlags)
+    $:vtable            (optional) the $OodleXFileOpsVTable to use for all ops on this file
+    $:autoDelete        (optional) see $OodleXHandleAutoDelete
+    $:priority          (optional) priority of the operation ; see $OodleXPriority
+    $:dependencies      (optional) dependencies; the async op won't start until these are all complete; note : these are not freed, they must be autodelete or you must free them some other way.
+    $:numDependencies   (optional)  number of handles in _deps_ array
+    $:return    handle to the operation, or 0 if it could not be started (usually due to invalid args)
+
+    If vtable is NULL, the vtable specified by the VFS mapping is used.
+    The file name provided is automatically run through VFS-to-OS name mapping, if applicable.
+    
+    Open returns a File ref right away for your convenience, but the file is not actually open for a little while.
+    You can however go ahead and queue more requests on the file reference before open is complete.
+    You cannot call things that require an open file, such as $OodleXIOQ_GetInfo.
+    OpenForRead is always shared access.
+    
+    To also perform an initial read, use $OodleXIOQ_OpenAndRead_Async
+
+*/
+
+OOFUNC1 OodleXHandle OOFUNC2 OodleXIOQ_OpenAndRead_Async(OodleXIOQFile * pFile,const char *name, 
+                        void * initialReadMemory, OO_SINTa initialReadSize, OO_S64 initialReadPos OODEFAULT(0),
+                        OodleXFileOpenFlags fileOpenFlags OODEFAULT(OodleXFileOpenFlags_Default),
+                        const OodleXFileOpsVTable * vtable OODEFAULT(NULL),OodleXHandleAutoDelete autoDelete OODEFAULT(OodleXHandleAutoDelete_No),OodleXPriority priority OODEFAULT(OodleXPriority_Normal),const OodleXHandle * dependencies OODEFAULT(NULL),OO_S32 numDependencies OODEFAULT(0));
+/* Start opening a file for read, and do an initial read
+
+    $:pFile     filled with a handle to the file which will be opened
+    $:name      name of the file to open (VFS, UTF-8)
+    $:initialReadMemory pointer to buffer to read into (must be $OODLEX_IO_MAX_ALIGNMENT aligned)
+    $:initialReadSize   amount to read (must be $OODLEX_IO_MAX_ALIGNMENT aligned)
+    $:initialReadPos    (optional) file position to read (must be $OODLEX_IO_MAX_ALIGNMENT aligned)
+    $:fileOpenFlags     (optional) flags for the os file open (see $OodleXFileOpenFlags)
+    $:vtable            (optional) the $OodleXFileOpsVTable to use for all ops on this file
+    $:autoDelete            (optional) see $OodleXHandleAutoDelete
+    $:priority          (optional) priority of the operation ; see $OodleXPriority
+    $:dependencies      (optional) dependencies; the async op won't start until these are all complete; note : these are not freed, they must be autodelete or you must free them some other way.
+    $:numDependencies   (optional)  number of handles in _deps_ array
+    $:return    handle to the operation, or 0 if it could not be started (usually due to invalid args)
+
+    If vtable is NULL, the vtable specified by the VFS mapping is used.
+    The file name provided is automatically run through VFS-to-OS name mapping, if applicable.
+    
+    Open returns a File ref right away for your convenience, but the file is not actually open for a little while.
+    You can however go ahead and queue more requests on the file reference before open is complete.
+    You cannot call things that require an open file, such as $OodleXIOQ_GetInfo.
+    OpenForRead is always shared access (when possible).
+    
+    Also performs an initial read.  Particularly useful when you need an initial header before you can start processing a file.
+
+*/
+
+// OpenForWrite is a create/truncate , exclusive access
+OOFUNC1 OodleXHandle OOFUNC2 OodleXIOQ_OpenForWriteCreate_Async(OodleXIOQFile * pFile,const char *name, OO_S64 initialFileSize OODEFAULT(OODLEX_FILE_OPEN_NO_RESERVE_SIZE),
+                                OodleXFileOpenFlags fileOpenFlags OODEFAULT(OodleXFileOpenFlags_Default),
+                                const OodleXFileOpsVTable * vtable OODEFAULT(NULL),OodleXHandleAutoDelete autoDelete OODEFAULT(OodleXHandleAutoDelete_No),OodleXPriority priority OODEFAULT(OodleXPriority_Normal),const OodleXHandle * dependencies OODEFAULT(NULL),OO_S32 numDependencies OODEFAULT(0));
+/* Start opening a file for write
+
+    $:pFile     filled with a handle to the file which will be opened
+    $:name      name of the file to open (VFS, UTF-8)
+    $:initialFileSize (optional) pre-allocate file size for writing (must be $OODLEX_IO_MAX_ALIGNMENT aligned)
+    $:fileOpenFlags     (optional) flags for the os file open (see $OodleXFileOpenFlags)
+    $:vtable    (optional) the $OodleXFileOpsVTable to use for all ops on this file
+    $:autoDelete            (optional) see $OodleXHandleAutoDelete
+    $:priority          (optional) priority of the operation ; see $OodleXPriority
+    $:dependencies      (optional) dependencies; the async op won't start until these are all complete; note : these are not freed, they must be autodelete or you must free them some other way.
+    $:numDependencies   (optional)  number of handles in _deps_ array
+    $:return    handle to the operation, or 0 if it could not be started (usually due to invalid args)
+
+    See many shared aspects of $OodleXIOQ_OpenForRead_Async .
+    
+    OpenForWrite opens files as create/truncate with exclusive access.
+    
+    initialFileSize performs an initial pre-allocation of file space, same as $OodleXIOQ_ReserveFileSizeForWrite_Async.
+    Pre-allocated file space has undefined (garbage) contents.  Writes are faster to pre-allocated space.
+    
+    WARNING : WriteCreate will overwrite (stomp) existing files by default.  If you don't want that, pass
+    OodleXFileOpenFlags_WriteCreateDontStomp in _fileOpenFlags_.
+*/
+
+OOFUNC1 OodleXHandle OOFUNC2 OodleXIOQ_OpenForWriteTempName_Async(OodleXIOQFile * pFile,const char *nameBase OODEFAULT(NULL), OO_S64 initialFileSize OODEFAULT(OODLEX_FILE_OPEN_NO_RESERVE_SIZE),
+                                OodleXFileOpenFlags fileOpenFlags OODEFAULT(OodleXFileOpenFlags_Default),
+                                const OodleXFileOpsVTable * vtable OODEFAULT(NULL),OodleXHandleAutoDelete autoDelete OODEFAULT(OodleXHandleAutoDelete_No),OodleXPriority priority OODEFAULT(OodleXPriority_Normal),const OodleXHandle * dependencies OODEFAULT(NULL),OO_S32 numDependencies OODEFAULT(0));
+/* Start opening a file for write
+
+    $:pFile     filled with a handle to the file which will be opened
+    $:nameBase      (optional) prefix of the temp file name that will be written (VFS, UTF-8)
+    $:initialFileSize (optional) pre-allocate file size for writing (must be $OODLEX_IO_MAX_ALIGNMENT aligned)
+    $:fileOpenFlags     (optional) flags for the os file open (see $OodleXFileOpenFlags)
+    $:vtable    (optional) the $OodleXFileOpsVTable to use for all ops on this file
+    $:autoDelete            (optional) see $OodleXHandleAutoDelete
+    $:priority          (optional) priority of the operation ; see $OodleXPriority
+    $:dependencies      (optional) dependencies; the async op won't start until these are all complete; note : these are not freed, they must be autodelete or you must free them some other way.
+    $:numDependencies   (optional)  number of handles in _deps_ array
+    $:return    handle to the operation, or 0 if it could not be started (usually due to invalid args)
+
+    Same as $OodleXIOQ_OpenForWriteCreate_Async except that it creates a unique temp name to write to.  The temp name starts with
+    _nameBase_, if given.  Providing _nameBase_ is helpful because it lets Oodle put the temp file in the same directory as the
+    final file name, which ensures that the final rename can be done without copying.
+    
+    Should be used with $OodleXIOQ_CloseFileRename_Async.
+    
+    Writing to a temp name and then renaming over the desired output file only on successful completion is the
+    recommended way to write all files.  It means you won't destroy the user's data by failing to successfully
+    overwrite a previously existing good file.
+*/
+
+// once you queue a CloseFile request, you should not touch OodleXIOQFile anymore :
+//  CloseFile also copiesthe error status from file to the final request, so you can check the aggregate error on the request
+OOFUNC1 OodleXHandle OOFUNC2 OodleXIOQ_CloseFile_Async(OodleXIOQFile file,OO_S64 truncateFileSize OODEFAULT(OODLEX_FILE_CLOSE_NO_TRUNCATE_SIZE),OodleXHandleAutoDelete autoDelete OODEFAULT(OodleXHandleAutoDelete_No),OodleXPriority priority OODEFAULT(OodleXPriority_Normal),const OodleXHandle * dependencies OODEFAULT(NULL),OO_S32 numDependencies OODEFAULT(0));
+/* Start a close-file request
+
+    $:file      the file to close
+    $:truncateFileSize  (optional) truncate an OpenforWrite file
+    $:autoDelete            (optional) see $OodleXHandleAutoDelete
+    $:priority          (optional) priority of the operation ; see $OodleXPriority
+    $:dependencies      (optional) dependencies; the async op won't start until these are all complete; note : these are not freed, they must be autodelete or you must free them some other way.
+    $:numDependencies   (optional)  number of handles in _deps_ array
+    $:return    handle to the operation, or 0 if it could not be started (usually due to invalid args)
+    
+    CloseFile also copies any errors on the file to the request, so that an $OodleXIOQ_GetStatus on the CloseFile
+    request will return $OodleXStatus_Error if there are any errors on the file.
+    
+    If the file was OpenForWrite, then truncateFileSize can be used to set the final file size.  This is mainly
+    used when the file was reserved with $OodleXIOQ_ReserveFileSizeForWrite_Async , but it should also be used any time
+    a file size that is not $OODLEX_IO_MAX_ALIGNMENT aligned is desired.  truncateFileSize does not need to be
+    $OODLEX_IO_MAX_ALIGNMENT aligned, but all sized for $OodleXIOQ_Write_Async do, so without doing this file sizes will
+    be aligned up.  Pass OODLEX_FILE_CLOSE_NO_TRUNCATE_SIZE (or use the default argument) if you don't want to truncate.
+*/
+
+OOFUNC1 OodleXHandle OOFUNC2 OodleXIOQ_CloseFileRename_Async(OodleXIOQFile file,const char * renameTo,OO_S64 truncateFileSize OODEFAULT(OODLEX_FILE_CLOSE_NO_TRUNCATE_SIZE),OodleXHandleAutoDelete autoDelete OODEFAULT(OodleXHandleAutoDelete_No),OodleXPriority priority OODEFAULT(OodleXPriority_Normal),const OodleXHandle * dependencies OODEFAULT(NULL),OO_S32 numDependencies OODEFAULT(0));
+/* Start a close-file-rename request
+
+    $:file      the file to close
+    $:renameTo  file to rename to (VFS, UTF-8)
+    $:truncateFileSize  (optional) truncate an OpenforWrite file
+    $:autoDelete            (optional) see $OodleXHandleAutoDelete
+    $:priority          (optional) priority of the operation ; see $OodleXPriority
+    $:dependencies      (optional) dependencies; the async op won't start until these are all complete; note : these are not freed, they must be autodelete or you must free them some other way.
+    $:numDependencies   (optional)  number of handles in _deps_ array
+    $:return    handle to the operation, or 0 if it could not be started (usually due to invalid args)
+    
+    Does an $OodleXIOQ_CloseFile_Async , then renames the file to "renameTo" - but only if there were no errors in writing
+    the file.  To stop unimportant errors from causing OodleXIOQ_CloseFileRename_Async to fail, use $OodleXIOQ_ClearError
+    before calling this.
+
+    _renameTo_ can be NULL to cancel the close and delete the temp file.
+
+    Useful with $OodleXIOQ_OpenForWriteTempName_Async.
+*/
+
+// Read & Write require size & position & memory to all be aligned
+OOFUNC1 OodleXHandle OOFUNC2 OodleXIOQ_Read_Async(OodleXIOQFile file,void * memory,OO_SINTa size,OO_S64 position,OodleXHandleAutoDelete autoDelete OODEFAULT(OodleXHandleAutoDelete_No),OodleXPriority priority OODEFAULT(OodleXPriority_Normal),const OodleXHandle * dependencies OODEFAULT(NULL),OO_S32 numDependencies OODEFAULT(0));
+/* Start a read request
+
+    $:file      the file to act on
+    $:memory    memory to read into (must be $OODLEX_IO_MAX_ALIGNMENT aligned)
+    $:size      number of bytes to read (must be $OODLEX_IO_MAX_ALIGNMENT aligned)
+    $:position  file position to start the read (must be $OODLEX_IO_MAX_ALIGNMENT aligned)
+    $:autoDelete            (optional) see $OodleXHandleAutoDelete
+    $:priority          (optional) priority of the operation ; see $OodleXPriority
+    $:dependencies      (optional) dependencies; the async op won't start until these are all complete; note : these are not freed, they must be autodelete or you must free them some other way.
+    $:numDependencies   (optional)  number of handles in _deps_ array
+    $:return    handle to the operation, or 0 if it could not be started (usually due to invalid args)
+
+    Starts an async read into _memory at file offset _position , of _size bytes.
+
+    To do unaligned reads, use $OodleXIOQ_ReadUnalignedAdjustPointer_Async , or
+    simply read a larger amount, and use $OodleX_IOAlignDownS64 on _position and $OodleX_IOAlignUpS64   on _size.
+    
+    The read is not done when OodleXIOQ_Read_Async returns.  You must not free _memory until the read is done, as reported by
+    the handle returned;
+*/
+
+OOFUNC1 OodleXHandle OOFUNC2 OodleXIOQ_Write_Async(OodleXIOQFile file,const void * memory,OO_SINTa size,OO_S64 position,OodleXHandleAutoDelete autoDelete OODEFAULT(OodleXHandleAutoDelete_No),OodleXPriority priority OODEFAULT(OodleXPriority_Normal),const OodleXHandle * dependencies OODEFAULT(NULL),OO_S32 numDependencies OODEFAULT(0));
+/* Start a write request
+
+    $:file      the file to act on
+    $:memory    memory to write from (must be $OODLEX_IO_MAX_ALIGNMENT aligned)
+    $:size      number of bytes to write (must be $OODLEX_IO_MAX_ALIGNMENT aligned)
+    $:position  file position to start the write (must be $OODLEX_IO_MAX_ALIGNMENT aligned)
+    $:autoDelete            (optional) see $OodleXHandleAutoDelete
+    $:priority          (optional) priority of the operation ; see $OodleXPriority
+    $:dependencies      (optional) dependencies; the async op won't start until these are all complete; note : these are not freed, they must be autodelete or you must free them some other way.
+    $:numDependencies   (optional)  number of handles in _deps_ array
+    $:return    handle to the operation, or 0 if it could not be started (usually due to invalid args)
+
+    Starts an async write from _memory at file offset _position , of _size bytes.
+    
+    The write is not done when OodleXIOQ_Write_Async returns.  You must not free _memory until the write is done, as reported by
+    the handle returned.
+    
+    Writes are faster on some platforms if the file size is first reserved past the end of the write, using
+    $OodleXIOQ_SetFileSize_Async or $OodleXIOQ_ReserveFileSizeForWrite_Async.   
+*/
+
+// SetFileSize for files you are writing
+//  file size should be incremented before writing
+OOFUNC1 OodleXHandle OOFUNC2 OodleXIOQ_SetFileSize_Async(OodleXIOQFile file,OO_S64 size,OodleXHandleAutoDelete autoDelete OODEFAULT(OodleXHandleAutoDelete_No),OodleXPriority priority OODEFAULT(OodleXPriority_Normal),const OodleXHandle * dependencies OODEFAULT(NULL),OO_S32 numDependencies OODEFAULT(0));
+/* Start a set-file-size request
+
+    $:file      the file to act on
+    $:size      the new file size
+    $:autoDelete            (optional) see $OodleXHandleAutoDelete
+    $:priority          (optional) priority of the operation ; see $OodleXPriority
+    $:dependencies      (optional) dependencies; the async op won't start until these are all complete; note : these are not freed, they must be autodelete or you must free them some other way.
+    $:numDependencies   (optional)  number of handles in _deps_ array
+    $:return    handle to the operation, or 0 if it could not be started (usually due to invalid args)
+    
+    Change the size of a file open for writing.
+    
+    OodleXIOQ_SetFileSize_Async may align up _size_ to the next sector boundary or $OODLEX_IO_MAX_ALIGNMENT.
+    The contents of the file in the resized but unwritten area are undefined/garbage.
+    
+    To write a file with non-aligned size, use $(OodleXIOQ_CloseFile_Async:truncateFileSize) in $OodleXIOQ_CloseFile_Async.
+    
+    If the purpose of calling SetFileSize is to pre-reserve space to make writes go faster, then use $OodleXIOQ_ReserveFileSizeForWrite_Async instead.
+*/
+
+// Reserve can be a NOP if it's faster not to ; SetFileSize is always done
+OOFUNC1 OodleXHandle OOFUNC2 OodleXIOQ_ReserveFileSizeForWrite_Async(OodleXIOQFile file,OO_S64 size,OodleXHandleAutoDelete autoDelete OODEFAULT(OodleXHandleAutoDelete_No),OodleXPriority priority OODEFAULT(OodleXPriority_Normal),const OodleXHandle * dependencies OODEFAULT(NULL),OO_S32 numDependencies OODEFAULT(0));
+/* Start a set-file-size request, if it helps write speed.
+
+    $:file      the file to act on
+    $:size      the new file size
+    $:autoDelete            (optional) see $OodleXHandleAutoDelete
+    $:priority          (optional) priority of the operation ; see $OodleXPriority
+    $:dependencies      (optional) dependencies; the async op won't start until these are all complete; note : these are not freed, they must be autodelete or you must free them some other way.
+    $:numDependencies   (optional)  number of handles in _deps_ array
+    $:return    handle to the operation, or 0 if it could not be started (usually due to invalid args)
+    
+    If the purpose of calling SetFileSize is to pre-reserve space to make writes go faster, then use $OodleXIOQ_ReserveFileSizeForWrite_Async instead.
+    ReserveFileSizeForWrite is the same as SetFileSize, but it uses some information about the platform and the file to decide whether the reserve
+    will help or not.  This function might do nothing if it thinks that the writes will be faster with no reservation.  
+    
+    The contents of the file in the resized but unwritten area are undefined/garbage.
+    
+    See $OodleXIOQ_SetFileSize_Async for more.
+*/
+
+OOFUNC1 OodleXHandle OOFUNC2 OodleXIOQ_ForceWriteable_Async(const char * name,OodleXHandleAutoDelete autoDelete OODEFAULT(OodleXHandleAutoDelete_No),OodleXPriority priority OODEFAULT(OodleXPriority_Normal),const OodleXHandle * dependencies OODEFAULT(NULL),OO_S32 numDependencies OODEFAULT(0));
+/* Start a force-writeable file request.
+
+    $:name              the file to make writeable (VFS, UTF-8)
+    $:autoDelete        (optional) see $OodleXHandleAutoDelete
+    $:priority          (optional) priority of the operation ; see $OodleXPriority
+    $:dependencies      (optional) dependencies; the async op won't start until these are all complete; note : these are not freed, they must be autodelete or you must free them some other way.
+    $:numDependencies   (optional)  number of handles in _deps_ array
+    $:return    handle to the operation, or 0 if it could not be started (usually due to invalid args)
+    
+    Asynchronously make a file writeable/deletable.  Useful if the file might have read-only or other-user permissions
+    and you want to modifity anyway.
+    
+    A common use is to enqueue a OodleXIOQ_ForceWriteable_Async right before a DeleteFile or RenameFile.
+*/
+
+OOFUNC1 OodleXHandle OOFUNC2 OodleXIOQ_Delete_Async(const char * name,OodleXHandleAutoDelete autoDelete OODEFAULT(OodleXHandleAutoDelete_No),OodleXPriority priority OODEFAULT(OodleXPriority_Normal),const OodleXHandle * dependencies OODEFAULT(NULL),OO_S32 numDependencies OODEFAULT(0));
+/* Start a delete request.
+
+    $:name      the file to delete (VFS, UTF-8)
+    $:autoDelete            (optional) see $OodleXHandleAutoDelete
+    $:priority          (optional) priority of the operation ; see $OodleXPriority
+    $:dependencies      (optional) dependencies; the async op won't start until these are all complete; note : these are not freed, they must be autodelete or you must free them some other way.
+    $:numDependencies   (optional)  number of handles in _deps_ array
+    $:return    handle to the operation, or 0 if it could not be started (usually due to invalid args)
+    
+    Asynchronously delete a file or dir.
+    
+    Use OodleXIOQ_ForceWriteable_Async before the Delete to force the deletion of read-only and other no-access conditions.
+*/
+
+OOFUNC1 OodleXHandle OOFUNC2 OodleXIOQ_Rename_Async(const char * fm,const char *to,OO_BOOL overwrite,OodleXHandleAutoDelete autoDelete OODEFAULT(OodleXHandleAutoDelete_No),OodleXPriority priority OODEFAULT(OodleXPriority_Normal),const OodleXHandle * dependencies OODEFAULT(NULL),OO_S32 numDependencies OODEFAULT(0));
+/* Start a renamefile request.
+
+    $:fm        the file to rename (VFS, UTF-8)
+    $:to        the new file name (VFS, UTF-8)
+    $:overwrite         if true, any existing file of name "to" will be overwrriten
+    $:autoDelete        (optional) see $OodleXHandleAutoDelete
+    $:priority          (optional) priority of the operation ; see $OodleXPriority
+    $:dependencies      (optional) dependencies; the async op won't start until these are all complete; note : these are not freed, they must be autodelete or you must free them some other way.
+    $:numDependencies   (optional)  number of handles in _deps_ array
+    $:return    handle to the operation, or 0 if it could not be started (usually due to invalid args)
+        
+    Asynchronously rename a file.
+    
+    Use OodleXIOQ_ForceWriteable_Async (on the to name) before the rename to force the ovewriting of read-only and other no-access conditions.
+*/
+
+
+OOFUNC1 OodleXHandle OOFUNC2 OodleXIOQ_MakeDir_Async(const char * name,OodleXHandleAutoDelete autoDelete OODEFAULT(OodleXHandleAutoDelete_No),OodleXPriority priority OODEFAULT(OodleXPriority_Normal),const OodleXHandle * dependencies OODEFAULT(NULL),OO_S32 numDependencies OODEFAULT(0));
+/* Start a make dir request.
+
+    $:name      the dir to make (VFS, UTF-8)
+    $:autoDelete            (optional) see $OodleXHandleAutoDelete
+    $:priority          (optional) priority of the operation ; see $OodleXPriority
+    $:dependencies      (optional) dependencies; the async op won't start until these are all complete; note : these are not freed, they must be autodelete or you must free them some other way.
+    $:numDependencies   (optional)  number of handles in _deps_ array
+    $:return    handle to the operation, or 0 if it could not be started (usually due to invalid args)
+    
+    Asynchronously make a dir.
+*/
+
+//OodleXHandle OodleIOQ_SetFileSizeUnaligned(const char * name,OO_S64 size);
+
+// runs OodleXFree_IOAligned :
+OOFUNC1 OodleXHandle OOFUNC2 OodleXIOQ_FreeBufferIOAligned_Async(OodleXIOQFile file,void * buffer,OodleXHandleAutoDelete autoDelete OODEFAULT(OodleXHandleAutoDelete_No),OodleXPriority priority OODEFAULT(OodleXPriority_Normal),const OodleXHandle * dependencies OODEFAULT(NULL),OO_S32 numDependencies OODEFAULT(0));
+/* Start a free-buffer request
+
+    $:file      the request is scheduled on this file
+    $:buffer    the buffer to free
+    $:autoDelete            (optional) see $OodleXHandleAutoDelete
+    $:priority          (optional) priority of the operation ; see $OodleXPriority
+    $:dependencies      (optional) dependencies; the async op won't start until these are all complete; note : these are not freed, they must be autodelete or you must free them some other way.
+    $:numDependencies   (optional)  number of handles in _deps_ array
+    $:return    handle to the operation, or 0 if it could not be started (usually due to invalid args)
+    
+    Runs $OodleXFree_IOAligned as an IO operation on the file's sequential list of ops.
+    
+    This is mainly useful with $OodleXIOQ_Write_Async.  When you write a buffer, you can't free it until the write
+    is done, with this you can just call Write and then FreeBuffer ; the FreeBuffer will be run when the Write is done.
+*/
+
+
+typedef enum OodleFileNotFoundIsAnError
+{
+    OodleFileNotFoundIsAnError_No  = 0,
+    OodleFileNotFoundIsAnError_Yes = 1,
+    OodleFileNotFoundIsAnError_Force32 = 0x40000000
+} OodleFileNotFoundIsAnError;
+/* Bool for whether a file not found is a completion status of $OodleXStatus_Error or $OodleXStatus_Done
+*/
+
+OOFUNC1 OodleXHandle OOFUNC2 OodleXIOQ_GetInfoByName_Async(const char * name,
+    OodleFileNotFoundIsAnError errorIfNotFound, // OODEFAULT(OodleFileNotFoundIsAnError_Yes)
+    OodleXHandleAutoDelete autoDelete OODEFAULT(OodleXHandleAutoDelete_No),OodleXPriority priority OODEFAULT(OodleXPriority_Normal),const OodleXHandle * dependencies OODEFAULT(NULL),OO_S32 numDependencies OODEFAULT(0));
+/* Start an asynchronous GetInfo request
+
+    $:name      the file name to query (VFS,UTF-8)
+    $:errorIfNotFound   (optional) should file-not-found be an error status or not? (see $OodleFileNotFoundIsAnError)
+    $:autoDelete        (optional) see $OodleXHandleAutoDelete
+    $:priority          (optional) priority of the operation ; see $OodleXPriority
+    $:dependencies      (optional) dependencies; the async op won't start until these are all complete; note : these are not freed, they must be autodelete or you must free them some other way.
+    $:numDependencies   (optional)  number of handles in _deps_ array
+    $:return    handle to the operation, or 0 if it could not be started (usually due to invalid args)
+    
+    Used to get an $OodleXFileInfo without blocking the calling thread. Use $OodleXIOQ_GetInfoByName_GetResult to
+    get the result.
+*/
+
+OOFUNC1 OodleXStatus OOFUNC2 OodleXIOQ_GetInfoByName_GetResult(OodleXHandle req,OodleXHandleDeleteIfDone andDeleteIfDone,OodleXFileInfo * pInfo);
+/* Finish an asynchronous GetInfo request
+
+    $:req       the handle returned by $OodleXIOQ_GetInfoByName_Async
+    $:andDeleteIfDone   if true and the request is done, delete the request 
+    $:pInfo     filled out with an $OodleXFileInfo
+    $:return    OodleXStatus of the request
+    
+    _pInfo_ is filled if the $OodleXStatus returned is OodleXStatus_Done.  If the return value is something
+    else, _pInfo_ is untouched (eg. not invalidated!).
+    
+    If the file does not exist and OodleFileNotFoundIsAnError_No was passed, this function will return
+    OodleXStatus_Done but _pInfo_ will be set to a $(OodleXFileInfo::size) of $OODLEX_FILE_SIZE_INVALID;
+*/
+
+OOFUNC1 OodleXHandle OOFUNC2 OodleXIOQ_SetInfoByName_Async(const char * name,OO_U32 flags,OO_U64 modTime,OodleXHandleAutoDelete autoDelete OODEFAULT(OodleXHandleAutoDelete_No),OodleXPriority priority OODEFAULT(OodleXPriority_Normal),const OodleXHandle * dependencies OODEFAULT(NULL),OO_S32 numDependencies OODEFAULT(0));
+/* Start an asynchronous SetInfo request
+
+    $:name      the file name to query (VFS,UTF-8)
+    $:flags         file flags (logical OR of $OODLEX_FILEINFO_FLAGS) , or $OODLEX_FILEINFO_FLAG_INVALID to leave unchanged
+    $:modTime       mod time to change, or %OODLEX_FILEINFO_MODTIME_INVALID to leave uncahnged
+    $:autoDelete            (optional) see $OodleXHandleAutoDelete
+    $:priority          (optional) priority of the operation ; see $OodleXPriority
+    $:dependencies      (optional) dependencies; the async op won't start until these are all complete; note : these are not freed, they must be autodelete or you must free them some other way.
+    $:numDependencies   (optional)  number of handles in _deps_ array
+    $:return    handle to the operation, or 0 if it could not be started (usually due to invalid args)
+    
+    Used to set flags or modtime on a file.
+    Flags and modTime correspond to $(OodleXFileInfo:flags) and $(OodleXFileInfo:modTime).
+
+    All members of $OodleXFileInfo can be set this way, except size; to set size use $OodleXIOQ_SetFileSize_Async.  
+*/
+
+//---------------------------------------------------------------------------------
+// high level section :
+// these are not independent request types
+//  just high level helpers that do several requests for you
+
+OOFUNC1 OodleXHandle OOFUNC2 OodleXIOQ_ReadMallocWholeFile_Async(OodleXIOQFile file,OodleXHandleAutoDelete autoDelete OODEFAULT(OodleXHandleAutoDelete_No),OodleXPriority priority OODEFAULT(OodleXPriority_Normal),const OodleXHandle * dependencies OODEFAULT(NULL),OO_S32 numDependencies OODEFAULT(0));
+/* Start a high level IO request to allocate a buffer for a whole file and read it
+
+    $:file      the file to act on
+    $:autoDelete            (optional) see $OodleXHandleAutoDelete
+    $:priority          (optional) priority of the operation ; see $OodleXPriority
+    $:dependencies      (optional) dependencies; the async op won't start until these are all complete; note : these are not freed, they must be autodelete or you must free them some other way.
+    $:numDependencies   (optional)  number of handles in _deps_ array
+    $:return    handle to the operation, or 0 if it could not be started (usually due to invalid args)
+    
+    High level IOQ operations are helpers built on the simpler IOQ low level ops.
+    
+    OodleXIOQ_ReadMallocWholeFile_Async calls $OodleXMalloc_IOAligned to allocate a buffer the size of the whole file
+    (aligned up by $OODLEX_IO_MAX_ALIGNMENT), and reads the whole file into that buffer.
+    
+    Get the buffer pointer with $OodleXIOQ_ReadMallocWholeFile_GetResult.  You must free it.
+*/
+
+OOFUNC1 OodleXStatus OOFUNC2 OodleXIOQ_ReadMallocWholeFile_GetResult(OodleXHandle req,OodleXHandleDeleteIfDone andDeleteIfDone, void ** pPtr, OO_S64 * pSize OODEFAULT(NULL));
+/* Finish a $OodleXIOQ_ReadMallocWholeFile_Async request
+
+    $:req   the OodleXHandle to the OodleXIOQ_ReadMallocWholeFile_Async request
+    $:andDeleteIfDone   if true and the returned status is >= Done the handle will be deleted
+    $:pPtr  filled out with the buffer allocated by OodleXIOQ_ReadMallocWholeFile_Async
+    $:pSize (optional) filled with the file size
+    $:return the status ; if <= OodleXStatus_Pending , result pointers are set to null
+
+    OodleXIOQ_ReadMallocWholeFile_GetResult does NOT wait on the handle.
+
+    See $OodleXIOQ_ReadMallocWholeFile_Async
+*/
+
+OOFUNC1 OodleXHandle OOFUNC2 OodleXIOQ_OpenAndReadMallocWholeFile_Async(OodleXIOQFile * pFile,const char *name,
+                                    OodleXFileOpenFlags fileOpenFlags OODEFAULT(OodleXFileOpenFlags_Default),
+                                    const OodleXFileOpsVTable * vtable OODEFAULT(NULL),OodleXHandleAutoDelete autoDelete OODEFAULT(OodleXHandleAutoDelete_No),OodleXPriority priority OODEFAULT(OodleXPriority_Normal),const OodleXHandle * dependencies OODEFAULT(NULL),OO_S32 numDependencies OODEFAULT(0));
+/* Start a high level IO request to open a file, allocate a buffer for a whole file and read it
+
+    $:pFile     filled with a handle to the file which will be opened
+    $:name      name of the file to open (VFS, UTF-8)
+    $:fileOpenFlags     (optional) flags for the os file open (see $OodleXFileOpenFlags)
+    $:vtable    (optional) the $OodleXFileOpsVTable to use for all ops on this file
+    $:return    handle to the RMWF op; use $OodleXIOQ_ReadMallocWholeFile_GetResult
+    
+    High level IOQ operations are helpers built on the simpler IOQ low level ops.
+    
+    Performs $OodleXIOQ_OpenForRead_Async and $OodleXIOQ_ReadMallocWholeFile_Async.  
+
+    The $OodleXHandle returned is to the RMWF operation; use $OodleXIOQ_ReadMallocWholeFile_GetResult.
+    
+    You will normally want to enqueue an $OodleXIOQ_CloseFile_Async after this.
+*/
+
+
+OOFUNC1 OodleXHandle OOFUNC2 OodleXIOQ_OpenAndReadMallocWholeFileAndClose_Async(const char *name,
+                        OodleXFileOpenFlags fileOpenFlags OODEFAULT(OodleXFileOpenFlags_Default),
+                        const OodleXFileOpsVTable * vtable OODEFAULT(NULL),OodleXHandleAutoDelete autoDelete OODEFAULT(OodleXHandleAutoDelete_No),OodleXPriority priority OODEFAULT(OodleXPriority_Normal),const OodleXHandle * dependencies OODEFAULT(NULL),OO_S32 numDependencies OODEFAULT(0));
+/* Start a high level IO request to open a file, allocate a buffer for a whole file and read it
+
+    $:name      name of the file to open (VFS, UTF-8)
+    $:fileOpenFlags     (optional) flags for the os file open (see $OodleXFileOpenFlags)
+    $:vtable    (optional) the $OodleXFileOpsVTable to use for all ops on this file
+    $:return    handle to the RMWF op; use $OodleXIOQ_ReadMallocWholeFile_GetResult
+    
+    High level IOQ operations are helpers built on the simpler IOQ low level ops.
+    
+    Performs $OodleXIOQ_OpenForRead_Async and $OodleXIOQ_ReadMallocWholeFile_Async and $OodleXIOQ_CloseFile_Async.
+    
+    The $OodleXHandle returned is to the RMWF operation; use $OodleXIOQ_ReadMallocWholeFile_GetResult.
+*/
+
+OOFUNC1 OodleXHandle OOFUNC2 OodleXIOQ_OpenWriteWholeFileClose_Async(const char *name, const void * buffer, OO_SINTa size,
+                        OodleXFileOpenFlags fileOpenFlags OODEFAULT(OodleXFileOpenFlags_Default),
+                        const OodleXFileOpsVTable * vtable OODEFAULT(NULL),OodleXHandleAutoDelete autoDelete OODEFAULT(OodleXHandleAutoDelete_No),OodleXPriority priority OODEFAULT(OodleXPriority_Normal),const OodleXHandle * dependencies OODEFAULT(NULL),OO_S32 numDependencies OODEFAULT(0));
+/* Start a high level IO request to open a file, write a buffer, and close it.
+
+    $:name      name of the file to open (VFS, UTF-8)
+    $:buffer    the buffer to write (must be $OODLEX_IO_MAX_ALIGNMENT aligned)
+    $:size      the final file size (no alignment required)
+    $:fileOpenFlags     (optional) flags for the os file open (see $OodleXFileOpenFlags)
+    $:vtable    (optional) the $OodleXFileOpsVTable to use for all ops on this file
+    $:autoDelete            (optional) see $OodleXHandleAutoDelete
+    $:priority          (optional) priority of the operation ; see $OodleXPriority
+    $:dependencies      (optional) dependencies; the async op won't start until these are all complete; note : these are not freed, they must be autodelete or you must free them some other way.
+    $:numDependencies   (optional)  number of handles in _deps_ array
+    $:return    handle to the operation, or 0 if it could not be started (usually due to invalid args)
+    
+    High level IOQ operations are helpers built on the simpler IOQ low level ops.
+    
+    Performs $OodleXIOQ_OpenForWriteCreate_Async, $OodleXIOQ_Write_Async, and $OodleXIOQ_CloseFile_Async.
+    
+    You might also want to enqueue a $OodleXIOQ_FreeBufferIOAligned_Async after this, but it is not done for you.
+    See also $Oodle_FAQ_BadWriteContents.
+    
+    The $OodleXHandle returned is not done until the entire compound operation is done.
+*/
+
+// same thing but use a temp name & only rename if all successful
+OOFUNC1 OodleXHandle OOFUNC2 OodleXIOQ_OpenWriteWholeFileCloseTempName_Async(const char *name, const void * buffer, OO_SINTa size, 
+                        OodleXFileOpenFlags fileOpenFlags OODEFAULT(OodleXFileOpenFlags_Default),
+                        const OodleXFileOpsVTable * vtable OODEFAULT(NULL),OodleXHandleAutoDelete autoDelete OODEFAULT(OodleXHandleAutoDelete_No),OodleXPriority priority OODEFAULT(OodleXPriority_Normal),const OodleXHandle * dependencies OODEFAULT(NULL),OO_S32 numDependencies OODEFAULT(0));
+/*Start a high level IO request to open a file, write a buffer, close it, and rename it.
+
+    $:name      name of the file to open (VFS, UTF-8)
+    $:buffer    the buffer to write (must be $OODLEX_IO_MAX_ALIGNMENT aligned)
+    $:size      the final file size (no alignment required)
+    $:fileOpenFlags     (optional) flags for the os file open (see $OodleXFileOpenFlags)
+    $:vtable            (optional) the $OodleXFileOpsVTable to use for all ops on this file
+    $:autoDelete        (optional) see $OodleXHandleAutoDelete
+    $:priority          (optional) priority of the operation ; see $OodleXPriority
+    $:dependencies      (optional) dependencies; the async op won't start until these are all complete; note : these are not freed, they must be autodelete or you must free them some other way.
+    $:numDependencies   (optional)  number of handles in _deps_ array
+    $:return    handle to the operation, or 0 if it could not be started (usually due to invalid args)
+    
+    High level IOQ operations are helpers built on the simpler IOQ low level ops.
+    
+    This is the same as $OodleXIOQ_OpenWriteWholeFileClose_Async , but the writing is done to a temp file, and then
+    renamed to _name at the end, like $OodleXIOQ_CloseFileRename_Async.  The rename is only done if the writing succeeded.
+
+*/
+
+// helper wrapper for OodleXIOQ_Read_Async that handles unaligned reads for you
+//  the returned pointer will contain the bytes you asked for when the request is done
+OOFUNC1 OodleXHandle OOFUNC2 OodleXIOQ_ReadUnalignedAdjustPointer_Async(void ** pPtr, OodleXIOQFile file,void * memory,OO_SINTa readSize,OO_S64 position,OO_SINTa memorySize,OodleXHandleAutoDelete autoDelete OODEFAULT(OodleXHandleAutoDelete_No),OodleXPriority priority OODEFAULT(OodleXPriority_Normal),const OodleXHandle * dependencies OODEFAULT(NULL),OO_S32 numDependencies OODEFAULT(0));
+/* Start a read request with unaligned position or size
+
+    $:pPtr      filled with a pointer to the read memory; NULL if the read is impossible
+    $:file      the file to act on
+    $:memory    memory to read into (no alignment required)
+    $:readSize  number of bytes to read (no alignment required)
+    $:position  file position to start the read (no alignment required)
+    $:memorySize the size of the buffer at "memory"
+    $:autoDelete            (optional) see $OodleXHandleAutoDelete
+    $:priority          (optional) priority of the operation ; see $OodleXPriority
+    $:dependencies      (optional) dependencies; the async op won't start until these are all complete; note : these are not freed, they must be autodelete or you must free them some other way.
+    $:numDependencies   (optional)  number of handles in _deps_ array
+    $:return    handle to the operation, or 0 if it could not be started (usually due to invalid args)
+
+    High level IOQ operations are helpers built on the simpler IOQ low level ops.
+    
+    memorySize should be larger than size ; generally at least aligned up with $OodleX_IOAlignUpS64.
+    
+    OodleXIOQ_ReadUnalignedAdjustPointer_Async reads a larger chunk than [position,size] , aligning down the start
+    and aligning up the end.  It reads somewhere into [memory,memorySize].  The returned pointer is somewhere
+    in _memory and contains the bytes you wanted from _position.
+    
+    If memorySize is not big enough, it returns NULL.
+*/
+
+OOFUNC1 OodleXHandle OOFUNC2 OodleXIOQ_MakeAllDirs_Async(const char * name,OodleXHandleAutoDelete autoDelete OODEFAULT(OodleXHandleAutoDelete_No),OodleXPriority priority OODEFAULT(OodleXPriority_Normal),const OodleXHandle * dependencies OODEFAULT(NULL),OO_S32 numDependencies OODEFAULT(0));
+/*  Start a high level IO request to make all dirs in name
+
+    $:name      name of the file to make dirs for (VFS, UTF-8)
+    $:autoDelete            (optional) see $OodleXHandleAutoDelete
+    $:priority          (optional) priority of the operation ; see $OodleXPriority
+    $:dependencies      (optional) dependencies; the async op won't start until these are all complete; note : these are not freed, they must be autodelete or you must free them some other way.
+    $:numDependencies   (optional)  number of handles in _deps_ array
+    $:return    handle to the operation, or 0 if it could not be started (usually due to invalid args)
+    
+    Makes the dirs in _name_ in sequence.  _name_ can be a file name, or a path with trailing path delim.
+    
+    eg. if name is "a/b/c/d" then dir a is made, then b, then c, but not d.
+*/
+
+OOFUNC1 OodleXHandle OOFUNC2 OodleXIOQ_CopyFile_Async(const char * from,const char * to,OO_U32 oodleCopyFileFlags,OodleXHandleAutoDelete autoDelete OODEFAULT(OodleXHandleAutoDelete_No),OodleXPriority priority OODEFAULT(OodleXPriority_Normal),const OodleXHandle * dependencies OODEFAULT(NULL),OO_S32 numDependencies OODEFAULT(0));
+/* Start a high level IO request to copy a file
+
+    $:from      source file name (VFS,UTF8)
+    $:to        dest file name (VFS,UTF8)
+    $:oodleCopyFileFlags    bitwise OR of flags from $OodleXCopyFileFlags
+    $:autoDelete            (optional) see $OodleXHandleAutoDelete
+    $:priority          (optional) priority of the operation ; see $OodleXPriority
+    $:dependencies      (optional) dependencies; the async op won't start until these are all complete; note : these are not freed, they must be autodelete or you must free them some other way.
+    $:numDependencies   (optional)  number of handles in _deps_ array
+    $:return    handle to the operation, or 0 if it could not be started (usually due to invalid args)
+
+    Copy a file as a single IOQ op.
+    
+    NOTE : you generally do not want this.  Use Oodle_CopyFile_Async instead.  Using this call blocks
+    the IOQ from servicing streams or doing other work.
+
+    CopyFile is a single IOQ op so it is guaranteed to be done before a subsequent call to
+    OodleIOQ_OpenForRead on the _to_ file, so it is useful for async transparent mirroring.
+    (the same is not true of Oodle_CopyFile_Async which has undefined scheduling).  
+    
+*/
+
+
+//---------------------------------------------------------------------------------
+
+//-----------------------------------------------------
+// Simple "Event" object that just does a Pending->Done state transition
+
+OOFUNC1 OodleXHandle OOFUNC2 OodleXHandleEvent_Alloc(OodleXHandleAutoDelete autoDelete OODEFAULT(OodleXHandleAutoDelete_No));
+/* Allocate an $OodleXHandle to a simple data-less event 
+
+    $:autoDelete    (optional) set the $OodleXHandleAutoDelete of the handle
+    $:return    the handle
+    
+    An "event" simply stores a transition from Pending -> Done/Error and can be used to wait on something you can trigger.
+*/
+
+OOFUNC1 void OOFUNC2 OodleXHandleEvent_SetDone(OodleXHandle h);
+/* Set an OodleXHandleEvent to $OodleXStatus_Done
+
+    $:h     handle created by $OodleXHandleEvent_Alloc
+
+    The state transition from Pending->Done is one way.  If the handle is OodleXHandleAutoDelete_Yes, it
+    goes away now.
+*/
+
+OOFUNC1 void OOFUNC2 OodleXHandleEvent_SetError(OodleXHandle h);
+/* Set an OodleXHandleEvent to $OodleXStatus_Error
+
+    $:h     handle created by $OodleXHandleEvent_Alloc
+    
+    The state transition from Pending->Error is one way.  If the handle is OodleXHandleAutoDelete_Yes, it
+    goes away now.
+*/
+
+//-------------------------------------------------------
+// A "Countdown" object that is done when count reaches 0
+//  note : this is decently more expensive than SimpleCountdown
+
+OOFUNC1 OodleXHandle OOFUNC2 OodleXHandleCountdown_Alloc(OO_S32 initialCount,OodleXHandleAutoDelete autoDelete OODEFAULT(OodleXHandleAutoDelete_No));
+/*  Allocate an $OodleXHandle to a simple data-less coutdown 
+
+    $:initialCount  inital count; becomes done when count reaches 0
+    $:autoDelete    (optional) set the $OodleXHandleAutoDelete of the handle
+    $:return    the handle
+    
+    _initialCount_ should be greater than 0.
+    
+    A Countdown is a simple handle which you can use to wait for completion of many tasks.
+    Use $OodleXHandleCountdown_Decrement to decrement it.  When it reaches 0 it becomes Done,
+    which means it satisfies an $OodleX_Wait.
+    
+    (A countdown is the same thing as a single-use Semaphore with an initial negative count)
+    
+    If _autoDelete_ is OodleXHandleAutoDelete_Yes , the Countdown handle
+    is deleted when count reaches zero.  (a deleted handle also satisfies OodleX_Wait).
+*/
+
+OOFUNC1 OodleXStatus OOFUNC2 OodleXHandleCountdown_Decrement(OodleXHandle h,OO_S32 decCount OODEFAULT(1));
+/* Decrement a countdown handle created by $OodleXHandleCountdown_Alloc
+
+    $:h             handle allocated by $OodleXHandleCountdown_Alloc
+    $:decCount      how much to decrement the countdown
+    $:return        status after the decrement
+    
+    Returns $OodleXStatus_Done if this decrement took the countdown to 0, else $OodleXStatus_Pending.
+*/
+
+OOFUNC1 OodleXHandle OOFUNC2 OodleXLZ_Decompress_ThreadPhased_Narrow_Async(
+                                                const void * compBuf,OO_SINTa compSize,
+                                                void * decBuf,OO_SINTa rawSize,
+                                                OodleLZ_CheckCRC checkCRC OODEFAULT(OodleLZ_CheckCRC_No),
+                                                void * decBufBase OODEFAULT(NULL),
+                                                OO_SINTa decBufSize OODEFAULT(0),
+                                                OO_S32 circularBufferBlockCount OODEFAULT(-1),
+                                                void * scratchBuf OODEFAULT(NULL),
+                                                OO_BOOL synchronous_use_current_thread OODEFAULT(false));
+/* Start an async LZ decompress for ThreadPhase decoding, using 2 threads
+
+    $:compBuf       pointer to compressed data
+    $:compBufferSize    number of compressed bytes to decode
+    $:rawBuf        pointer to output uncompressed data into
+    $:rawLen        number of uncompressed bytes to output
+    $:checkCRC      (optional) if data could be corrupted and you want to know about it, pass OodleLZ_CheckCRC_Yes
+    $:decBufBase    (optional) if not NULL, provides preceding data to prime the dictionary; must be contiguous with rawBuf, the data between the pointers _dictionaryBase_ and _rawBuf_ is used as the preconditioning data.   The exact same precondition must be passed to encoder and decoder.  The decBufBase must be a reset point.
+    $:decBufSize    (optional) size of circular buffer starting at decBufBase
+    $:circularBufferBlockCount  (optional) number of blocks for circular buffer; generally more is faster but takes more memory; < 0 means use default
+    $:scratchBuf    (optional) memory to use for scratch; must be OodleLZ_ThreadPhased_BlockDecoderMemorySizeNeeded() * circularBufferBlockCount ; if NULL will be allocated
+    $:synchronous_use_current_thread    (optional) if true, runs on the current thread and uses 1 additional thread; this makes this a synchronous call and won't return until decompression is done (default is to use 2 worker threads and be fully async)
+    $:return        OodleXHandle to the operation, wait and check status to get result
+    
+    Runs a 2-thread Narrow decompress using the Oodle Worker system.
+    You must wait and delete the return handle, for example with OodleX_WaitAndDelete
+    
+    This only works on data that has been compressed with a compressor that's eligible for ThreadPhased decode;
+    check OodleLZ_Compressor_CanDecodeThreadPhased.  (currently just Kraken).
+    
+    See $OodleLZ_About_ThreadPhasedDecode
+    
+    This function (OodleXLZ_Decompress_ThreadPhased_Narrow_Async) does NOT parallelize at seek reset points.
+    You can however do so yourself externally to calling this function.  Simply scan the compressed buffer for
+    seek points and launch a separate OodleXLZ_Decompress_ThreadPhased_Narrow_Async call on each seek chunk.
+
+    ThreadPhased decode is always fuzz safe.
+    
+    If synchronous_use_current_thread then the returned handle is not async, you may check its status to get the result.
+    
+    
+*/
+                                                
+// asyncSelect is made from $OodleXAsyncSelect
+OOFUNC1 OodleXHandle OOFUNC2 OodleXLZ_Decompress_Narrow_Async(OO_U32 asyncSelect,
+                                                const void * packedDataPtr,OO_SINTa packedLen,
+                                                void * rawPtr,OO_SINTa rawChunkLen,
+                                                OodleLZ_FuzzSafe fuzzSafe OODEFAULT(OodleLZ_FuzzSafe_No), 
+                                                OodleLZ_CheckCRC checkCRC OODEFAULT(OodleLZ_CheckCRC_No), 
+                                                OodleLZ_Verbosity verbosity OODEFAULT(OodleLZ_Verbosity_None),
+                                                void * decBufBase OODEFAULT(NULL),OO_SINTa decBufSize OODEFAULT(0),
+                                                OodleDecompressCallback * pcb OODEFAULT(NULL),
+                                                void * pcbData OODEFAULT(NULL),
+                                                void *decMem OODEFAULT(NULL), OO_SINTa decMemSize OODEFAULT(0),
+                                                OodleLZ_Decode_ThreadPhase threadPhase OODEFAULT(OodleLZ_Decode_Unthreaded),
+                                                OodleXIOQFile writeToFile OODEFAULT(0),
+                                                OO_S64 writeToFileStartPos OODEFAULT(0),
+                                                OodleXHandle writeHandleGroup OODEFAULT(0),
+                                                OO_S32 writeHandleGroupIndex OODEFAULT(0),
+                                                OodleXHandleAutoDelete autoDelete OODEFAULT(OodleXHandleAutoDelete_No),
+                                                const OodleXHandle * dependencies OODEFAULT(NULL),OO_S32 numDependencies OODEFAULT(0),
+                                                OodleXPriority work_priority OODEFAULT(OodleXPriority_Normal));
+/* Start an async LZ decompress
+
+    $:asyncSelect   logical OR of $OodleXAsyncSelect flags determine how the async is run
+    $:packedDataPtr pointer to LZ compressed data
+    $:packedLen     compressed data length
+    $:rawPtr        pointer to memory filled with decompressed data
+    $:rawChunkLen   length of decompressed data
+    $:checkCRC      if OodleLZ_CheckCRC_Yes, the decompressor checks the crc to ensure data integrity
+    $:verbosity     (optional) if not OodleLZ_Verbosity_None, will log some information
+    $:decBufBase    (optional) if not NULL, provides preceding data to prime the dictionary; must be contiguous with rawBuf, the data between the pointers _dictionaryBase_ and _rawBuf_ is used as the preconditioning data.   The exact same precondition must be passed to encoder and decoder.
+    $:decBufSize    (optional) size of circular buffer starting at decBufBase   
+    $:pcb           (optional) OodleDecompressCallback called during decompression
+    $:pcbData       (optional) user data passed to pcb
+    $:writeToFile           (optional) OodleXIOQFile to write raw data to
+    $:writeToFileStartPos   (optional) file position where writeToFile should start (must be $OODLEX_IO_MAX_ALIGNMENT aligned)
+    $:writeHandleGroup      (optional) OodleAsyncGroup handle which the write handle is put into
+    $:writeHandleGroupIndex (optional) index in writeHandleGroup to use; must previously be $OODLEX_ASYNC_HANDLE_PENDING
+    $:autoDelete            (optional) see $OodleXHandleAutoDelete
+    $:dependencies      (optional) dependencies; the async op won't start until these are all complete; note : these are not freed, they must be autodelete or you must free them some other way.
+    $:numDependencies   (optional)  number of handles in _deps_ array
+
+    $:return        OodleXHandle to the operation, or OodleXHandle_Null for invalid arguments
+    
+    Start an async LZ decompress with the runner specified in asyncSelect.
+
+    A Narrow decompress means the entire decompression is done on one thread.
+    Data will always be decompressed sequentially, eg. in order.
+    
+    _rawChunkLen_ can be less than the entire original block, if it is a multiple of $OODLELZ_BLOCK_LEN.
+    
+    If provided, the $OodleDecompressCallback is called as quanta of raw data are available.  The callback
+    may be called more often than $OODLELZ_BLOCK_LEN granularity.
+
+    _rawPtr_ and _packedDataPtr_ memory blocks passed to this function must be kept alive for the duration of the async.
+
+    NOTE !! : if _writeToFile_ is provided, the writes are async and are NOT necessarily done when
+    the returned handle is done; the returned handle is for the decompress.  The handle for the write can
+    be retrieved by passing in _writeHandleGroup_.  You must not free the buffer being written until the
+    write operation is done.
+*/
+
+OOFUNC1 OodleXHandle OOFUNC2 OodleXLZ_ReadAndDecompress_Wide_Async(
+            OO_U32 asyncSelect,
+            const OodleLZ_SeekTable * seekTable,
+            const void * packedDataPtr,OO_SINTa packedLen,OO_SINTa packedLenPreviouslyRead,
+            OodleXIOQFile packedFile,OO_S64 packedDataStartPos,
+            void * rawArray,OO_SINTa rawArrayLen,
+            OodleLZ_FuzzSafe fuzzSafe, 
+            OodleLZ_CheckCRC checkCRC, 
+            OodleLZ_Verbosity verbosity OODEFAULT(OodleLZ_Verbosity_None),
+            void * decBufBase OODEFAULT(NULL),OO_SINTa decBufSize OODEFAULT(0),
+            OodleLZ_PackedRawOverlap packedRawOverlap OODEFAULT(OodleLZ_PackedRawOverlap_No),
+            OodleXIOQFile writeToFile OODEFAULT(0),
+            OO_S64 writeToFileStartPos OODEFAULT(0),
+            OodleXHandle * pWriteHandleGroup OODEFAULT(0),
+            OodleXHandleAutoDelete autoDelete OODEFAULT(OodleXHandleAutoDelete_No),
+            const OodleXHandle * dependencies OODEFAULT(NULL),OO_S32 numDependencies OODEFAULT(0));
+/* Start an async LZ decompress, possibly read packed data and write raw data
+
+    $:asyncSelect   logical OR of $OodleXAsyncSelect flags determine how the async is run
+    $:seekTable     seek locations as created by $OodleLZ_CreateSeekTable
+    $:packedDataPtr pointer to LZ compressed data
+    $:packedLen     compressed data length
+    $:packedLenPreviouslyRead   number of packed bytes already in packedDataPtr from previous IO ; eg. packedLen if the whole buffer is full
+    $:packedFile    OodleXIOQFile to read packed bytes from
+    $:packedDataStartPos  file position where the packed data starts (must be misaligned the same way as _packedDataPtr_)
+    $:rawArray      pointer to memory filled with decompressed data
+    $:rawArrayLen   length of decompressed data
+    $:checkCRC      if OodleLZ_CheckCRC_Yes, the decompressor checks the crc to ensure data integrity
+    $:verbosity     (optional) if not OodleLZ_Verbosity_None, will log some information
+    $:decBufBase    (optional) if not NULL, provides preceding data to prime the dictionary; must be contiguous with rawBuf, the data between the pointers _dictionaryBase_ and _rawBuf_ is used as the preconditioning data.   The exact same precondition must be passed to encoder and decoder.
+    $:decBufSize    (optional) size of circular buffer starting at decBufBase   
+    $:packedRawOverlap  (optional) if OodleLZ_PackedRawOverlap_Yes, the compressed data is in the same memory array as the output raw data
+    $:writeToFile           (optional) OodleXIOQFile to write raw data to
+    $:writeToFileStartPos   (optional) file position where writeToFile should start (must be $OODLEX_IO_MAX_ALIGNMENT aligned)
+    $:pWriteHandleGroup     (optional) if writeToFile is given, this is filled with an OodleAsyncGroup OodleXHandle containing all the file IO operations
+    $:autoDelete            (optional) see $OodleXHandleAutoDelete
+    $:dependencies      (optional) dependencies; the async op won't start until these are all complete; note : these are not freed, they must be autodelete or you must free them some other way.
+    $:numDependencies   (optional)  number of handles in _deps_ array
+    $:return        OodleXHandle to the operation, or OodleXHandle_Null for invalid arguments
+    
+    Start an async LZ decompress with the runner specified in asyncSelect.
+
+    The entire Read-Decomp-Write is done with maximum parallelism; reads are done in chunks, as each compressed chunk is available,
+    it is decompressed, and as raw chunks are done, they are written.
+
+    A note on the alignment of _packedDataPtr_ and _packedDataStartPos_ : the simplest way is if both are $OODLEX_IO_MAX_ALIGNMENT.
+    However, if the packed data starts some non-aligned way into the file, ensure the misalignment of both is the same.  This is
+    automatic if you allocate a buffer to correspond to the whole file, or start your read at the preceding aligned position.
+
+    If you have the data already read into memory, use $OodleXLZ_Decompress_Wide_Async instead.
+    
+    To use $OodleLZ_PackedRawOverlap_Yes , make a buffer of size at least $OodleLZ_GetInPlaceDecodeBufferSize ; you then read
+    the compressed data in to the *end* of that array, and decompress with the raw pointer set to the *front* of that array.
+    This lets you avoid allocating two large arrays.  It does hurt parallelism.
+
+    _rawArray_ and _packedDataPtr_ memory blocks passed to this function must be kept alive for the duration of the async.
+    
+    To use this function, you should have stored the _seekTable_ for the compressed data in a file.
+    
+    NOTE !! : if _writeToFile_ is provided, the writes are async and are NOT necessarily done when
+    the returned handle is done; the returned handle is for the decompress.  They are done when the
+    handle in *pWriteHandleGroup is done.  You must not free the buffer being written until *pWriteHandleGroup is done.
+*/
+            
+OOFUNC1 OodleXHandle OOFUNC2 OodleXLZ_Decompress_Wide_Async(
+            OO_U32 asyncSelect,
+            const OodleLZ_SeekTable * seekTable,
+            const void * packedDataPtr,OO_SINTa packedLen,
+            void * rawArray,OO_SINTa rawArrayLen,
+            OodleLZ_FuzzSafe fuzzSafe OODEFAULT(OodleLZ_FuzzSafe_No),
+            OodleLZ_CheckCRC checkCRC OODEFAULT(OodleLZ_CheckCRC_No),
+            OodleLZ_Verbosity verbosity OODEFAULT(OodleLZ_Verbosity_None),
+            void * decBufBase OODEFAULT(NULL),OO_SINTa decBufSize OODEFAULT(0),
+            OodleLZ_PackedRawOverlap packedRawOverlap OODEFAULT(OodleLZ_PackedRawOverlap_No),
+            OodleXIOQFile writeToFile OODEFAULT(0),
+            OO_S64 writeToFileStartPos OODEFAULT(0),
+            OodleXHandle * pWriteHandleGroup OODEFAULT(0),
+            OodleXHandleAutoDelete autoDelete OODEFAULT(OodleXHandleAutoDelete_No),
+            const OodleXHandle * dependencies OODEFAULT(NULL),OO_S32 numDependencies OODEFAULT(0));
+/* Start an async LZ decompress, possibly write raw data
+
+    $:asyncSelect   logical OR of $OodleXAsyncSelect flags determine how the async is run
+    $:seekTable     seek locations as created by $OodleLZ_CreateSeekTable
+    $:packedDataPtr pointer to LZ compressed data
+    $:packedLen     compressed data length
+    $:rawArray      pointer to memory filled with decompressed data
+    $:rawArrayLen   length of decompressed data
+    $:checkCRC      if OodleLZ_CheckCRC_Yes, the decompressor checks the crc to ensure data integrity
+    $:verbosity     (optional) if not OodleLZ_Verbosity_None, will log some information
+    $:decBufBase    (optional) if not NULL, provides preceding data to prime the dictionary; must be contiguous with rawBuf, the data between the pointers _dictionaryBase_ and _rawBuf_ is used as the preconditioning data.   The exact same precondition must be passed to encoder and decoder.
+    $:decBufSize    (optional) size of circular buffer starting at decBufBase   
+    $:packedRawOverlap  (optional) if OodleLZ_PackedRawOverlap_Yes, the compressed data is in the same memory array as the output raw data
+    $:writeToFile           (optional) OodleXIOQFile to write raw data to
+    $:writeToFileStartPos   (optional) file position where writeToFile should start (must be $OODLEX_IO_MAX_ALIGNMENT aligned)
+    $:pWriteHandleGroup     (optional) if writeToFile is given, this is filled with an OodleAsyncGroup OodleXHandle containing all the file IO operations
+    $:autoDelete            (optional) see $OodleXHandleAutoDelete
+    $:dependencies      (optional) dependencies; the async op won't start until these are all complete; note : these are not freed, they must be autodelete or you must free them some other way.
+    $:numDependencies   (optional)  number of handles in _deps_ array
+    $:return        OodleXHandle to the operation, or OodleXHandle_Null for invalid arguments
+    
+    
+    Same as $OodleXLZ_ReadAndDecompress_Wide_Async, except this API doesn't include the option to read
+    the packed data, it must be already fully loaded.
+
+*/
+
+OOFUNC1 OodleXHandle OOFUNC2 OodleXLZ_Decompress_MakeSeekTable_Wide_Async(
+            OO_U32 asyncSelect,
+            OO_S32 seekChunkLen,
+            const void * packedDataPtr,OO_SINTa packedLen,
+            void * rawArray,OO_SINTa rawArrayLen,
+            OodleLZ_FuzzSafe fuzzSafe OODEFAULT(OodleLZ_FuzzSafe_No),
+            OodleLZ_CheckCRC checkCRC OODEFAULT(OodleLZ_CheckCRC_No),
+            OodleLZ_Verbosity verbosity OODEFAULT(OodleLZ_Verbosity_None),
+            void * decBufBase OODEFAULT(NULL),OO_SINTa decBufSize OODEFAULT(0),
+            OodleLZ_PackedRawOverlap packedRawOverlap OODEFAULT(OodleLZ_PackedRawOverlap_No),
+            OodleXIOQFile writeToFile OODEFAULT(0),
+            OO_S64 writeToFileStartPos OODEFAULT(0),
+            OodleXHandle * pWriteHandleGroup OODEFAULT(0),
+            OodleXHandleAutoDelete autoDelete OODEFAULT(OodleXHandleAutoDelete_No),
+            const OodleXHandle * dependencies OODEFAULT(NULL),OO_S32 numDependencies OODEFAULT(0));
+/* Start an async LZ decompress, possibly write raw data
+
+    $:asyncSelect   logical OR of $OodleXAsyncSelect flags determine how the async is run
+    $:seekChunkLen  length of seekChunks used in compression $(OodleLZ_CompressOptions:seekChunkLen)
+    $:packedDataPtr pointer to LZ compressed data
+    $:packedLen     compressed data length
+    $:rawArray      pointer to memory filled with decompressed data
+    $:rawArrayLen   length of decompressed data
+    $:checkCRC      if OodleLZ_CheckCRC_Yes, the decompressor checks the crc to ensure data integrity
+    $:verbosity     (optional) if not OodleLZ_Verbosity_None, will log some information
+    $:decBufBase    (optional) if not NULL, provides preceding data to prime the dictionary; must be contiguous with rawBuf, the data between the pointers _dictionaryBase_ and _rawBuf_ is used as the preconditioning data.   The exact same precondition must be passed to encoder and decoder.
+    $:decBufSize    (optional) size of circular buffer starting at decBufBase   
+    $:packedRawOverlap  (optional) if OodleLZ_PackedRawOverlap_Yes, the compressed data is in the same memory array as the output raw data
+    $:writeToFile           (optional) OodleXIOQFile to write raw data to
+    $:writeToFileStartPos   (optional) file position where writeToFile should start (must be $OODLEX_IO_MAX_ALIGNMENT aligned)
+    $:pWriteHandleGroup     (optional) if writeToFile is given, this is filled with an OodleAsyncGroup OodleXHandle containing all the file IO operations
+    $:autoDelete            (optional) see $OodleXHandleAutoDelete
+    $:dependencies      (optional) dependencies; the async op won't start until these are all complete; note : these are not freed, they must be autodelete or you must free them some other way.
+    $:numDependencies   (optional)  number of handles in _deps_ array
+    $:return        OodleXHandle to the operation, or OodleXHandle_Null for invalid arguments
+    
+    
+    Same as $OodleXLZ_Decompress_Wide_Async , but makes the seek table for you.
+    
+    Can be used as a drop-in replacement for OodleLZ_Decompress() but with parallel decoding.
+    
+    If the data is not parallel-decodable (because it has no seek resets, eg.
+    $(OodleLZ_CompressOptions:seekChunkReset) was not set) this is slower than just calling OodleLZ_Decompress.
+    So this should only be used when you believe parallel decoding is possible.
+    
+    _seekChunkLen_ most follow the rules for Oodle seek chunk lengths.  See $OodleLZ_MakeSeekChunkLen.
+    It should be a power of two and greater-equal than $OODLELZ_BLOCK_LEN.
+*/
+
+//===================================
+
+OOFUNC1 OodleXHandle OOFUNC2 OodleXLZ_Compress_Async(OO_U32 asyncSelect,OodleLZ_Compressor compressor,
+                                                    const void * rawBuf,OO_SINTa rawLen,
+                                                    void * compBuf,
+                                                    OodleLZ_CompressionLevel compressSelect,const OodleLZ_CompressOptions * pOptions OODEFAULT(NULL),
+                                                    const void * dictionaryBase OODEFAULT(NULL),
+                    OodleXHandleAutoDelete autoDelete OODEFAULT(OodleXHandleAutoDelete_No),
+                    const OodleXHandle * dependencies OODEFAULT(NULL),OO_S32 numDependencies OODEFAULT(0));
+/* Start an async LZ compress
+
+    $:asyncSelect   logical OR of $OodleXAsyncSelect flags determine how the async is run
+    $:compressor    A member of $OodleLZ_Compressor to select the compressor
+    $:rawBuf        raw data to compress
+    $:rawLen        amount of raw data to compress
+    $:compBuf       output compressed data
+    $:compressSelect        A member of $OodleLZ_CompressionLevel to select the compression level
+    $:pOptions      (optional) compression options
+    $:dictionaryBase    (optional) if not NULL, provides preceding data to prime the dictionary; must be contiguous with rawBuf, the data
+                            between the pointers _dictionaryBase_ and _rawBuf_ is used as the preconditioning data
+    $:autoDelete            (optional) see $OodleXHandleAutoDelete
+    $:dependencies      (optional) dependencies; the async op won't start until these are all complete; note : these are not freed, they must be autodelete or you must free them some other way.
+    $:numDependencies   (optional)  number of handles in _deps_ array
+    $:return        OodleXHandle to the operation, or OodleXHandle_Null for invalid arguments
+                            
+    Runs "wide" if _asyncFlags_ includes $OodleXAsyncSelect_Workers + $OodleXAsyncSelect_Wide.
+    
+    The output compressed data can be decompressed "wide" if _pOptions_ set seekChunkReset = true.
+
+    "wide" means use many threads at once on this single operation.  If the compression is selected to run wide, but the decompression
+    cannot run wide (eg. the compressed data does not have small independent chunks), then it will still compress wide, but on
+    a very large granularity, instead of the small $OODLELZ_BLOCK_LEN granularity).  In that case, only very large buffers will
+    be compressed in parallel.
+    
+    Use $OodleXLZ_Compress_Wait_GetResult to get the result and free associated structures.
+    $OodleXLZ_Compress_Wait_GetResult must be called even if you don't want the result, or memory will be leaked.
+*/
+                                                    
+OOFUNC1 OodleXStatus OOFUNC2 OodleXLZ_Compress_Wait_GetResult(OodleXHandle h, OO_SINTa * pCompLen );
+/* Wait, get result, and delete the handle
+
+    $:h     the handle from $OodleXLZ_Compress_Async
+    $:pCompLen  filled with the compressed len
+    $:return    OodleXStatus_Done for success
+
+*/
+
+/*
+OOFUNC1 OO_SINTa OOFUNC2 OodleXLZ_Compress_Wide_CompBufSizeNeeded(OO_SINTa rawLen,  
+                    OodleLZ_Compressor compressor,
+                    OodleLZ_CompressionLevel selection,
+                    const OodleLZ_CompressOptions * pOptions );
+*/                  
+                    
+OOFUNC1 OO_SINTa OOFUNC2 OodleXLZ_Compress_AsyncAndWait(OO_U32 asyncSelect,OodleLZ_Compressor compressor,
+                                                    const void * rawBuf,OO_SINTa rawLen,
+                                                    void * compBuf,
+                                                    OodleLZ_CompressionLevel compressSelect,const OodleLZ_CompressOptions * pOptions OODEFAULT(NULL),
+                                                    const void * dictionaryBase OODEFAULT(NULL));
+/* Does $OodleXLZ_Compress_Async and $OodleXLZ_Compress_Wait_GetResult
+
+*/
+
+OOFUNC1 OodleXHandle OOFUNC2 OodleXLZ_ReadAndDecompress_Stream_Async(OO_U32 asyncSelect,
+                                                const void * packedDataPtr,OO_SINTa packedLen,
+                                                void * rawPtr,OO_SINTa rawChunkLen,
+                                                OodleLZ_FuzzSafe fuzzSafe,
+                                                OodleLZ_CheckCRC checkCRC, 
+                                                OodleLZ_Verbosity verbosity,
+                                                OodleDecompressCallback * pcb, void * pcbData,
+                                                OodleXIOQFile readFile, void * readBuf,OO_S64 readStartPos, 
+                                                OodleXHandle readPending, OO_SINTa alreadyReadSize,
+                                                OodleXHandleAutoDelete autoDelete OODEFAULT(OodleXHandleAutoDelete_No),
+                                                const OodleXHandle * dependencies OODEFAULT(NULL),OO_S32 numDependencies OODEFAULT(0));
+/* Start an async op to incrementally stream in data and decompress
+
+    $:asyncSelect   logical OR of $OodleXAsyncSelect flags determine how the async is run (but Wide is ignored, this func is always narrow)
+    $:packedDataPtr pointer to start of compressed data
+    $:packedLen     length of compressed data
+    $:rawPtr        pointer to memory to decompress into
+    $:rawChunkLen   lenght of raw data to decompress
+    $:checkCRC      if OodleLZ_CheckCRC_Yes, the decompressor checks the crc to ensure data integrity
+    $:verbosity     if not OodleLZ_Verbosity_None, will log some information
+    $:pcb           OodleDecompressCallback called during decompression (NULL for none)
+    $:pcbData       user data passed to pcb
+    $:readFile      IOQ file to read compressed data from (0 for none)
+    $:readBuf       pointer to memory where the reads from readFile should go (must be IO aligned)
+    $:readStartPos  file position where readBuf starts (must be IO aligned)
+    $:readPending   handle to previously fired read on the IOQ file
+    $:alreadyReadSize  the number of bytes of readBuf that are already read (not the number in packedDataPtr)
+    $:autoDelete            (optional) see $OodleXHandleAutoDelete
+    $:dependencies      (optional) dependencies; the async op won't start until these are all complete; note : these are not freed, they must be autodelete or you must free them some other way.
+    $:numDependencies   (optional)  number of handles in _deps_ array
+    $:return            handle to the operation
+
+    OodleLZ_Async_Decompress_ReadStream :
+        coroutine streaming LZ decoder ;
+        does incremental file reads (optionally - readFile can be zero) ;
+        calls back decode progress so you can do incremental writes (or whatever) ;
+        does not need a seek table (reads raw LZ data).
+
+    OodleLZ_Async_Decompress_ReadStream is not "wide" (only one thread is used).  It can be used to
+    overlap IO with decompression, but doesn't multi-thread decompression, even if the LZ data has seek chunks.
+
+    OodleXLZ_ReadAndDecompress_Stream_Async reads raw LZ data.
+    
+    OodleXLZ_ReadAndDecompress_Stream_Async is mainly used when you want small granularity incremental callbacks;
+    if you only need $OODLELZ_BLOCK_LEN callbacks, then $OodleXLZ_Decompress_Narrow_Async is generally better,
+    and $OodleXLZ_ReadAndDecompress_Wide_Async is fastest if you want "Wide" async decompression.
+    
+    _packedDataPtr_ should be somewhere inside _readBuf_ (if the packed data is at the start of the file, they are equal).
+    That is, (_packedDataPtr_ - _readBuf_ + _readStartPos_) is the position in the file where compressed data starts.
+    Note that _readBuf_ and _readStartPos_ must be IO aligned, but _packedDataPtr_ does not need to be, so to read compressed
+    data from a non-aligned 
+    
+    If provided, the $OodleDecompressCallback is called as quanta of raw data are available.  The callback
+    may be called more often than $OODLELZ_BLOCK_LEN granularity.
+
+    Set $OodleDecompressCallback to OodleDecompressCallback_WriteFile to perform a streaming read-compress-write.
+
+*/
+
+//=======================================
+
+typedef OOSTRUCT OodleDecompressCallback_WriteFile_Data
+{
+    OodleXIOQFile   file;           // the file handle to write to
+    OodleXHandle    lastWriteH;     // handle to the last write operation ; it's autodelete
+    OodleXHandle    closeH;         // handle to the file close operation ; NOT autoDelete
+    OO_SINTa        written;        // number of bytes written so far
+    OO_BOOL         doCloseFile;    // should the file be closed after the last write?
+} OodleDecompressCallback_WriteFile_Data;
+/* A $OodleDecompressCallback_WriteFile_Data for use with $OodleDecompressCallback_WriteFile
+
+    The $OodleDecompressCallback_WriteFile_Data struct is passed as "userdata" to $OodleDecompressCallback_WriteFile.
+    
+    You must supply one as _pcbData_ in functions that take a decompression callback.
+    
+    Warning : if you make this object on the stack, ensure the lifetime is sufficient for the async operation!
+*/
+
+OOFUNC1 void OOFUNC2 OodleXDecompressCallback_WriteFile_Data_Init(OodleDecompressCallback_WriteFile_Data * pcbData,
+    const char * fileName,
+    OO_BOOL closeFileAfterWriting,
+    OO_SINTa reserveSize OODEFAULT(0));
+/* fills out an $OodleDecompressCallback_WriteFile_Data struct
+
+    $:pcbData                       The $OodleDecompressCallback_WriteFile_Data to fill
+    $:fileName                  The name of the file to write to (will be opened)
+    $:closeFileAfterWriting     Should the file be closed for you after the last write
+    $:reserveSize               (optional) size to reserve
+    
+    Fills out the _pcbData_ for use with OodleDecompressCallback_WriteFile.
+    
+    Opens _fileName_ for write with $OodleXIOQ_OpenForWriteCreate_Async.
+    
+*/
+    
+OodleDecompressCallbackRet OODLE_CALLBACK OodleDecompressCallback_WriteFile(void * pcbData, const OO_U8 * rawBuf,OO_SINTa rawLen,const OO_U8 * compBuf,OO_SINTa compBufferSize , OO_SINTa rawDone, OO_SINTa compUsed);
+/* A $OodleDecompressCallback which writes the decompressed data to a file
+
+    $:pcbData   the $OodleDecompressCallback_WriteFile_Data you passed to OodleLZ_Decompress 
+    $:rawBuf    the decompressed buffer
+    $:rawLen    the total decompressed length
+    $:compBuf   the compressed buffer
+    $:compBufferSize  the total compressed length
+    $:rawDone   number of bytes in rawBuf decompressed so far
+    $:compUsed  number of bytes in compBuf consumed so far
+
+    OodleDecompressCallback is called incrementally during decompression.
+    
+    This is provided as a convenience for use as an $OodleDecompressCallback in functions that take that callback,
+    such as $OodleXLZ_ReadAndDecompress_Stream_Async.
+    
+    NOTE : you typically need to do OodleX_WaitAndDelete on the _closeH_ from OodleDecompressCallback_WriteFile_Data
+*/
+
+//===================================
+
+OOFUNC1 OO_S32 OOFUNC2 OodleX_IOAlignUpS32(const OO_S32 x);
+/* Align up to $OODLEX_IO_MAX_ALIGNMENT
+
+    $:x value to align
+    $:return aligned value
+    
+    Align x up to $OODLEX_IO_MAX_ALIGNMENT
+ */
+
+OOFUNC1 OO_S64 OOFUNC2 OodleX_IOAlignUpS64(const OO_S64 x);
+/* Align up to $OODLEX_IO_MAX_ALIGNMENT
+
+    $:x value to align
+    $:return aligned value
+    
+    Align x up to $OODLEX_IO_MAX_ALIGNMENT
+ */
+  
+OOFUNC1 OO_SINTa OOFUNC2 OodleX_IOAlignUpSINTa(const OO_SINTa x);
+/* Align up to $OODLEX_IO_MAX_ALIGNMENT
+
+    $:x value to align
+    $:return aligned value
+    
+    Align x up to $OODLEX_IO_MAX_ALIGNMENT
+ */
+ 
+OOFUNC1 OO_S32 OOFUNC2 OodleX_IOAlignDownS32(const OO_S32 x);
+/* Align down to $OODLEX_IO_MAX_ALIGNMENT
+
+    $:x value to align
+    $:return aligned value
+    
+    Align x down to $OODLEX_IO_MAX_ALIGNMENT
+ */
+ 
+OOFUNC1 OO_S64 OOFUNC2 OodleX_IOAlignDownS64(const OO_S64 x);
+/* Align down to $OODLEX_IO_MAX_ALIGNMENT
+
+    $:x value to align
+    $:return aligned value
+    
+    Align x down to $OODLEX_IO_MAX_ALIGNMENT
+ */
+
+OOFUNC1 OO_SINTa OOFUNC2 OodleX_IOAlignDownSINTa(const OO_SINTa x);
+/* Align down to $OODLEX_IO_MAX_ALIGNMENT
+
+    $:x value to align
+    $:return aligned value
+    
+    Align x down to $OODLEX_IO_MAX_ALIGNMENT
+ */
+ 
+OOFUNC1 OO_SINTa OOFUNC2 OodleX_S64_to_SINTa_check(const OO_S64 x);
+/* Convert OO_S64 to OO_SINTa and check
+
+    Used for loading 64-bit file sizes into memory buffers.
+    Converts type and checks that file size fits in memory.
+    
+    In 64 bit builds, this is a no-op.
+    On 32 bit builds it ensures you to don't lose bits in the cast accidentally.
+    
+    See also $Oodle_FAQ_S64_And_SINTa.
+*/
+
+OOFUNC1 OO_F64 OOFUNC2 OodleX_GetSeconds();
+/* Get current time in seconds
+
+    $:return    current time in seconds
+
+    What it do.
+*/
+
+//=================================
+// All Oodle file names are UTF8
+//  use these for conversion if you like :
+OOFUNC1 OO_BOOL OOFUNC2 OodleXUtil_ConvertUTF8ToUTF16(const char * from, OO_U16 * to, int toSize );
+
+OOFUNC1 OO_BOOL OOFUNC2 OodleXUtil_ConvertUTF16ToUTF8(const OO_U16 * from, char * to, int toSize );
+
+//=======================================================
+// OodleLog_ logging support
+//
+// log function bits :
+//  combine with OR to create a OO_U32 log state for OodleXLog_SetState()
+typedef enum OodleXLog_StateFlags
+{
+    OODLEXLOG_TO_FILE               =((OO_U32)1<<0), // log to the log file
+    OODLEXLOG_ECHO                  =((OO_U32)1<<1), // echo to a stdio file (stdout/stderr typically)
+    OODLEXLOG_TO_DEBUGGER           =((OO_U32)1<<2), // log to the debugger
+    OODLEXLOG_FILE_LINE             =((OO_U32)1<<3), // put file & line on all logs
+    OODLEXLOG_CALLBACK              =((OO_U32)1<<4), // log to the user-provided callback
+    OODLEXLOG_PREFIX_THREAD_TIME    =((OO_U32)1<<5), // prefix the thread id & time
+    OODLEXLOG_AUTOFLUSH_THREADLOG   =((OO_U32)1<<6), // flush the threadlog to the primary log automatically
+    OODLEXLOG_FLUSH_EVERY_WRITE     =((OO_U32)1<<7),  // flush log file after every write, useful for debugging crashes
+    OODLEXLOG_STATE_VERBOSITY_NONE  =((OO_U32)0<<16), // verbosity in state
+    OODLEXLOG_STATE_VERBOSITY0      =((OO_U32)1<<16),
+    OODLEXLOG_STATE_VERBOSITY1      =((OO_U32)2<<16),
+    OODLEXLOG_STATE_VERBOSITY2      =((OO_U32)3<<16)
+} OodleXLog_StateFlags;
+/* Flags for use with $OodleXLog_SetState
+*/
+
+OO_COMPILER_ASSERT( sizeof(OodleXLog_StateFlags) == 4 );
+
+typedef enum OodleXLog_VerboseLevel
+{
+    OodleXLog_Verbose_None = -1,        // log nothing
+    OodleXLog_Verbose_Minimal = 0,  // log only very important messages, such as errors
+    OodleXLog_Verbose_Some = 1,     // default setting during development
+    OodleXLog_Verbose_Lots = 2,     // log lots; may be slow (note: these are compiled out in release builds)
+    OodleXLog_Verbose_Force32 = 0x40000000
+} OodleXLog_VerboseLevel;
+/* Standard verbosity levels for use with $OodleXLog_SetVerboseLevel
+*/
+
+// set the state bits to enable/disable various functions
+// SetState(0) disables all
+OOFUNC1 void OOFUNC2 OodleXLog_SetState(OO_U32 options);
+
+OOFUNC1 OO_U32 OOFUNC2 OodleXLog_GetState();
+
+// set where the echo goes (usually stdout,stderr, or NULL)
+// note : you must also turn on RR_LOG_ECHO in state to actually get logs to the echo file
+ // set to NULL to disable echo
+OOFUNC1 void OOFUNC2 OodleXLog_SetEcho(void * echo);
+
+OOFUNC1 void * OOFUNC2 OodleXLog_GetEcho();
+
+// callback OO_BOOL tells you whether to echo the log to other outputs; a false supresses all output
+typedef enum OodleXLogCallbackRetRet
+{
+    OodleXLogCallbackRetRet_Continue = 1,  // output to other log States
+    OodleXLogCallbackRetRet_Terminate = 0,  // suppress further logging of this message
+    OodleXLogCallbackRetRet_Force32 = 0x40000000
+} OodleXLogCallbackRetRet;
+/* Return value for $OodleXLogCallbackRet */
+
+typedef OodleXLogCallbackRetRet (OODLE_CALLBACK OodleXLogCallbackRet) (const char * buffer);
+/* Function pointer for $OodleXLog_SetCallback
+
+    $:buffer    the log message
+    $:return    whether to supress the message or not
+    
+    OodleXLogCallbackRet is provided by the client to take log messages.
+    It is called before other log outputs so that it has the chance to return $OodleXLogCallbackRetRet_Terminate
+    and supress other output.
+*/
+
+// note : you must also turn on the callback enable bit in State to actually get logs to your callback
+OOFUNC1 void OOFUNC2 OodleXLog_SetCallback(OodleXLogCallbackRet * cb);
+
+OOFUNC1 OodleXLogCallbackRet * OOFUNC2 OodleXLog_GetCallback();
+
+// set runtime verbose level :
+OOFUNC1 int OOFUNC2 OodleXLog_GetVerboseLevel();
+
+OOFUNC1 int OOFUNC2 OodleXLog_SetVerboseLevel(int v);
+
+OOFUNC1 void OOFUNC2 OodleXLog_Flush();
+
+OOFUNC1 void OOFUNC2 OodleXLog_PrintfError(OodleXError err);
+
+OOFUNC1 void OOFUNC2 OodleXLog_Printf_Raw(int verboseLevel,const char * file,int line,const char * fmt,...);
+
+#define OodleXLog_Printf(verboseLevel,...)      OodleXLog_Printf_Raw(verboseLevel,__FILE__,__LINE__,##__VA_ARGS__)  /* OodleXLog_Printf lets you write to Oodle's log.
+
+   Use like printf : OodleXLog_Printf(verbose,fmt,arg1,arg2,...)
+
+    What kind of output is produced from this depends on the bit flags set in $OodleXLog_SetState.
+
+   If the global verbose level set by $OodleXLog_SetVerboseLevel is < _verboseLevel passed here,
+   the message is supressed.
+
+   OodleXLog_Printf_vN(fmt,...) is the same as OodleXLog_Printf(N,fmt,...)
+*/
+
+#define OodleXLog_Printf_v0(...)    OodleXLog_Printf(0,##__VA_ARGS__)
+#define OodleXLog_Printf_v1(...)    OodleXLog_Printf(1,##__VA_ARGS__)
+#define OodleXLog_Printf_v2(...)    OodleXLog_Printf(2,##__VA_ARGS__)
+
+OOFUNC1 OO_BOOL OOFUNC2 OodleX_DisplayAssertion(const char * fileName,const int line,const char * function,const char * message);
+
+//=====================================================
+
+
+typedef OOSTRUCT OodleXConfigValues
+{
+    OO_S32 m_Oodle_DefaultIOBufferSize; // the buffer size to use when none is given
+    OO_S32 m_Oodle_DefaultWriteReserveSize;// default size to reserve in files opened for write, if none is given
+    OO_S32 m_Oodle_MaxSingleIOSize;     // the maximum IO size to submit to the system; larger IO's than this are broken into several pieces; this allows other IO's to interleave, and also prevents heavy loads on kernel resources
+    
+    OO_S32 m_OodleIOQStream_MaxReadSize;   // IOQStream doesn't read larger than this (unless a client is blocking on needing more than these bytes immediately).  Smaller MaxReadSize reduces IOQStream service latency, but also reduces max throughput
+    OO_S32 m_OodleIOQStream_MinReadSize;    // IOQStream tries not to read less than this in a single IO op (unless a client is blocking or we're at EOF or the loop point).  
+    OO_S32 m_OodleIOQStream_OffsetAlignment; // IOQStream tries to align all its reads this granularity; some platforms are much faster if the position of IO ops are aligned to large sectors (eg. on the PS3 DVD)
+
+    OO_S32 m_Oodle_very_long_wait_seconds; // seconds to consider "very long" and warn about possible deadlock
+    
+    OO_S32 m_deprecated_Desired_Parallel_BranchFactor; // number of buffer splits for parallel compress
+
+    OO_BOOL m_Oodle_OSFileOpen_Default_Read_Buffered;   // should files opened with $OodleXFileOpenFlags_Default for Read be buffered or not?
+    OO_BOOL m_Oodle_OSFileOpen_Default_Write_Buffered;  // should files opened with $OodleXFileOpenFlags_Default for Write be buffered or not?
+
+    OO_BOOL m_Oodle_PathsCaseSensitive;  // are paths compared case-sensitive or not?  Defaults to the per-platform value OODLEX_PLATFORM_CASE_SENSITIVE.
+
+    OO_U32  m_oodle_header_version; // = OODLE_HEADER_VERSION
+    
+} OodleXConfigValues;
+/* OodleXConfigValues
+
+    Struct of user-settable low level config values.  See $OodleX_SetConfigValues.
+
+    May have different defaults per platform.
+*/
+
+OOFUNC1 void OOFUNC2 OodleX_GetConfigValues(OodleXConfigValues * ptr);
+/* Get $OodleXConfigValues
+
+    $:ptr   filled with OodleXConfigValues
+    
+    Gets the current $OodleXConfigValues.
+
+    May be different per platform.
+*/
+
+OOFUNC1 void OOFUNC2 OodleX_SetConfigValues(const OodleXConfigValues * ptr);
+/* Set $OodleXConfigValues
+
+    $:ptr   your desired OodleXConfigValues
+    
+    Sets the global $OodleXConfigValues from your struct.
+
+    You should call $OodleX_GetConfigValues to fill the struct, then change the values you
+    want to change, then call $OodleX_SetConfigValues.
+
+    This should generally be done before doing anything with Oodle (eg. even before OodleX_Init).
+    Changing OodleXConfigValues after Oodle has started has undefined effects.
+*/
+
+typedef OO_U32 OodleX_Semaphore;
+/* Semaphore ; initialize with = 0 , no cleanup necessary
+
+    NOTE : it is not intended that you use these in production.  They are for use in the Oodle
+    examples.  Replace with your own thread functions for shipping.
+*/
+
+OOFUNC1 void OOFUNC2 OodleX_Semaphore_Post(OodleX_Semaphore * sem, OO_S32 count OODEFAULT(1));
+/* OodleX_Semaphore_Post
+
+    NOTE : it is not intended that you use these in production.  They are for use in the Oodle
+    examples.  Replace with your own thread functions for shipping.
+*/
+
+OOFUNC1 void OOFUNC2 OodleX_Semaphore_Wait(OodleX_Semaphore * sem);
+/* OodleX_Semaphore_Wait
+
+    NOTE : it is not intended that you use these in production.  They are for use in the Oodle
+    examples.  Replace with your own thread functions for shipping.
+*/
+
+typedef OO_U32 (OODLE_CALLBACK OodleX_ThreadFunc)(void * userdata);
+/* User-provided callback for threads
+
+    NOTE : it is not intended that you use these in production.  They are for use in the Oodle
+    examples.  Replace with your own thread functions for shipping.
+*/
+
+typedef void * OodleX_Thread;
+
+OOFUNC1 OodleX_Thread OOFUNC2 OodleX_CreateThread(OodleX_ThreadFunc * func,void * userdata);
+/* Start a thread running threadfunc
+
+    NOTE : it is not intended that you use these in production.  They are for use in the Oodle
+    examples.  Replace with your own thread functions for shipping.
+*/
+
+OOFUNC1 void OOFUNC2 OodleX_WaitAndDestroyThread(OodleX_Thread t);
+/* Wait on thread being complete and free all resources
+
+    NOTE : it is not intended that you use these in production.  They are for use in the Oodle
+    examples.  Replace with your own thread functions for shipping.
+*/
+
+OOFUNC1 void OOFUNC2 OodleX_ReleaseThreadTLS();
+/* Release OodleX TLS resources on the calling thread
+
+    Call on a thread before it terminates to release resources that OodleX may have put in the TLS
+    of this thread.
+    
+    The purpose of this is to avoid increasing memory use in code bases that create & destroy a lot of
+    threads for jobs.  In that case, Oodle may allocate a bit of memory per thread and never free it,
+    which will add up over time.
+    
+    In normal game code bases that create a fixed number or low number of threads, you should not
+    bother calling this.
+
+    NOTE : any use of OodleX functions on this thread after calling this may crash!
+    This should be the last thing called on this thread before it terminates or returns from
+    its thread function.    
+*/
+
+//=====================================================
+/**
+
+OodleUtil is pretty heavy stuff for your tools to make things easier
+
+not fast.  not for releasing games.
+
+**/
+//-------------------------------------------------------------------------------------
+// these should generally work on vfsNames or osNames
+//
+// this is slow and synchronous :
+//  buffer is allocated with OodleXMalloc_IOAligned - you must free it !
+
+OOFUNC1 void * OOFUNC2 OodleXIOQ_ReadMallocWholeFile_AsyncAndWait(const char * vfsName,OO_S64 * pSize,
+                                                                        OodleXFileOpenFlags fileOpenFlags OODEFAULT(OodleXFileOpenFlags_Default));
+/* See $OodleXIOQ_ReadMallocWholeFile_Async
+*/
+
+OOFUNC1 OO_BOOL OOFUNC2 OodleXIOQ_WriteWholeFile_AsyncAndWait(const char * vfsName,const void * buffer,OO_SINTa size,
+                                                                        OodleXFileOpenFlags fileOpenFlags OODEFAULT(OodleXFileOpenFlags_Default));
+/* See $OodleXIOQ_OpenWriteWholeFileClose_Async
+*/
+
+OOFUNC1 OO_BOOL OOFUNC2 OodleXIOQ_CopyFile_AsyncAndWait(const char * from,const char *to,OO_U32 oodleCopyFileFlags);
+/* See $OodleXIOQ_CopyFile_Async
+*/
+
+OOFUNC1 OO_BOOL OOFUNC2 OodleXIOQ_GetInfoByName_AsyncAndWait(const char * vfsName,OodleXFileInfo * pInfo,OodleFileNotFoundIsAnError fnfiae);
+/* See $OodleXIOQ_GetInfoByName_Async
+
+    OodleXIOQ_GetInfoByName_AsyncAndWait returns true if the file was found and info was retrieved successfully.
+    The return value is always false for file not found, even if you pass OodleFileNotFoundIsAnError_No.
+*/
+
+OOFUNC1 OO_BOOL OOFUNC2 OodleXIOQ_SetInfoByName_AsyncAndWait(const char * name,OO_U32 flags,OO_U64 modTime);
+/* See $OodleXIOQ_SetInfoByName_Async
+*/
+
+// Oodle_FileHash_ReadDoubleBuffered_Async is much better :
+//OOFUNC1 OO_BOOL OOFUNC2 OodleXIOQ_ComputeFileHash_AsyncAndWait(const char * vfsName,OO_U64 * pHash);
+
+OOFUNC1 OO_BOOL OOFUNC2 OodleXIOQ_MakeAllDirs_AsyncAndWait(const char * path);
+/* See $OodleXIOQ_MakeAllDirs_Async
+*/
+
+OOFUNC1 OO_BOOL OOFUNC2 OodleXIOQ_Delete_AsyncAndWait(const char * path);
+/* See $OodleXIOQ_Delete_Async
+*/
+
+OOFUNC1 OO_BOOL OOFUNC2 OodleXIOQ_Rename_AsyncAndWait(const char * fm,const char * to,OO_BOOL overwrite);
+/* See $OodleXIOQ_Rename_Async
+*/
+
+OOFUNC1 OO_S64 OOFUNC2 OodleXIOQ_GetFileSize_AsyncAndWait(const char * vfsName,OodleFileNotFoundIsAnError fnfiae);
+/* Convenience version of $OodleXIOQ_GetInfoByName_AsyncAndWait
+    returns negative for error
+*/
+
+OOFUNC1 OO_BOOL OOFUNC2 OodleXIOQ_NameIsDir_AsyncAndWait(const char * vfsName,OodleFileNotFoundIsAnError fnfiae);
+/* Convenience version of $OodleXIOQ_GetInfoByName_AsyncAndWait
+*/
+
+//-------------------------------------------------------------------------------------
+
+OOFUNC1 const OodleXFileOpsVTable * OOFUNC2 OodleX_GetOSFileOps();
+/* Return a const $OodleXFileOpsVTable with the base OS implementations
+
+    Contains the base file ops functions for the current OS.
+
+    Do not change this struct!
+*/
+
+OOFUNC1 const OodleXFileOpsVTable * OOFUNC2 OodleX_GetDefaultFileOps();
+/* Return the current $OodleXFileOpsVTable
+
+    Contains the default file ops function vtable that is used whenever no other vtable is provided.
+
+    This begins life equal to the vtable of $OodleX_GetOSFileOps , but can be changed.
+        
+    To mutate use $OodleX_SetDefaultFileOps
+*/
+
+OOFUNC1 void OOFUNC2 OodleX_SetDefaultFileOps( const OodleXFileOpsVTable * pNewVTable );
+/* Change the default file ops vtable
+
+    _pNewVTable_ is copied into the global default file ops.
+    
+    These file ops are used by Oodle whenever no other vtable is provided.
+
+    WARNING : access to $OodleXFileOpsVTable is not thread safe.  It should generally only be done
+    at app initialization time to set your desired func pointers, and then not done thereafter.
+*/
+
+OOFUNC1 void OOFUNC2 OodleX_CombinePaths(const char *base,const char *add, char * into,OO_S32 intoSize);
+
+// has a backslash on it
+OOFUNC1 void OOFUNC2 OodleX_GetOSCwd(char * into,OO_S32 intoSize);
+
+// stick cwd in front of AddTo ; works with ".." in AddTo
+OOFUNC1 void OOFUNC2 OodleX_PrefixOSCwd(char * addTo,OO_S32 addToSize);
+
+//===========================================
+// OodleXMalloc
+
+//===========================================
+// OodleXMalloc calls that use the static/global vtable :
+//  these do leak tracking and debugging, optionally
+
+OOFUNC1 void OOFUNC2 OodleXMalloc_InstallVTable( const OodleXMallocVTable * pvt , const OodleXMallocVTable * pBaseVT );
+/*  Install the vtable that OodleX will use to allocate memory
+
+    $:pvt       pointer to the vtable to be installed (will be copied)
+    $:pBaseVT   if pvt is a layered vtable, this is the underlying alloc; if not it should be = pvt
+    
+    Sets the global vtable that will de used by the OodleXMalloc calls.  Typically let $OodleX_Init install a
+    suitable vtable for you.  If you do it manually, it must be done before any other OodleX initialization.
+    
+    WARNING : You must not change the vtable after OodleX is running; pointers allocated from the previous
+    vtable will still need to be freed and will call to the global vtable.
+*/
+
+typedef OO_BOOL (OODLE_CALLBACK OodleXMallocFailedHandler)( OO_SINTa bytes ); 
+/* OodleXMallocFailedHandler is called when a malloc fails
+    Return true to retry.
+*/
+
+OOFUNC1 void OOFUNC2 OodleXMalloc_SetFailedHandler( OodleXMallocFailedHandler * f );
+/* Install the $OodleXMallocFailedHandler that will be used
+
+    $:f the function pointer to call (can be null for none)
+    
+    
+*/
+
+OOFUNC1 void * OOFUNC2 OodleXMalloc( OO_SINTa bytes );
+/*  Allocate some memory
+
+    $:bytes     the amount to allocate (must be > 0)
+    $:return    pointer to allocated memory
+
+    OodleXMalloc uses the installed $OodleXMallocVTable.
+    Pointer will be aligned to at least $OODLE_MALLOC_MINIMUM_ALIGNMENT.
+    If a malloc fails, any installed $OodleXMallocFailedHandler will be called.
+*/
+
+
+OOFUNC1 void * OOFUNC2 OodleXMallocAligned( OO_SINTa bytes , OO_S32 alignment );
+/*  Allocate some memory with specified alignment
+
+    $:bytes     the amount to allocate (must be > 0)
+    $:alignment the desired alignment
+    $:return    pointer to allocated memory
+
+    alignment must be <= bytes.
+    alignment must be power of 2.
+    OodleXMalloc uses the installed $OodleXMallocVTable.
+    Pointer will be aligned to at least $OODLE_MALLOC_MINIMUM_ALIGNMENT.
+    If a malloc fails, any installed $OodleXMallocFailedHandler will be called.
+*/
+
+OOFUNC1 void OOFUNC2 OodleXFree(void * ptr);
+/* free a pointer allocated by $OodleXMalloc or $OodleXMallocAligned 
+
+    $:ptr   the pointer to free (must not be NULL)
+
+    Uses the current $OodleXMallocVTable ; this is an error if ptr was allocated from a different VTable.
+    Prefer $OodleXFreeSized whenever possible.
+*/
+
+OOFUNC1 void OOFUNC2 OodleXFreeSized(void * ptr, OO_SINTa bytes);
+/* Free a pointer allocated by $OodleXMalloc or $OodleXMallocAligned
+  
+  $:ptr   the pointer to free; allocated by OodleXMalloc or OodleXMallocAligned (must not be NULL)
+  $:bytes the size of the allocation as originally requested
+  
+   Providing the size of the malloc allows much faster freeing
+    Size must match the allocated size!
+    Uses the current $OodleXMallocVTable ; this is an error if ptr was allocated from a different VTable.
+ */
+
+OOFUNC1 OO_S32 OOFUNC2 OodleXMallocBigAlignment();
+/* returns the alignment of $OodleXMallocBig pointers 
+
+   $:return  the alignment of $OodleXMallocBig pointers 
+
+   Should be >= $OODLEX_IO_MAX_ALIGNMENT 
+*/
+
+OOFUNC1 void * OOFUNC2 OodleXMallocBig( OO_SINTa bytes );
+/* alloc a large block with "Big" alignment
+
+    $:bytes     the amount to allocate (must be > 0)
+    $:return    pointer to allocated memory
+
+   query the alignment via $OodleXMallocBigAlignment 
+*/
+   
+   
+OOFUNC1 void OOFUNC2 OodleXFreeBig( void * ptr );
+/* free a pointer allocated by $OodleXMallocBig 
+
+    $:ptr       pointer to free (must not be NULL)
+
+    You cannot call $OodleXFree on a pointer allocated by $OodleXMallocBig.
+    Uses the current $OodleXMallocVTable ; this is an error if ptr was allocated from a different VTable.   
+*/
+
+OOFUNC1 OO_BOOL OOFUNC2 OodleXMalloc_ValidatePointer(const void * ptr, OO_SINTa bytes);
+/* debug check if a pointer is a valid malloc
+
+    $:ptr       pointer to validate
+    $:bytes     size of allocation if known; -1 if not
+    $:return    true if the malloc headers are all okay
+
+   Should work on $OodleXMalloc and $OodleXMallocBig pointers.
+   Bytes can be -1 if unknown, but there will be less validation checks.
+   ValidatePointer is most useful if the OodleXMalloc debug thunk layer is installed in $OodleX_Init.
+*/
+
+OOFUNC1 void * OOFUNC2  OodleXMalloc_IOAligned(OO_SINTa size); 
+/*   OodleXMalloc_IOAligned result is guaranteed to be aligned to $OODLEX_IO_MAX_ALIGNMENT
+
+  $:size  bytes to allocate ; will be aligned up to $OODLEX_IO_MAX_ALIGNMENT !
+  $:return pointer to the allocated memory
+  
+  OodleXMalloc_IOAligned should be used to get memory that can be used in OodleIOQ and other places
+  that require disk-aligned pointers.
+  OodleXMalloc_IOAligned may just pass through to $OodleXMallocBig provided by the client, or it may not
+  if the $OodleXMallocBigAlignment is very large.
+*/
+OOFUNC1 void OOFUNC2    OodleXFree_IOAligned(void * ptr);
+/* Free a pointer allocated with $OodleXMalloc_IOAligned 
+
+  $:ptr pointer to free
+*/
+   
+//===========================================
+
+OOFUNC1 const OodleXMallocVTable * OOFUNC2 OodleXMalloc_GetVTable_Clib(OodleXMalloc_OS_Options options);
+
+OOFUNC1 const OodleXMallocVTable * OOFUNC2 OodleXMalloc_GetVTable_OS(OodleXMalloc_OS_Options options);
+
+OOFUNC1 OO_U64 OOFUNC2 OodleX_CorePlugin_RunJob(t_fp_Oodle_Job * fp_job,void * job_data, OO_U64 * dependencies, int num_dependencies, void * user_ptr);
+/* Function to plug in the OodleX Worker system to $OodleCore_Plugins_SetJobSystem
+
+    NOTE : OodleX_Init does $OodleCore_Plugins_SetJobSystem automatically.
+
+*/
+
+OOFUNC1 void OOFUNC2 OodleX_CorePlugin_WaitJob(OO_U64 job_handle, void * user_ptr);
+/* Function to plug in the OodleX Worker system to $OodleCore_Plugins_SetJobSystem
+
+    NOTE : OodleX_Init does $OodleCore_Plugins_SetJobSystem automatically.
+
+*/
+
+OOFUNC1 OO_S32 OOFUNC2 OodleX_GetNumWorkerThreads();
+/* Returns the number of worker threads
+
+    When there are 0 worker threads, the OodleWork system still succeeds, it just runs Worklets synchronously
+    on the calling thread.
+    
+    The worker thread count is set in $(OodleXInitOptions:m_OodleInit_Workers_Count)
+*/
+
+#endif // __OODLE2X_H_INCLUDED__
