Setup CI in Horde for Horde.Storage

Also added ability to override path to dotnet for DotnetTask, to allow us to use a system install dotnet for special cases (in this case we use a newer .NET version then the rest of the engine)

diff --git a/Engine/Source/Programs/AutomationTool/BuildGraph/Tasks/DotNetTask.cs b/Engine/Source/Programs/AutomationTool/BuildGraph/Tasks/DotNetTask.cs
--- a/Engine/Source/Programs/AutomationTool/BuildGraph/Tasks/DotNetTask.cs
+++ b/Engine/Source/Programs/AutomationTool/BuildGraph/Tasks/DotNetTask.cs
@@ -44,6 +44,12 @@ namespace BuildGraph.Tasks
 		/// </summary>
 		[TaskParameter(Optional = true)]
 		public int ErrorLevel = 1;
+
+		/// <summary>
+		/// Override path to dotnet executable
+		/// </summary>
+		[TaskParameter(Optional = true)]
+		public FileReference DotNetPath;
 	}
 
 	/// <summary>
@@ -74,7 +80,7 @@ namespace BuildGraph.Tasks
 		/// <param name="TagNameToFileSet">Mapping from tag names to the set of files they include</param>
 		public override void Execute(JobContext Job, HashSet<FileReference> BuildProducts, Dictionary<string, HashSet<FileReference>> TagNameToFileSet)
 		{
-			FileReference DotNetFile = HostPlatform.Current.GetDotnetExe();
+			FileReference DotNetFile = Parameters.DotNetPath == null ? HostPlatform.Current.GetDotnetExe() : Parameters.DotNetPath;
 			if(!FileReference.Exists(DotNetFile))
 			{
 				throw new AutomationException("DotNet is missing from {0}", DotNetFile);
diff --git a/Engine/Source/Programs/Horde/BuildHorde.xml b/Engine/Source/Programs/Horde/BuildHorde.xml
--- a/Engine/Source/Programs/Horde/BuildHorde.xml
+++ b/Engine/Source/Programs/Horde/BuildHorde.xml
@@ -28,4 +28,26 @@
 			<Docker-Build BaseDir="Engine/Source" Files="#InputFiles" Tag="hordeserver-public" DockerFile="Engine/Source/Programs/Horde/Horde.Build/Dockerfile" Arguments="--build-arg msbuild_args=&quot;$(VersionArguments)&quot;"/>
 		</Node>
 	</Agent>
+
+	<Agent Name="HordeStorage" Type="Linux">
+		<Property Name="ComposeFile" Value="Engine/Source/Programs/Horde/HordeStorage/Composes/docker-compose-tests.yml"/>
+		<Node Name="Test HordeStorage">
+			<Docker-Compose-Up File="$(ComposeFile)"/>
+			<!-- Run the unit tests, and override the dotnet sdk used to the global one for the machine rather then what is in the sync, to allow us to use newer versions then the rest of the engine has migrated to -->
+			<DotNet Arguments="test Engine/Source/Programs/Horde/HordeStorage/Horde.Storage.sln" DotNetPath="/usr/share/dotnet/dotnet"/>
+			<Docker-Compose-Down File="$(ComposeFile)"/>
+		</Node>
+
+		<Node Name="Build HordeStorage" Requires="Test HordeStorage">
+			<!-- Tag all the files that need to be staged to build -->
+			<Property Name="StagePaths">
+				Engine/Source/Programs/Horde/HordeStorage/...
+			</Property>
+			<Tag Files="$(StagePaths)" Except=".../.vs/...;.../.git/...;.../bin/...;.../obj/..." With="#InputFiles"/>
+
+			<!-- Build the Docker images -->
+			<Docker-Build BaseDir="Engine/Source/Programs/Horde/HordeStorage" Files="#InputFiles" Tag="horde_storage:$(InformationalVersion)" DockerFile="Engine/Source/Programs/Horde/HordeStorage/dockerfile" Arguments="--build-arg Project=Horde.Storage" Environment="UE_P4_CHANGELIST=(Change)"/>
+			<Docker-Build BaseDir="Engine/Source/Programs/Horde/HordeStorage" Files="#InputFiles" Tag="callisto:$(InformationalVersion)" DockerFile="Engine/Source/Programs/Horde/HordeStorage/dockerfile" Arguments="--build-arg Project=Callisto" Environment="UE_P4_CHANGELIST=(Change)" />
+		</Node>
+	</Agent>
 </BuildGraph>
