Horde: Ensure correct index name is used

Referencing by name already causing issues.

diff --git a/Engine/Source/Programs/Horde/Horde.Build/Jobs/Impl/JobCollection.cs b/Engine/Source/Programs/Horde/Horde.Build/Jobs/Impl/JobCollection.cs
--- a/Engine/Source/Programs/Horde/Horde.Build/Jobs/Impl/JobCollection.cs
+++ b/Engine/Source/Programs/Horde/Horde.Build/Jobs/Impl/JobCollection.cs
@@ -557,7 +557,7 @@ namespace HordeServer.Collections.Impl
 			List<JobDocument> Results;
 			using (IScope Scope = GlobalTracer.Instance.BuildSpan("Jobs.Find").StartActive())
 			{
-				BsonString IndexHint = new (nameof(JobDocument.CreateTimeUtc) + "_1");
+				BsonString IndexHint = new (nameof(JobDocument.CreateTimeUtc) + "_-1");
 				IFindFluent<JobDocument, JobDocument> Query = Jobs.Find<JobDocument>(Filter, new FindOptions { Hint = IndexHint }).SortByDescending(x => x.CreateTimeUtc!);
 				if (Index != null)
 				{
