Disable Pdb2Pdb on Horde agent binaries. Setting debugmode to Pdbonly instead.

diff --git a/Engine/Source/Programs/Horde/Horde.Agent/Horde.Agent.csproj b/Engine/Source/Programs/Horde/Horde.Agent/Horde.Agent.csproj
--- a/Engine/Source/Programs/Horde/Horde.Agent/Horde.Agent.csproj
+++ b/Engine/Source/Programs/Horde/Horde.Agent/Horde.Agent.csproj
@@ -9,6 +9,7 @@
     <AssemblyName>HordeAgent</AssemblyName>
     <PublishRepositoryUrl>true</PublishRepositoryUrl>
     <EmbedUntrackedSources>false</EmbedUntrackedSources>
+    <DebugType>pdbonly</DebugType>
   </PropertyGroup>
 
   <PropertyGroup Condition="'$(Configuration)|$(Platform)'=='Debug|AnyCPU'">
@@ -47,10 +48,6 @@
     <PackageReference Include="System.ServiceProcess.ServiceController" Version="4.7.0" />
   </ItemGroup>
 
-  <ItemGroup Condition="'$(FullSymbolsDir)' != ''">
-    <PackageReference Include="Microsoft.DiaSymReader.Pdb2Pdb" Version="1.1.0-beta2-20370-01" />
-  </ItemGroup>
-
   <ItemGroup>
     <Content Include="appsettings.json" Condition="Exists('appsettings.json')">
       <Generator></Generator>
@@ -141,20 +138,16 @@
     </EmbeddedResource>
   </ItemGroup>
 
-  <PropertyGroup>
-	<Pdb2PdbPath>$(NuGetPackageRoot)microsoft.diasymreader.pdb2pdb\1.1.0-beta2-20370-01\tools\Pdb2Pdb.exe</Pdb2PdbPath>
-  </PropertyGroup>
-
   <Target Name="PublishSymbols" AfterTargets="AfterBuild" Condition="'$(FullSymbolsDir)' != ''">
     <ItemGroup>
       <OutputFiles Include="$(OutDir)**" />
       <BinaryFilesBase Include="$(OutDir)\*.dll" />
-      <BinaryFiles Include="@(BinaryFilesBase)" Condition="Exists('%(RelativeDir)%(FileName).pdb') And '%(FileName)' != 'EpicGames.Perforce.Native'" />
+      <BinaryFilesBase Include="$(OutDir)\*.pdb" />
+      <BinaryFiles Include="@(BinaryFilesBase)" />
     </ItemGroup>
 
     <MakeDir Directories="$(FullSymbolsDir)" />
     <Copy SourceFiles="@(Files)" DestinationFolder="$(FullSymbolsDir)" />
-    <Exec Command="&quot;$(Pdb2PdbPath)&quot; &quot;%(BinaryFiles.Identity)&quot; /out &quot;$(FullSymbolsDir)\%(BinaryFiles.FileName).pdb&quot; /nowarn 21" IgnoreExitCode="false" />
   </Target>
   
 </Project>
diff --git a/Engine/Source/Programs/Shared/EpicGames.Horde.Client/EpicGames.Horde.Client.csproj b/Engine/Source/Programs/Shared/EpicGames.Horde.Client/EpicGames.Horde.Client.csproj
--- a/Engine/Source/Programs/Shared/EpicGames.Horde.Client/EpicGames.Horde.Client.csproj
+++ b/Engine/Source/Programs/Shared/EpicGames.Horde.Client/EpicGames.Horde.Client.csproj
@@ -9,6 +9,7 @@
     <GenerateTargetFrameworkAttribute>false</GenerateTargetFrameworkAttribute>
     <Configurations>Debug;Release;Development</Configurations>
     <Nullable>enable</Nullable>
+    <DebugType>pdbonly</DebugType>
   </PropertyGroup>
   <PropertyGroup Condition=" '$(Configuration)|$(Platform)' == 'Debug|AnyCPU' ">
     <DocumentationFile></DocumentationFile>
diff --git a/Engine/Source/Programs/Shared/EpicGames.Horde/EpicGames.Horde.csproj b/Engine/Source/Programs/Shared/EpicGames.Horde/EpicGames.Horde.csproj
--- a/Engine/Source/Programs/Shared/EpicGames.Horde/EpicGames.Horde.csproj
+++ b/Engine/Source/Programs/Shared/EpicGames.Horde/EpicGames.Horde.csproj
@@ -9,6 +9,7 @@
     <GenerateTargetFrameworkAttribute>false</GenerateTargetFrameworkAttribute>
     <Configurations>Debug;Release;Development</Configurations>
     <Nullable>enable</Nullable>
+    <DebugType>pdbonly</DebugType>
   </PropertyGroup>
   <PropertyGroup Condition=" '$(Configuration)|$(Platform)' == 'Debug|AnyCPU' ">
     <DocumentationFile></DocumentationFile>
