Horde Storage: Added missing copyright headers

diff --git a/Engine/Source/Programs/Horde/HordeStorage/Horde.Storage/Controllers/AuthController.cs b/Engine/Source/Programs/Horde/HordeStorage/Horde.Storage/Controllers/AuthController.cs
--- a/Engine/Source/Programs/Horde/HordeStorage/Horde.Storage/Controllers/AuthController.cs
+++ b/Engine/Source/Programs/Horde/HordeStorage/Horde.Storage/Controllers/AuthController.cs
@@ -1,4 +1,6 @@
-﻿using Microsoft.AspNetCore.Http;
+﻿// Copyright Epic Games, Inc. All Rights Reserved.
+
+using Microsoft.AspNetCore.Http;
 using Microsoft.AspNetCore.Mvc;
 using System;
 using System.Collections.Generic;
diff --git a/Engine/Source/Programs/Horde/HordeStorage/Horde.Storage/Controllers/ReplicationLogController.cs b/Engine/Source/Programs/Horde/HordeStorage/Horde.Storage/Controllers/ReplicationLogController.cs
--- a/Engine/Source/Programs/Horde/HordeStorage/Horde.Storage/Controllers/ReplicationLogController.cs
+++ b/Engine/Source/Programs/Horde/HordeStorage/Horde.Storage/Controllers/ReplicationLogController.cs
@@ -1,4 +1,6 @@
-﻿using System;
+﻿// Copyright Epic Games, Inc. All Rights Reserved.
+
+using System;
 using System.Collections.Generic;
 using System.ComponentModel.DataAnnotations;
 using System.Threading;
diff --git a/Engine/Source/Programs/Horde/HordeStorage/Horde.Storage/HealthChecks.cs b/Engine/Source/Programs/Horde/HordeStorage/Horde.Storage/HealthChecks.cs
--- a/Engine/Source/Programs/Horde/HordeStorage/Horde.Storage/HealthChecks.cs
+++ b/Engine/Source/Programs/Horde/HordeStorage/Horde.Storage/HealthChecks.cs
@@ -1,3 +1,5 @@
+// Copyright Epic Games, Inc. All Rights Reserved.
+
 using System;
 using System.Threading;
 using System.Threading.Tasks;
diff --git a/Engine/Source/Programs/Horde/HordeStorage/Horde.Storage/Implementation/Blob/FileSystemStore.cs b/Engine/Source/Programs/Horde/HordeStorage/Horde.Storage/Implementation/Blob/FileSystemStore.cs
--- a/Engine/Source/Programs/Horde/HordeStorage/Horde.Storage/Implementation/Blob/FileSystemStore.cs
+++ b/Engine/Source/Programs/Horde/HordeStorage/Horde.Storage/Implementation/Blob/FileSystemStore.cs
@@ -1,4 +1,6 @@
-﻿using System;
+﻿// Copyright Epic Games, Inc. All Rights Reserved.
+
+using System;
 using System.Collections.Generic;
 using System.IO;
 using System.Linq;
diff --git a/Engine/Source/Programs/Horde/HordeStorage/Horde.Storage/Implementation/Kestrel/AsciiString.cs b/Engine/Source/Programs/Horde/HordeStorage/Horde.Storage/Implementation/Kestrel/AsciiString.cs
--- a/Engine/Source/Programs/Horde/HordeStorage/Horde.Storage/Implementation/Kestrel/AsciiString.cs
+++ b/Engine/Source/Programs/Horde/HordeStorage/Horde.Storage/Implementation/Kestrel/AsciiString.cs
@@ -1,4 +1,6 @@
-﻿using System;
+﻿// Copyright Epic Games, Inc. All Rights Reserved.
+
+using System;
 using System.Text;
 
 namespace Horde.Storage.Implementation.Kestrel
diff --git a/Engine/Source/Programs/Horde/HordeStorage/Horde.Storage/Implementation/Refs/CachedRefStore.cs b/Engine/Source/Programs/Horde/HordeStorage/Horde.Storage/Implementation/Refs/CachedRefStore.cs
--- a/Engine/Source/Programs/Horde/HordeStorage/Horde.Storage/Implementation/Refs/CachedRefStore.cs
+++ b/Engine/Source/Programs/Horde/HordeStorage/Horde.Storage/Implementation/Refs/CachedRefStore.cs
@@ -1,4 +1,6 @@
-﻿using System;
+﻿// Copyright Epic Games, Inc. All Rights Reserved.
+
+using System;
 using System.Collections.Concurrent;
 using System.Collections.Generic;
 using System.Threading;
diff --git a/Engine/Source/Programs/Horde/HordeStorage/Horde.Storage/Implementation/ScyllaSessionManager.cs b/Engine/Source/Programs/Horde/HordeStorage/Horde.Storage/Implementation/ScyllaSessionManager.cs
--- a/Engine/Source/Programs/Horde/HordeStorage/Horde.Storage/Implementation/ScyllaSessionManager.cs
+++ b/Engine/Source/Programs/Horde/HordeStorage/Horde.Storage/Implementation/ScyllaSessionManager.cs
@@ -1,4 +1,6 @@
-﻿using System;
+﻿// Copyright Epic Games, Inc. All Rights Reserved.
+
+using System;
 using System.Collections.Generic;
 using System.Linq;
 using System.Threading.Tasks;
diff --git a/Engine/Source/Programs/Horde/HordeStorage/Horde.Storage/Implementation/ServiceCredentials.cs b/Engine/Source/Programs/Horde/HordeStorage/Horde.Storage/Implementation/ServiceCredentials.cs
--- a/Engine/Source/Programs/Horde/HordeStorage/Horde.Storage/Implementation/ServiceCredentials.cs
+++ b/Engine/Source/Programs/Horde/HordeStorage/Horde.Storage/Implementation/ServiceCredentials.cs
@@ -1,4 +1,6 @@
-﻿using System;
+﻿// Copyright Epic Games, Inc. All Rights Reserved.
+
+using System;
 using System.Collections.Generic;
 using System.Linq;
 using System.Threading.Tasks;
diff --git a/Engine/Source/Programs/Horde/HordeStorage/Horde.Storage/Implementation/TransactionLog/ReplicationLogSnapshotBuilder.cs b/Engine/Source/Programs/Horde/HordeStorage/Horde.Storage/Implementation/TransactionLog/ReplicationLogSnapshotBuilder.cs
--- a/Engine/Source/Programs/Horde/HordeStorage/Horde.Storage/Implementation/TransactionLog/ReplicationLogSnapshotBuilder.cs
+++ b/Engine/Source/Programs/Horde/HordeStorage/Horde.Storage/Implementation/TransactionLog/ReplicationLogSnapshotBuilder.cs
@@ -1,4 +1,6 @@
-﻿using System;
+﻿// Copyright Epic Games, Inc. All Rights Reserved.
+
+using System;
 using System.IO;
 using System.Threading;
 using System.Threading.Tasks;
diff --git a/Engine/Source/Programs/Horde/HordeStorage/Horde.Storage/Settings.cs b/Engine/Source/Programs/Horde/HordeStorage/Horde.Storage/Settings.cs
--- a/Engine/Source/Programs/Horde/HordeStorage/Horde.Storage/Settings.cs
+++ b/Engine/Source/Programs/Horde/HordeStorage/Horde.Storage/Settings.cs
@@ -1,4 +1,6 @@
-﻿using System;
+﻿// Copyright Epic Games, Inc. All Rights Reserved.
+
+using System;
 using System.Collections.Generic;
 using System.ComponentModel.DataAnnotations;
 using System.Linq;
diff --git a/Engine/Source/Programs/Horde/HordeStorage/Horde.Storage/Tests/UnitTests/CompactBinaryTests.cs b/Engine/Source/Programs/Horde/HordeStorage/Horde.Storage/Tests/UnitTests/CompactBinaryTests.cs
--- a/Engine/Source/Programs/Horde/HordeStorage/Horde.Storage/Tests/UnitTests/CompactBinaryTests.cs
+++ b/Engine/Source/Programs/Horde/HordeStorage/Horde.Storage/Tests/UnitTests/CompactBinaryTests.cs
@@ -1,4 +1,6 @@
-﻿using System;
+﻿// Copyright Epic Games, Inc. All Rights Reserved.
+
+using System;
 using System.Collections.Generic;
 using System.IO;
 using System.Linq;
diff --git a/Engine/Source/Programs/Horde/HordeStorage/Jupiter.Common/CustomMediaTypeNames.cs b/Engine/Source/Programs/Horde/HordeStorage/Jupiter.Common/CustomMediaTypeNames.cs
--- a/Engine/Source/Programs/Horde/HordeStorage/Jupiter.Common/CustomMediaTypeNames.cs
+++ b/Engine/Source/Programs/Horde/HordeStorage/Jupiter.Common/CustomMediaTypeNames.cs
@@ -1,4 +1,6 @@
-﻿namespace Jupiter
+﻿// Copyright Epic Games, Inc. All Rights Reserved.
+
+namespace Jupiter
 {
     public class CustomMediaTypeNames
     {
diff --git a/Engine/Source/Programs/Horde/HordeStorage/Jupiter.Common/Implementation/BlobContents.cs b/Engine/Source/Programs/Horde/HordeStorage/Jupiter.Common/Implementation/BlobContents.cs
--- a/Engine/Source/Programs/Horde/HordeStorage/Jupiter.Common/Implementation/BlobContents.cs
+++ b/Engine/Source/Programs/Horde/HordeStorage/Jupiter.Common/Implementation/BlobContents.cs
@@ -1,4 +1,6 @@
-﻿using System;
+﻿// Copyright Epic Games, Inc. All Rights Reserved.
+
+using System;
 using System.IO;
 using System.Threading.Tasks;
 
diff --git a/Engine/Source/Programs/Horde/HordeStorage/Jupiter.Common/Implementation/BucketId.cs b/Engine/Source/Programs/Horde/HordeStorage/Jupiter.Common/Implementation/BucketId.cs
--- a/Engine/Source/Programs/Horde/HordeStorage/Jupiter.Common/Implementation/BucketId.cs
+++ b/Engine/Source/Programs/Horde/HordeStorage/Jupiter.Common/Implementation/BucketId.cs
@@ -1,4 +1,6 @@
-﻿using System;
+﻿// Copyright Epic Games, Inc. All Rights Reserved.
+
+using System;
 using System.ComponentModel;
 using System.Globalization;
 using Newtonsoft.Json;
diff --git a/Engine/Source/Programs/Horde/HordeStorage/Jupiter.Common/Implementation/CompactBinaryObject.cs b/Engine/Source/Programs/Horde/HordeStorage/Jupiter.Common/Implementation/CompactBinaryObject.cs
--- a/Engine/Source/Programs/Horde/HordeStorage/Jupiter.Common/Implementation/CompactBinaryObject.cs
+++ b/Engine/Source/Programs/Horde/HordeStorage/Jupiter.Common/Implementation/CompactBinaryObject.cs
@@ -1,4 +1,6 @@
-﻿using System;
+﻿// Copyright Epic Games, Inc. All Rights Reserved.
+
+using System;
 using System.Buffers;
 using System.Collections;
 using System.Collections.Generic;
diff --git a/Engine/Source/Programs/Horde/HordeStorage/Jupiter.Common/Implementation/InstrumentedStream.cs b/Engine/Source/Programs/Horde/HordeStorage/Jupiter.Common/Implementation/InstrumentedStream.cs
--- a/Engine/Source/Programs/Horde/HordeStorage/Jupiter.Common/Implementation/InstrumentedStream.cs
+++ b/Engine/Source/Programs/Horde/HordeStorage/Jupiter.Common/Implementation/InstrumentedStream.cs
@@ -1,4 +1,6 @@
-﻿using System;
+﻿// Copyright Epic Games, Inc. All Rights Reserved.
+
+using System;
 using System.Collections.Generic;
 using System.IO;
 using System.Linq;
diff --git a/Engine/Source/Programs/Horde/HordeStorage/Jupiter.Common/Implementation/KeyId.cs b/Engine/Source/Programs/Horde/HordeStorage/Jupiter.Common/Implementation/KeyId.cs
--- a/Engine/Source/Programs/Horde/HordeStorage/Jupiter.Common/Implementation/KeyId.cs
+++ b/Engine/Source/Programs/Horde/HordeStorage/Jupiter.Common/Implementation/KeyId.cs
@@ -1,4 +1,6 @@
-﻿using System;
+﻿// Copyright Epic Games, Inc. All Rights Reserved.
+
+using System;
 using System.ComponentModel;
 using System.Globalization;
 using Newtonsoft.Json;
diff --git a/Engine/Source/Programs/Horde/HordeStorage/Jupiter.Common/Implementation/NamespaceId.cs b/Engine/Source/Programs/Horde/HordeStorage/Jupiter.Common/Implementation/NamespaceId.cs
--- a/Engine/Source/Programs/Horde/HordeStorage/Jupiter.Common/Implementation/NamespaceId.cs
+++ b/Engine/Source/Programs/Horde/HordeStorage/Jupiter.Common/Implementation/NamespaceId.cs
@@ -1,4 +1,6 @@
-﻿using System;
+﻿// Copyright Epic Games, Inc. All Rights Reserved.
+
+using System;
 using System.ComponentModel;
 using System.Globalization;
 using Newtonsoft.Json;
diff --git a/Engine/Source/Programs/Horde/HordeStorage/Jupiter.Common/Implementation/SecretResolver.cs b/Engine/Source/Programs/Horde/HordeStorage/Jupiter.Common/Implementation/SecretResolver.cs
--- a/Engine/Source/Programs/Horde/HordeStorage/Jupiter.Common/Implementation/SecretResolver.cs
+++ b/Engine/Source/Programs/Horde/HordeStorage/Jupiter.Common/Implementation/SecretResolver.cs
@@ -1,4 +1,6 @@
-﻿using System;
+﻿// Copyright Epic Games, Inc. All Rights Reserved.
+
+using System;
 using System.Collections.Generic;
 using System.Text.Json;
 using System.Threading.Tasks;
diff --git a/Engine/Source/Programs/Horde/HordeStorage/Jupiter.Common/Utils/OodleCompressor.cs b/Engine/Source/Programs/Horde/HordeStorage/Jupiter.Common/Utils/OodleCompressor.cs
--- a/Engine/Source/Programs/Horde/HordeStorage/Jupiter.Common/Utils/OodleCompressor.cs
+++ b/Engine/Source/Programs/Horde/HordeStorage/Jupiter.Common/Utils/OodleCompressor.cs
@@ -1,4 +1,6 @@
-﻿using System;
+﻿// Copyright Epic Games, Inc. All Rights Reserved.
+
+using System;
 using System.Reflection;
 using System.Runtime.InteropServices;
 using IntPtr = System.IntPtr;
diff --git a/Engine/Source/Programs/Horde/HordeStorage/Jupiter.Common/Utils/StreamUtils.cs b/Engine/Source/Programs/Horde/HordeStorage/Jupiter.Common/Utils/StreamUtils.cs
--- a/Engine/Source/Programs/Horde/HordeStorage/Jupiter.Common/Utils/StreamUtils.cs
+++ b/Engine/Source/Programs/Horde/HordeStorage/Jupiter.Common/Utils/StreamUtils.cs
@@ -1,4 +1,6 @@
-﻿using System;
+﻿// Copyright Epic Games, Inc. All Rights Reserved.
+
+using System;
 using System.Collections.Generic;
 using System.IO;
 using System.Threading.Tasks;
diff --git a/Engine/Source/Programs/Horde/HordeStorage/Jupiter.Common/Utils/StringUtils.cs b/Engine/Source/Programs/Horde/HordeStorage/Jupiter.Common/Utils/StringUtils.cs
--- a/Engine/Source/Programs/Horde/HordeStorage/Jupiter.Common/Utils/StringUtils.cs
+++ b/Engine/Source/Programs/Horde/HordeStorage/Jupiter.Common/Utils/StringUtils.cs
@@ -1,3 +1,5 @@
+// Copyright Epic Games, Inc. All Rights Reserved.
+
 using System;
 using System.Text;
 
