Horde Storage - Conmented away the automatic inclusion of the git hash to remove warnings when not building under Git. Also fixed some issues exanding the msbuild properties if they were not set to any values.

diff --git a/Engine/Source/Programs/Horde/HordeStorage/Horde.Storage/Controllers/StatusController.cs b/Engine/Source/Programs/Horde/HordeStorage/Horde.Storage/Controllers/StatusController.cs
--- a/Engine/Source/Programs/Horde/HordeStorage/Horde.Storage/Controllers/StatusController.cs
+++ b/Engine/Source/Programs/Horde/HordeStorage/Horde.Storage/Controllers/StatusController.cs
@@ -47,11 +47,11 @@ namespace Horde.Storage.Controllers
 
             string srcControlIdentifier = "Unknown";
             AssemblyMetadataAttribute? gitHashAttribute = attrs.FirstOrDefault(attr => attr.Key == "GitHash");
-            if (gitHashAttribute?.Value != null)
+            if (gitHashAttribute?.Value != null && !string.IsNullOrEmpty(gitHashAttribute.Value))
                 srcControlIdentifier = gitHashAttribute.Value;
 
             AssemblyMetadataAttribute? p4ChangeAttribute = attrs.FirstOrDefault(attr => attr.Key == "PerforceChangelist");
-            if (p4ChangeAttribute?.Value != null)
+            if (p4ChangeAttribute?.Value != null && !string.IsNullOrEmpty(p4ChangeAttribute.Value))
                 srcControlIdentifier = p4ChangeAttribute.Value;
             return Ok(new StatusResponse(_versionFile.VersionString ?? "Unknown", srcControlIdentifier, GetCapabilities(), _replicationSettings.CurrentValue.CurrentSite));
         }
diff --git a/Engine/Source/Programs/Horde/HordeStorage/Horde.Storage/Horde.Storage.csproj b/Engine/Source/Programs/Horde/HordeStorage/Horde.Storage/Horde.Storage.csproj
--- a/Engine/Source/Programs/Horde/HordeStorage/Horde.Storage/Horde.Storage.csproj
+++ b/Engine/Source/Programs/Horde/HordeStorage/Horde.Storage/Horde.Storage.csproj
@@ -13,6 +13,7 @@
     <ItemGroup>
       <AssemblyAttributes Include="AssemblyMetadata">
         <_Parameter1>GitHash</_Parameter1>
+        <!-- Make sure to enable the Github source link to make msbuild automatically populate this -->
         <_Parameter2>$(SourceRevisionId)</_Parameter2>
       </AssemblyAttributes>
 
@@ -79,10 +80,12 @@
 
     <PackageReference Include="CassandraCSharpDriver" Version="3.16.3" />
 
-    <!-- Embeds git information into the build process -->
-    <PackageReference Include="Microsoft.SourceLink.GitHub" Version="1.1.1">
+    <!-- Embeds git information into the build process 
+        Enable this if you are on git to embedd the git hash into the build 
+        Produces warnings while source is in perforce -->
+    <!--<PackageReference Include="Microsoft.SourceLink.GitHub" Version="1.1.1">
       <PrivateAssets>all</PrivateAssets>
-    </PackageReference>
+    </PackageReference>-->
   </ItemGroup>
 
   <ItemGroup>
