Horde: Fix tests.

diff --git a/Engine/Source/Programs/Horde/Horde.Build.Tests/IssueServiceTests.cs b/Engine/Source/Programs/Horde/Horde.Build.Tests/IssueServiceTests.cs
--- a/Engine/Source/Programs/Horde/Horde.Build.Tests/IssueServiceTests.cs
+++ b/Engine/Source/Programs/Horde/Horde.Build.Tests/IssueServiceTests.cs
@@ -62,9 +62,12 @@ namespace HordeServerTests
 
 			public async ValueTask DisposeAsync()
 			{
-				foreach ((LogLevel Level, byte[] Lines) in Events)
+				foreach ((LogLevel Level, byte[] Line) in Events)
 				{
-					await WriteAsync(Level, Lines);
+					byte[] LineWithNewLine = new byte[Line.Length + 1];
+					Array.Copy(Line, LineWithNewLine, Line.Length);
+					LineWithNewLine[LineWithNewLine.Length - 1] = (byte)'\n';
+					await WriteAsync(Level, LineWithNewLine);
 				}
 			}
 
