Horde: Fix exception serializing empty response.

diff --git a/Engine/Source/Programs/Horde/Horde.Build/Controllers/UgsController.cs b/Engine/Source/Programs/Horde/Horde.Build/Controllers/UgsController.cs
--- a/Engine/Source/Programs/Horde/Horde.Build/Controllers/UgsController.cs
+++ b/Engine/Source/Programs/Horde/Horde.Build/Controllers/UgsController.cs
@@ -426,7 +426,7 @@ namespace HordeServer.Controllers
 		public ActionResult<object> PostNetCoreInfo(string? User = null, string? Machine = null, bool NetCore = false)
 		{
 			Logger.LogInformation("NetCore: User={User}, Machine={Machine}, NetCore={NetCore}", User, Machine, NetCore);
-			return new object();
+			return new { };
 		}
 	}
 }
