Horde: Allow null arguments to GetCachedUser(), to simplify call sites.

diff --git a/Engine/Source/Programs/Horde/Horde.Build/Users/IUserCollection.cs b/Engine/Source/Programs/Horde/Horde.Build/Users/IUserCollection.cs
--- a/Engine/Source/Programs/Horde/Horde.Build/Users/IUserCollection.cs
+++ b/Engine/Source/Programs/Horde/Horde.Build/Users/IUserCollection.cs
@@ -32,7 +32,7 @@ namespace HordeServer.Collections
 		/// </summary>
 		/// <param name="Id">Id of the user</param>
 		/// <returns>The user information</returns>
-		Task<IUser?> GetCachedUserAsync(UserId Id);
+		ValueTask<IUser?> GetCachedUserAsync(UserId? Id);
 
 		/// <summary>
 		/// Gets a user by unique id
diff --git a/Engine/Source/Programs/Horde/Horde.Build/Users/Impl/UserCollection.cs b/Engine/Source/Programs/Horde/Horde.Build/Users/Impl/UserCollection.cs
--- a/Engine/Source/Programs/Horde/Horde.Build/Users/Impl/UserCollection.cs
+++ b/Engine/Source/Programs/Horde/Horde.Build/Users/Impl/UserCollection.cs
@@ -106,9 +106,16 @@ namespace HordeServer.Collections.Impl
 		}
 
 		/// <inheritdoc/>
-		public Task<IUser?> GetCachedUserAsync(UserId Id)
+		public async ValueTask<IUser?> GetCachedUserAsync(UserId? Id)
 		{
-			return GetUserAsync(Id);
+			if (Id == null)
+			{
+				return null;
+			}
+			else
+			{
+				return await GetUserAsync(Id.Value);
+			}
 		}
 
 		/// <inheritdoc/>
diff --git a/Engine/Source/Programs/Horde/Horde.Build/Users/Impl/UserCollectionV2.cs b/Engine/Source/Programs/Horde/Horde.Build/Users/Impl/UserCollectionV2.cs
--- a/Engine/Source/Programs/Horde/Horde.Build/Users/Impl/UserCollectionV2.cs
+++ b/Engine/Source/Programs/Horde/Horde.Build/Users/Impl/UserCollectionV2.cs
@@ -209,16 +209,20 @@ namespace HordeServer.Collections.Impl
 		}
 
 		/// <inheritdoc/>
-		public Task<IUser?> GetCachedUserAsync(UserId Id)
+		public async ValueTask<IUser?> GetCachedUserAsync(UserId? Id)
 		{
 			IUser? User;
-			if (UserCache.TryGetValue(Id, out User))
+			if(Id == null)
 			{
-				return Task.FromResult(User);
+				return null;
+			}
+			else if (UserCache.TryGetValue(Id.Value, out User))
+			{
+				return User;
 			}
 			else
 			{
-				return GetUserAsync(Id);
+				return await GetUserAsync(Id.Value);
 			}
 		}
 
