Horde: Added platform and platformgroup properties to agent which mimic the names defined in the engine.

diff --git a/Engine/Source/Programs/Horde/Horde.Agent/Services/WorkerService.cs b/Engine/Source/Programs/Horde/Horde.Agent/Services/WorkerService.cs
--- a/Engine/Source/Programs/Horde/Horde.Agent/Services/WorkerService.cs
+++ b/Engine/Source/Programs/Horde/Horde.Agent/Services/WorkerService.cs
@@ -1357,6 +1357,11 @@ namespace HordeAgent.Services
 			List<DeviceCapabilities> OtherDevices = new List<DeviceCapabilities>();
 			if (RuntimeInformation.IsOSPlatform(OSPlatform.Windows))
 			{
+				PrimaryDevice.Properties.Add("Platform=Win64");
+				PrimaryDevice.Properties.Add("PlatformGroup=Windows");
+				PrimaryDevice.Properties.Add("PlatformGroup=Microsoft");
+				PrimaryDevice.Properties.Add("PlatformGroup=Desktop");
+
 				PrimaryDevice.Properties.Add("OSFamily=Windows");
 
 				// Add OS info
@@ -1490,6 +1495,11 @@ namespace HordeAgent.Services
 			}
 			else if (RuntimeInformation.IsOSPlatform(OSPlatform.Linux))
 			{
+				PrimaryDevice.Properties.Add("Platform=Linux");
+				PrimaryDevice.Properties.Add("PlatformGroup=Linux");
+				PrimaryDevice.Properties.Add("PlatformGroup=Unix");
+				PrimaryDevice.Properties.Add("PlatformGroup=Desktop");
+
 				PrimaryDevice.Properties.Add("OSFamily=Linux");
 				PrimaryDevice.Properties.Add("OSVersion=Linux");
 
@@ -1504,6 +1514,10 @@ namespace HordeAgent.Services
 			}
 			else if(RuntimeInformation.IsOSPlatform(OSPlatform.OSX))
 			{
+				PrimaryDevice.Properties.Add("Platform=Mac");
+				PrimaryDevice.Properties.Add("PlatformGroup=Apple");
+				PrimaryDevice.Properties.Add("PlatformGroup=Desktop");
+
 				PrimaryDevice.Properties.Add("OSFamily=MacOS");
 				PrimaryDevice.Properties.Add("OSVersion=MacOS");
 
