Horde Storage

Use OIDCToken executable that is already in the p4 tree

Also no longer outputting a text file from updateAuthHeader but instead updating env. variable

diff --git a/Engine/Source/Programs/Horde/HordeStorage/Shared/OidcToken/OidcToken.deps.json b/Engine/Source/Programs/Horde/HordeStorage/Shared/OidcToken/OidcToken.deps.json
deleted file mode 100644
diff --git a/Engine/Source/Programs/Horde/HordeStorage/Shared/OidcToken/OidcToken.runtimeconfig.json b/Engine/Source/Programs/Horde/HordeStorage/Shared/OidcToken/OidcToken.runtimeconfig.json
deleted file mode 100644
diff --git a/Engine/Source/Programs/Horde/HordeStorage/Shared/OidcToken/appsettings.json b/Engine/Source/Programs/Horde/HordeStorage/Shared/OidcToken/appsettings.json
deleted file mode 100644
diff --git a/Engine/Source/Programs/Horde/HordeStorage/updateAuthHeaderUser.ps1 b/Engine/Source/Programs/Horde/HordeStorage/updateAuthHeaderUser.ps1
--- a/Engine/Source/Programs/Horde/HordeStorage/updateAuthHeaderUser.ps1
+++ b/Engine/Source/Programs/Horde/HordeStorage/updateAuthHeaderUser.ps1
@@ -1,4 +1,4 @@
 $TempFile = New-TemporaryFile
 
-& "dotnet" ".\Shared\OidcToken\OidcToken.dll" "--Service=EpicGames-Okta-Jupiter" "--OutFile" $TempFile.FullName 
-Get-Content -Raw -Path $TempFile.FullName | ConvertFrom-Json | ForEach-Object { "Authorization: Bearer $($_.Token)" } | Out-File .\curl_auth_header
\ No newline at end of file
+& "dotnet" "..\..\..\..\..\Engine\Restricted\NotForLicensees\Binaries\DotNET\OidcToken\portable\OidcToken.dll" "--Service=EpicGames-Okta-Jupiter" "--OutFile" $TempFile.FullName 
+$env:curl_auth_header = Get-Content -Raw -Path $TempFile.FullName | ConvertFrom-Json | ForEach-Object { "Authorization: Bearer $($_.Token)" }
\ No newline at end of file
