Horde: Fix incorrect batches being cancelled with "agent not found" error.

Failing a batch within a job can result in the batches being reordered. Iterating through batches by index and modifying them as we go can result in indices no longer referring to the same batch.

diff --git a/Engine/Source/Programs/Horde/Horde.Build/Jobs/IJobCollection.cs b/Engine/Source/Programs/Horde/Horde.Build/Jobs/IJobCollection.cs
--- a/Engine/Source/Programs/Horde/Horde.Build/Jobs/IJobCollection.cs
+++ b/Engine/Source/Programs/Horde/Horde.Build/Jobs/IJobCollection.cs
@@ -199,11 +199,11 @@ namespace HordeServer.Collections
 		/// Marks a batch as skipped
 		/// </summary>
 		/// <param name="Job">The job to update</param>
-		/// <param name="BatchIdx">The batch to mark as skipped</param>
+		/// <param name="BatchId">The batch to mark as skipped</param>
 		/// <param name="Graph">Graph for the job</param>
 		/// <param name="Reason">Reason for this batch being failed</param>
 		/// <returns>Updated version of the job</returns>
-		Task<IJob?> SkipBatchAsync(IJob? Job, int BatchIdx, IGraph Graph, JobStepBatchError Reason);
+		Task<IJob?> SkipBatchAsync(IJob? Job, SubResourceId BatchId, IGraph Graph, JobStepBatchError Reason);
 
 		/// <summary>
 		/// Abort an agent's lease, and update the payload accordingly
diff --git a/Engine/Source/Programs/Horde/Horde.Build/Jobs/Impl/JobCollection.cs b/Engine/Source/Programs/Horde/Horde.Build/Jobs/Impl/JobCollection.cs
--- a/Engine/Source/Programs/Horde/Horde.Build/Jobs/Impl/JobCollection.cs
+++ b/Engine/Source/Programs/Horde/Horde.Build/Jobs/Impl/JobCollection.cs
@@ -1076,12 +1076,17 @@ namespace HordeServer.Collections.Impl
 		}
 
 		/// <inheritdoc/>
-		public async Task<IJob?> SkipBatchAsync(IJob? Job, int BatchIdx, IGraph Graph, JobStepBatchError Reason)
+		public async Task<IJob?> SkipBatchAsync(IJob? Job, SubResourceId BatchId, IGraph Graph, JobStepBatchError Reason)
 		{
 			while (Job != null)
 			{
 				JobDocument JobDocument = (JobDocument)Job;
-				JobStepBatchDocument Batch = JobDocument.Batches[BatchIdx];
+
+				JobStepBatchDocument? Batch = JobDocument.Batches.FirstOrDefault(x => x.Id == BatchId);
+				if (Batch == null)
+				{
+					return Job;
+				}
 
 				Batch.State = JobStepBatchState.Complete;
 				Batch.Error = Reason;
diff --git a/Engine/Source/Programs/Horde/Horde.Build/Jobs/Impl/JobTaskSource.cs b/Engine/Source/Programs/Horde/Horde.Build/Jobs/Impl/JobTaskSource.cs
--- a/Engine/Source/Programs/Horde/Horde.Build/Jobs/Impl/JobTaskSource.cs
+++ b/Engine/Source/Programs/Horde/Horde.Build/Jobs/Impl/JobTaskSource.cs
@@ -431,7 +431,7 @@ namespace HordeServer.Tasks.Impl
 
 			// Query for a new list of jobs for the queue
 			List<IJob> NewJobs = await Jobs.GetDispatchQueueAsync();
-			for(int Idx = 0; Idx < NewJobs.Count; Idx++)
+			for (int Idx = 0; Idx < NewJobs.Count; Idx++)
 			{
 				IJob? NewJob = NewJobs[Idx];
 
@@ -460,52 +460,53 @@ namespace HordeServer.Tasks.Impl
 				}
 
 				// Update all the batches
-				bool IsRunning = false;
-				for (int BatchIdx = 0; NewJob != null && BatchIdx < NewJob.Batches.Count; BatchIdx++)
+				HashSet<SubResourceId> CheckedBatchIds = new HashSet<SubResourceId>();
+				while (NewJob != null)
 				{
-					IJobStepBatch Batch = NewJob.Batches[BatchIdx];
-					IsRunning |= (Batch.State == JobStepBatchState.Ready || Batch.State == JobStepBatchState.Starting || Batch.State == JobStepBatchState.Running || Batch.State == JobStepBatchState.Stopping);
+					// Find the next batch within this job that is ready
+					int BatchIdx = NewJob.Batches.FindIndex(x => x.State == JobStepBatchState.Ready && CheckedBatchIds.Add(x.Id));
+					if (BatchIdx == -1)
+					{
+						break;
+					}
 
-					// Skip any batches which aren't ready to execute
-					if (Batch.State == JobStepBatchState.Ready)
+					// Validate the agent type and workspace settings
+					IJobStepBatch Batch = NewJob.Batches[BatchIdx];
+					if (!Stream.AgentTypes.TryGetValue(Graph.Groups[Batch.GroupIdx].AgentType, out AgentType? AgentType))
 					{
-						// Validate the agent type and workspace settings
-						IPool? Pool;
-						if (!Stream.AgentTypes.TryGetValue(Graph.Groups[Batch.GroupIdx].AgentType, out AgentType? AgentType))
-						{
-							NewJob = await SkipBatchAsync(NewJob, BatchIdx, Graph, JobStepBatchError.UnknownAgentType);
-						}
-						else if (!CachedPoolIdToInstance.TryGetValue(AgentType.Pool, out Pool))
-						{
-							NewJob = await SkipBatchAsync(NewJob, BatchIdx, Graph, JobStepBatchError.UnknownPool);
-						}
-						else if (!ValidPools.Contains(AgentType.Pool))
-						{
-							NewJob = await SkipBatchAsync(NewJob, BatchIdx, Graph, JobStepBatchError.NoAgentsInPool);
-						}
-						else if (!OnlinePools.Contains(AgentType.Pool))
-						{
-							NewJob = await SkipBatchAsync(NewJob, BatchIdx, Graph, JobStepBatchError.NoAgentsOnline);
-						}
-						else if (!Stream.TryGetAgentWorkspace(AgentType, out AgentWorkspace? Workspace))
-						{
-							NewJob = await SkipBatchAsync(NewJob, BatchIdx, Graph, JobStepBatchError.UnknownWorkspace);
-						}
-						else
-						{
-							QueueItem NewQueueItem = new QueueItem(Stream, NewJob, BatchIdx, AgentType.Pool, Workspace);
-							NewQueue.Add(NewQueueItem);
-							NewBatchIdToQueueItem[(NewJob.Id, Batch.Id)] = NewQueueItem;
-						}
+						NewJob = await SkipBatchAsync(NewJob, Batch.Id, Graph, JobStepBatchError.UnknownAgentType);
+					}
+					else if (!CachedPoolIdToInstance.TryGetValue(AgentType.Pool, out IPool? Pool))
+					{
+						NewJob = await SkipBatchAsync(NewJob, Batch.Id, Graph, JobStepBatchError.UnknownPool);
+					}
+					else if (!ValidPools.Contains(AgentType.Pool))
+					{
+						NewJob = await SkipBatchAsync(NewJob, Batch.Id, Graph, JobStepBatchError.NoAgentsInPool);
+					}
+					else if (!OnlinePools.Contains(AgentType.Pool))
+					{
+						NewJob = await SkipBatchAsync(NewJob, Batch.Id, Graph, JobStepBatchError.NoAgentsOnline);
+					}
+					else if (!Stream.TryGetAgentWorkspace(AgentType, out AgentWorkspace? Workspace))
+					{
+						NewJob = await SkipBatchAsync(NewJob, Batch.Id, Graph, JobStepBatchError.UnknownWorkspace);
+					}
+					else
+					{
+						QueueItem NewQueueItem = new QueueItem(Stream, NewJob, BatchIdx, AgentType.Pool, Workspace);
+						NewQueue.Add(NewQueueItem);
+						NewBatchIdToQueueItem[(NewJob.Id, Batch.Id)] = NewQueueItem;
 					}
 				}
 
-				// Add a warning if a job looks to be idle
-				if (!IsRunning && NewJob != null)
+				if (NewJob != null)
 				{
-					Logger.LogError("Job {JobId} is in dispatch queue but not currently executing", NewJob.Id);
-					await Jobs.TryRemoveFromDispatchQueueAsync(NewJob);
-					continue;
+					if (!NewJob.Batches.Any(Batch => Batch.State == JobStepBatchState.Ready || Batch.State == JobStepBatchState.Starting || Batch.State == JobStepBatchState.Running || Batch.State == JobStepBatchState.Stopping))
+					{
+						Logger.LogError("Job {JobId} is in dispatch queue but not currently executing", NewJob.Id);
+						await Jobs.TryRemoveFromDispatchQueueAsync(NewJob);
+					}
 				}
 			}
 
@@ -539,12 +540,12 @@ namespace HordeServer.Tasks.Impl
 			}
 		}
 
-		private async Task<IJob?> SkipBatchAsync(IJob Job, int BatchIdx, IGraph Graph, JobStepBatchError Reason)
+		private async Task<IJob?> SkipBatchAsync(IJob Job, SubResourceId BatchId, IGraph Graph, JobStepBatchError Reason)
 		{
-			Logger.LogInformation("Skipping batch {BatchIdx} for job {JobId} (reason: {Reason})", BatchIdx, Job.Id, Reason);
+			Logger.LogInformation("Skipping batch {BatchId} for job {JobId} (reason: {Reason})", BatchId, Job.Id, Reason);
 
 			IReadOnlyList<(LabelState, LabelOutcome)> OldLabelStates = Job.GetLabelStates(Graph);
-			IJob? NewJob = await Jobs.SkipBatchAsync(Job, BatchIdx, Graph, Reason);
+			IJob? NewJob = await Jobs.SkipBatchAsync(Job, BatchId, Graph, Reason);
 			if(NewJob != null)
 			{
 				IReadOnlyList<(LabelState, LabelOutcome)> NewLabelStates = NewJob.GetLabelStates(Graph);
@@ -743,7 +744,7 @@ namespace HordeServer.Tasks.Impl
 		private async Task<AgentLease?> TryCreateLeaseAsync(QueueItem Item, QueueWaiter Waiter)
 		{
 			IJob Job = Item.Job;
-			IJobStepBatch Batch = Job.Batches[Item.BatchIdx];
+			IJobStepBatch Batch = Item.Batch;
 			IAgent Agent = Waiter.Agent;
 			Logger.LogDebug("Assigning job {JobId}, batch {BatchId} to waiter (agent {AgentID})", Job.Id, Batch.Id, Agent.Id);
 
