Horde: Drastically speed up job searches for min create date hint
#rnx

diff --git a/Engine/Source/Programs/Horde/HordeDashboard/src/backend/JobHandler.ts b/Engine/Source/Programs/Horde/HordeDashboard/src/backend/JobHandler.ts
--- a/Engine/Source/Programs/Horde/HordeDashboard/src/backend/JobHandler.ts
+++ b/Engine/Source/Programs/Horde/HordeDashboard/src/backend/JobHandler.ts
@@ -121,6 +121,14 @@ export class JobHandler {
 
             }
 
+            // Set a min time as a TEMPORARY hint which greatly speeds up *all* job query requests
+            // @todo https://jira.it.epicgames.com/browse/UE-135429
+
+            if (!query.minCreateTime) {
+                query.minCreateTime = new Date('2010-1-1').toISOString();
+            }
+            
+
             const cancelId = this.cancelId++;
 
             const queryTime = moment.utc().toISOString();
