Horde: Filter out global device pools
#rnx

diff --git a/Engine/Source/Programs/Horde/Horde.Build/Devices/Impl/DeviceService.cs b/Engine/Source/Programs/Horde/Horde.Build/Devices/Impl/DeviceService.cs
--- a/Engine/Source/Programs/Horde/Horde.Build/Devices/Impl/DeviceService.cs
+++ b/Engine/Source/Programs/Horde/Horde.Build/Devices/Impl/DeviceService.cs
@@ -458,7 +458,13 @@ namespace HordeServer.Services
 				// for global pools which aren't associated with a project
 				if (Pool.ProjectIds == null || Pool.ProjectIds.Count == 0)
 				{
-					AuthPools.Add(new DevicePoolAuthorization(Pool, true, InternalEmployee));
+					if (Pool.PoolType == DevicePoolType.Shared && !InternalEmployee)
+					{
+						AuthPools.Add(new DevicePoolAuthorization(Pool, false, false));
+						continue;
+					}
+
+					AuthPools.Add(new DevicePoolAuthorization(Pool, true, true));
 					continue;
 				}
 
