Horde Storage - Prepared version 0.34.8

diff --git a/Engine/Source/Programs/Horde/HordeStorage/changelog.md b/Engine/Source/Programs/Horde/HordeStorage/changelog.md
--- a/Engine/Source/Programs/Horde/HordeStorage/changelog.md
+++ b/Engine/Source/Programs/Horde/HordeStorage/changelog.md
@@ -1,6 +1,8 @@
 # Unreleased
 
+# 0.34.8
 * Added info to batch request span tags to see the state of each operation
+* Fixes for status endpoint when deployed from Perforce
 
 # 0.34.7
 * Enabled worker deployments in us-west and eu-central
diff --git a/Engine/Source/Programs/Horde/HordeStorage/verifyChangelog.sh b/Engine/Source/Programs/Horde/HordeStorage/verifyChangelog.sh
new file mode 100644
--- /dev/null
+++ b/Engine/Source/Programs/Horde/HordeStorage/verifyChangelog.sh
@@ -0,0 +1,3 @@
+current_version=$(yq e '.version' ./version.yaml)
+cat ./changelog.md | grep -q $current_version && ( echo 'Version found' && exit 0 ) || (echo "No changelog section found for release $current_version" >&2 && exit 1)
+       
\ No newline at end of file
