Horde: Add support for sending job find queries to database replica

All job find requests will query using this as a start. Possibly something we can add to more calls if this works well.

diff --git a/Engine/Source/Programs/Horde/Horde.Build/Controllers/JobsController.cs b/Engine/Source/Programs/Horde/Horde.Build/Controllers/JobsController.cs
--- a/Engine/Source/Programs/Horde/Horde.Build/Controllers/JobsController.cs
+++ b/Engine/Source/Programs/Horde/Horde.Build/Controllers/JobsController.cs
@@ -651,7 +651,7 @@ namespace HordeServer.Controllers
 			{
 				Jobs = await JobService.FindJobsAsync(JobIdValues, StreamIdValue, Name, TemplateRefIds, MinChange,
 					MaxChange, PreflightChange, PreflightStartedByUserIdValue, StartedByUserIdValue, MinCreateTime?.UtcDateTime, MaxCreateTime?.UtcDateTime, Target, State, Outcome,
-					ModifiedBefore, ModifiedAfter, Index, Count);
+					ModifiedBefore, ModifiedAfter, Index, Count, false);
 			}
 
 			StreamPermissionsCache PermissionsCache = new StreamPermissionsCache();
diff --git a/Engine/Source/Programs/Horde/Horde.Build/Jobs/IJobCollection.cs b/Engine/Source/Programs/Horde/Horde.Build/Jobs/IJobCollection.cs
--- a/Engine/Source/Programs/Horde/Horde.Build/Jobs/IJobCollection.cs
+++ b/Engine/Source/Programs/Horde/Horde.Build/Jobs/IJobCollection.cs
@@ -104,8 +104,9 @@ namespace HordeServer.Collections
 		/// <param name="ModifiedAfter">Filter the results by modified time</param>
 		/// <param name="Index">Index of the first result to return</param>
 		/// <param name="Count">Number of results to return</param>
+		/// <param name="ConsistentRead">If the database read should be made to the replica server</param>
 		/// <returns>List of jobs matching the given criteria</returns>
-		Task<List<IJob>> FindAsync(JobId[]? JobIds = null, StreamId? StreamId = null, string? Name = null, TemplateRefId[]? Templates = null, int? MinChange = null, int? MaxChange = null, int? PreflightChange = null, UserId? PreflightStartedByUser = null, UserId? StartedByUser = null, DateTimeOffset? MinCreateTime = null, DateTimeOffset? MaxCreateTime = null, DateTimeOffset? ModifiedBefore = null, DateTimeOffset? ModifiedAfter = null, int? Index = null, int? Count = null);
+		Task<List<IJob>> FindAsync(JobId[]? JobIds = null, StreamId? StreamId = null, string? Name = null, TemplateRefId[]? Templates = null, int? MinChange = null, int? MaxChange = null, int? PreflightChange = null, UserId? PreflightStartedByUser = null, UserId? StartedByUser = null, DateTimeOffset? MinCreateTime = null, DateTimeOffset? MaxCreateTime = null, DateTimeOffset? ModifiedBefore = null, DateTimeOffset? ModifiedAfter = null, int? Index = null, int? Count = null, bool ConsistentRead = true);
 
 		/// <summary>
 		/// Updates a new job
diff --git a/Engine/Source/Programs/Horde/Horde.Build/Jobs/Impl/JobCollection.cs b/Engine/Source/Programs/Horde/Horde.Build/Jobs/Impl/JobCollection.cs
--- a/Engine/Source/Programs/Horde/Horde.Build/Jobs/Impl/JobCollection.cs
+++ b/Engine/Source/Programs/Horde/Horde.Build/Jobs/Impl/JobCollection.cs
@@ -496,7 +496,7 @@ namespace HordeServer.Collections.Impl
 		}
 
 		/// <inheritdoc/>
-		public async Task<List<IJob>> FindAsync(JobId[]? JobIds, StreamId? StreamId, string? Name, TemplateRefId[]? Templates, int? MinChange, int? MaxChange, int? PreflightChange, UserId? PreflightStartedByUser, UserId? StartedByUser, DateTimeOffset? MinCreateTime, DateTimeOffset? MaxCreateTime, DateTimeOffset? ModifiedBefore, DateTimeOffset? ModifiedAfter, int? Index, int? Count)
+		public async Task<List<IJob>> FindAsync(JobId[]? JobIds, StreamId? StreamId, string? Name, TemplateRefId[]? Templates, int? MinChange, int? MaxChange, int? PreflightChange, UserId? PreflightStartedByUser, UserId? StartedByUser, DateTimeOffset? MinCreateTime, DateTimeOffset? MaxCreateTime, DateTimeOffset? ModifiedBefore, DateTimeOffset? ModifiedAfter, int? Index, int? Count, bool ConsistentRead)
 		{
 			FilterDefinitionBuilder<JobDocument> FilterBuilder = Builders<JobDocument>.Filter;
 
@@ -557,7 +557,9 @@ namespace HordeServer.Collections.Impl
 			List<JobDocument> Results;
 			using (IScope Scope = GlobalTracer.Instance.BuildSpan("Jobs.Find").StartActive())
 			{
-				IFindFluent<JobDocument, JobDocument> Query = Jobs.Find<JobDocument>(Filter).SortByDescending(x => x.CreateTimeUtc!);
+				IMongoCollection<JobDocument> Collection = ConsistentRead ? Jobs : Jobs.WithReadPreference(ReadPreference.SecondaryPreferred);
+				IFindFluent<JobDocument, JobDocument> Query = Collection.Find<JobDocument>(Filter).SortByDescending(x => x.CreateTimeUtc!);
+				
 				if (Index != null)
 				{
 					Query = Query.Skip(Index.Value);
diff --git a/Engine/Source/Programs/Horde/Horde.Build/Jobs/Impl/JobService.cs b/Engine/Source/Programs/Horde/Horde.Build/Jobs/Impl/JobService.cs
--- a/Engine/Source/Programs/Horde/Horde.Build/Jobs/Impl/JobService.cs
+++ b/Engine/Source/Programs/Horde/Horde.Build/Jobs/Impl/JobService.cs
@@ -531,8 +531,9 @@ namespace HordeServer.Services
 		/// <param name="ModifiedAfter">Filter the results by last modified time</param>
 		/// <param name="Index">Index of the first result to return</param>
 		/// <param name="Count">Number of results to return</param>
+		/// <param name="ConsistentRead">If the database read should be made to the replica server</param>
 		/// <returns>List of jobs matching the given criteria</returns>
-		public async Task<List<IJob>> FindJobsAsync(JobId[]? JobIds = null, StreamId? StreamId = null, string? Name = null, TemplateRefId[]? Templates = null, int? MinChange = null, int? MaxChange = null, int? PreflightChange = null, UserId? PreflightStartedByUser = null, UserId? StartedByUser = null, DateTimeOffset ? MinCreateTime = null, DateTimeOffset? MaxCreateTime = null, string? Target = null, JobStepState[]? State = null, JobStepOutcome[]? Outcome = null, DateTimeOffset? ModifiedBefore = null, DateTimeOffset? ModifiedAfter = null, int? Index = null, int? Count = null)
+		public async Task<List<IJob>> FindJobsAsync(JobId[]? JobIds = null, StreamId? StreamId = null, string? Name = null, TemplateRefId[]? Templates = null, int? MinChange = null, int? MaxChange = null, int? PreflightChange = null, UserId? PreflightStartedByUser = null, UserId? StartedByUser = null, DateTimeOffset ? MinCreateTime = null, DateTimeOffset? MaxCreateTime = null, string? Target = null, JobStepState[]? State = null, JobStepOutcome[]? Outcome = null, DateTimeOffset? ModifiedBefore = null, DateTimeOffset? ModifiedAfter = null, int? Index = null, int? Count = null, bool ConsistentRead = true)
 		{
 			using IScope Scope = GlobalTracer.Instance.BuildSpan("JobService.FindJobsAsync").StartActive();
 			Scope.Span.SetTag("JobIds", JobIds);
@@ -556,7 +557,7 @@ namespace HordeServer.Services
 			
 			if (Target == null && (State == null || State.Length == 0) && (Outcome == null || Outcome.Length == 0))
 			{
-				return await Jobs.FindAsync(JobIds, StreamId, Name, Templates, MinChange, MaxChange, PreflightChange, PreflightStartedByUser, StartedByUser, MinCreateTime, MaxCreateTime, ModifiedBefore, ModifiedAfter, Index, Count);
+				return await Jobs.FindAsync(JobIds, StreamId, Name, Templates, MinChange, MaxChange, PreflightChange, PreflightStartedByUser, StartedByUser, MinCreateTime, MaxCreateTime, ModifiedBefore, ModifiedAfter, Index, Count, ConsistentRead);
 			}
 			else
 			{
@@ -566,7 +567,7 @@ namespace HordeServer.Services
 				int MaxCount = (Count ?? 1);
 				while (Results.Count < MaxCount)
 				{
-					List<IJob> ScanJobs = await Jobs.FindAsync(JobIds, StreamId, Name, Templates, MinChange, MaxChange, PreflightChange, PreflightStartedByUser, StartedByUser, MinCreateTime, MaxCreateTime, ModifiedBefore, ModifiedAfter, 0, 5);
+					List<IJob> ScanJobs = await Jobs.FindAsync(JobIds, StreamId, Name, Templates, MinChange, MaxChange, PreflightChange, PreflightStartedByUser, StartedByUser, MinCreateTime, MaxCreateTime, ModifiedBefore, ModifiedAfter, 0, 5, ConsistentRead);
 					if (ScanJobs.Count == 0)
 					{
 						break;
