Horde: Add template setting for updating issues, also make it possible to set a job to update issues in advanced settings
#rnx

diff --git a/Engine/Source/Programs/Horde/Horde.Build.Tests/Fixture.cs b/Engine/Source/Programs/Horde/Horde.Build.Tests/Fixture.cs
--- a/Engine/Source/Programs/Horde/Horde.Build.Tests/Fixture.cs
+++ b/Engine/Source/Programs/Horde/Horde.Build.Tests/Fixture.cs
@@ -56,7 +56,7 @@ namespace HordeServerTests
 			Fg.Aggregates = new List<IAggregate>();
 			Fg.Labels = new List<ILabel>();
 
-			Template = await TemplateCollection.AddAsync("Test template", null, false, null, null,
+			Template = await TemplateCollection.AddAsync("Test template", null, false, false, null, null,
 				new List<string>(), new List<Parameter>());
 			Graph = await GraphCollection.AddAsync(Template);
 
diff --git a/Engine/Source/Programs/Horde/Horde.Build.Tests/JobServiceTests.cs b/Engine/Source/Programs/Horde/Horde.Build.Tests/JobServiceTests.cs
--- a/Engine/Source/Programs/Horde/Horde.Build.Tests/JobServiceTests.cs
+++ b/Engine/Source/Programs/Horde/Horde.Build.Tests/JobServiceTests.cs
@@ -30,7 +30,7 @@ namespace HordeServerTests
 			IProject? Project = await ProjectService.Collection.AddOrUpdateAsync(ProjectId, "", "", 0, new ProjectConfig { Name = "UE5" });
 			Assert.IsNotNull(Project);
 
-			ITemplate Template = await TemplateCollection.AddAsync("Test template", null, false, null, null, new List<string>(), new List<Parameter>());
+			ITemplate Template = await TemplateCollection.AddAsync("Test template", null, false, false, null, null, new List<string>(), new List<Parameter>());
 			IGraph Graph = await GraphCollection.AddAsync(Template);
 
 			TemplateRefId TemplateRefId1 = new TemplateRefId("template1");
@@ -153,7 +153,7 @@ namespace HordeServerTests
 			IStream? Stream = await StreamCollection.GetAsync(StreamId);
 			Stream = await StreamCollection.TryCreateOrReplaceAsync(StreamId, Stream, "", "", Project!.Id, new StreamConfig { Name = "//UE5/Main" });
 
-			ITemplate Template = await TemplateCollection.AddAsync("Test template", null, false, null, null, new List<string>(), new List<Parameter>());
+			ITemplate Template = await TemplateCollection.AddAsync("Test template", null, false, false, null, null, new List<string>(), new List<Parameter>());
 			IGraph Graph = await GraphCollection.AddAsync(Template);
 
 			NewGroup GroupA = new NewGroup("win", new List<NewNode>());
diff --git a/Engine/Source/Programs/Horde/Horde.Build.Tests/SchedulerTests.cs b/Engine/Source/Programs/Horde/Horde.Build.Tests/SchedulerTests.cs
--- a/Engine/Source/Programs/Horde/Horde.Build.Tests/SchedulerTests.cs
+++ b/Engine/Source/Programs/Horde/Horde.Build.Tests/SchedulerTests.cs
@@ -42,7 +42,7 @@ namespace HordeServerTests
 			IProject ? Project = ProjectService.Collection.AddOrUpdateAsync(ProjectId, "", "", 0, new ProjectConfig { Name = "UE5" }).Result;
 			Assert.IsNotNull(Project);
 
-			Template = TemplateCollection.AddAsync("Test template", null, false, null, null, new List<string>(), new List<Parameter>()).Result;
+			Template = TemplateCollection.AddAsync("Test template", null, false, false, null, null, new List<string>(), new List<Parameter>()).Result;
 
 			InitialJobIds = new HashSet<JobId>(JobCollection.FindAsync().Result.Select(x => x.Id));
 
diff --git a/Engine/Source/Programs/Horde/Horde.Build/Api/Jobs.cs b/Engine/Source/Programs/Horde/Horde.Build/Api/Jobs.cs
--- a/Engine/Source/Programs/Horde/Horde.Build/Api/Jobs.cs
+++ b/Engine/Source/Programs/Horde/Horde.Build/Api/Jobs.cs
@@ -330,6 +330,11 @@ namespace HordeServer.Api
 		/// </summary>
 		public DateTimeOffset UpdateTime { get; set; }
 
+		/// <summary>
+		/// Whether issues are being updated by this job
+		/// </summary>
+		public bool UpdateIssues { get; set; }
+
 		/// <summary>
 		/// Per-object permissions
 		/// </summary>
@@ -368,6 +373,7 @@ namespace HordeServer.Api
 			this.Reports = Job.Reports?.ConvertAll(x => new GetReportResponse(x));
 			this.Arguments = Job.Arguments.ToList();
 			this.UpdateTime = new DateTimeOffset(Job.UpdateTimeUtc);
+			this.UpdateIssues = Job.UpdateIssues;
 			this.Acl = AclResponse;
 		}
 	}
diff --git a/Engine/Source/Programs/Horde/Horde.Build/Api/Templates.cs b/Engine/Source/Programs/Horde/Horde.Build/Api/Templates.cs
--- a/Engine/Source/Programs/Horde/Horde.Build/Api/Templates.cs
+++ b/Engine/Source/Programs/Horde/Horde.Build/Api/Templates.cs
@@ -410,6 +410,11 @@ namespace HordeServer.Api
 		/// </summary>
 		public bool AllowPreflights { get; set; } = true;
 
+		/// <summary>
+		/// Whether issues should be updated for all jobs using this template
+		/// </summary>
+		public bool UpdateIssues { get; set; } = false;
+
 		/// <summary>
 		/// Initial agent type to parse the buildgraph script on
 		/// </summary>
@@ -451,6 +456,11 @@ namespace HordeServer.Api
 		/// </summary>
 		public bool AllowPreflights { get; set; }
 
+		/// <summary>
+		/// Whether to always update issues on jobs using this template
+		/// </summary>
+		public bool UpdateIssues { get; set; }
+
 		/// <summary>
 		/// The initial agent type to parse the BuildGraph script on
 		/// </summary>
@@ -491,6 +501,7 @@ namespace HordeServer.Api
 			this.Name = Template.Name;
 			this.Priority = Template.Priority;
 			this.AllowPreflights = Template.AllowPreflights;
+			this.UpdateIssues = Template.UpdateIssues;
 			this.InitialAgentType = Template.InitialAgentType;
 			this.SubmitNewChange = Template.SubmitNewChange;
 			this.Arguments = new List<string>(Template.Arguments);
diff --git a/Engine/Source/Programs/Horde/Horde.Build/Controllers/JobsController.cs b/Engine/Source/Programs/Horde/Horde.Build/Controllers/JobsController.cs
--- a/Engine/Source/Programs/Horde/Horde.Build/Controllers/JobsController.cs
+++ b/Engine/Source/Programs/Horde/Horde.Build/Controllers/JobsController.cs
@@ -169,8 +169,18 @@ namespace HordeServer.Controllers
 				}
 			}
 
+			bool? UpdateIssues = null;
+			if (Template.UpdateIssues)
+			{
+				UpdateIssues = true;
+			}
+			else if (Create.UpdateIssues.HasValue)
+			{
+				UpdateIssues = Create.UpdateIssues.Value;
+			}
+
 			// Create the job
-			IJob Job = await JobService.CreateJobAsync(null, Stream, TemplateRefId, Template.Id, Graph, Name, Change, CodeChange, Create.PreflightChange, null, User.GetUserId(), Priority, Create.AutoSubmit, Create.UpdateIssues, TemplateRef.ChainedJobs, TemplateRef.ShowUgsBadges, TemplateRef.ShowUgsAlerts, TemplateRef.NotificationChannel, TemplateRef.NotificationChannelFilter, Arguments);
+			IJob Job = await JobService.CreateJobAsync(null, Stream, TemplateRefId, Template.Id, Graph, Name, Change, CodeChange, Create.PreflightChange, null, User.GetUserId(), Priority, Create.AutoSubmit, UpdateIssues, TemplateRef.ChainedJobs, TemplateRef.ShowUgsBadges, TemplateRef.ShowUgsAlerts, TemplateRef.NotificationChannel, TemplateRef.NotificationChannelFilter, Arguments);
 			await UpdateNotificationsAsync(Job.Id, new UpdateNotificationsRequest { Slack = true });
 			return new CreateJobResponse(Job.Id.ToString());
 		}
diff --git a/Engine/Source/Programs/Horde/Horde.Build/Controllers/StreamsController.cs b/Engine/Source/Programs/Horde/Horde.Build/Controllers/StreamsController.cs
--- a/Engine/Source/Programs/Horde/Horde.Build/Controllers/StreamsController.cs
+++ b/Engine/Source/Programs/Horde/Horde.Build/Controllers/StreamsController.cs
@@ -318,7 +318,7 @@ namespace HordeServer.Controllers
 			foreach (CreateTemplateRefRequest Request in Requests)
 			{
 				// Create the template
-				ITemplate NewTemplate = await TemplateCollection.AddAsync(Request.Name, Request.Priority, Request.AllowPreflights, Request.InitialAgentType, Request.SubmitNewChange, Request.Arguments, Request.Parameters.ConvertAll(x => x.ToModel()));
+				ITemplate NewTemplate = await TemplateCollection.AddAsync(Request.Name, Request.Priority, Request.AllowPreflights, Request.UpdateIssues, Request.InitialAgentType, Request.SubmitNewChange, Request.Arguments, Request.Parameters.ConvertAll(x => x.ToModel()));
 
 				// Get an identifier for the new template ref
 				TemplateRefId NewTemplateRefId;
diff --git a/Engine/Source/Programs/Horde/Horde.Build/Jobs/ITemplate.cs b/Engine/Source/Programs/Horde/Horde.Build/Jobs/ITemplate.cs
--- a/Engine/Source/Programs/Horde/Horde.Build/Jobs/ITemplate.cs
+++ b/Engine/Source/Programs/Horde/Horde.Build/Jobs/ITemplate.cs
@@ -441,6 +441,11 @@ namespace HordeServer.Models
 		/// </summary>
 		public bool AllowPreflights { get; }
 
+		/// <summary>
+		/// Whether to always issues for jobs using this template
+		/// </summary>
+		public bool UpdateIssues { get; }
+
 		/// <summary>
 		/// Agent type to use for parsing the job state
 		/// </summary>
diff --git a/Engine/Source/Programs/Horde/Horde.Build/Jobs/ITemplateCollection.cs b/Engine/Source/Programs/Horde/Horde.Build/Jobs/ITemplateCollection.cs
--- a/Engine/Source/Programs/Horde/Horde.Build/Jobs/ITemplateCollection.cs
+++ b/Engine/Source/Programs/Horde/Horde.Build/Jobs/ITemplateCollection.cs
@@ -28,11 +28,12 @@ namespace HordeServer.Collections
 		/// <param name="Name">Name of the template</param>
 		/// <param name="Priority">Priority of this template</param>
 		/// <param name="bAllowPreflights">Whether to allow preflights of this job</param>
+		/// <param name="bUpdateIssues"> Whether to update issues for all jobs using this template</param>
 		/// <param name="InitialAgentType">The agent type to parse the buildgraph script</param>
 		/// <param name="SubmitNewChange">Path to a file within the stream to submit to generate a new changelist for jobs</param>
 		/// <param name="Arguments">List of arguments which are always specified</param>
 		/// <param name="Parameters">List of template parameters</param>
-		Task<ITemplate> AddAsync(string Name, Priority? Priority = null, bool bAllowPreflights = true, string? InitialAgentType = null, string? SubmitNewChange = null, List<string>? Arguments = null, List<Parameter>? Parameters = null);
+		Task<ITemplate> AddAsync(string Name, Priority? Priority = null, bool bAllowPreflights = true, bool bUpdateIssues = false, string ? InitialAgentType = null, string? SubmitNewChange = null, List<string>? Arguments = null, List<Parameter>? Parameters = null);
 
 		/// <summary>
 		/// Gets all the available templates
diff --git a/Engine/Source/Programs/Horde/Horde.Build/Jobs/Impl/TemplateCollection.cs b/Engine/Source/Programs/Horde/Horde.Build/Jobs/Impl/TemplateCollection.cs
--- a/Engine/Source/Programs/Horde/Horde.Build/Jobs/Impl/TemplateCollection.cs
+++ b/Engine/Source/Programs/Horde/Horde.Build/Jobs/Impl/TemplateCollection.cs
@@ -36,6 +36,7 @@ namespace HordeServer.Collections.Impl
 
 			public Priority? Priority { get; private set; }
 			public bool AllowPreflights { get; set; } = true;
+			public bool UpdateIssues { get; set; } = false;
 			public string? InitialAgentType { get; set; }
 
 			[BsonIgnoreIfNull]
@@ -53,11 +54,12 @@ namespace HordeServer.Collections.Impl
 				Name = null!;
 			}
 
-			public TemplateDocument(string Name, Priority? Priority, bool bAllowPreflights, string? InitialAgentType, string? SubmitNewChange, List<string>? Arguments, List<Parameter>? Parameters)
+			public TemplateDocument(string Name, Priority? Priority, bool bAllowPreflights, bool UpdateIssues, string? InitialAgentType, string? SubmitNewChange, List<string>? Arguments, List<Parameter>? Parameters)
 			{
 				this.Name = Name;
 				this.Priority = Priority;
 				this.AllowPreflights = bAllowPreflights;
+				this.UpdateIssues = UpdateIssues;
 				this.InitialAgentType = InitialAgentType;
 				this.SubmitNewChange = SubmitNewChange;
 				this.Arguments = Arguments ?? new List<string>();
@@ -103,9 +105,9 @@ namespace HordeServer.Collections.Impl
 		}
 
 		/// <inheritdoc/>
-		public async Task<ITemplate> AddAsync(string Name, Priority? Priority, bool bAllowPreflights, string? InitialAgentType, string? SubmitNewChange, List<string>? Arguments, List<Parameter>? Parameters)
+		public async Task<ITemplate> AddAsync(string Name, Priority? Priority, bool bAllowPreflights, bool bUpdateIssues, string? InitialAgentType, string? SubmitNewChange, List<string>? Arguments, List<Parameter>? Parameters)
 		{
-			TemplateDocument Template = new TemplateDocument(Name, Priority, bAllowPreflights, InitialAgentType, SubmitNewChange, Arguments, Parameters);
+			TemplateDocument Template = new TemplateDocument(Name, Priority, bAllowPreflights, bUpdateIssues, InitialAgentType, SubmitNewChange, Arguments, Parameters);
 			if (await GetAsync(Template.Id) == null)
 			{
 				await Templates.ReplaceOneAsync(x => x.Id == Template.Id, Template, new ReplaceOptions { IsUpsert = true });
diff --git a/Engine/Source/Programs/Horde/Horde.Build/Projects/Impl/StreamCollection.cs b/Engine/Source/Programs/Horde/Horde.Build/Projects/Impl/StreamCollection.cs
--- a/Engine/Source/Programs/Horde/Horde.Build/Projects/Impl/StreamCollection.cs
+++ b/Engine/Source/Programs/Horde/Horde.Build/Projects/Impl/StreamCollection.cs
@@ -172,7 +172,7 @@ namespace HordeServer.Collections.Impl
 			foreach (CreateTemplateRefRequest Request in Requests)
 			{
 				// Create the template
-				ITemplate NewTemplate = await TemplateCollection.AddAsync(Request.Name, Request.Priority, Request.AllowPreflights, Request.InitialAgentType, Request.SubmitNewChange, Request.Arguments, Request.Parameters.ConvertAll(x => x.ToModel()));
+				ITemplate NewTemplate = await TemplateCollection.AddAsync(Request.Name, Request.Priority, Request.AllowPreflights, Request.UpdateIssues, Request.InitialAgentType, Request.SubmitNewChange, Request.Arguments, Request.Parameters.ConvertAll(x => x.ToModel()));
 
 				// Get an identifier for the new template ref
 				TemplateRefId NewTemplateRefId;
diff --git a/Engine/Source/Programs/Horde/HordeDashboard/src/backend/Api.ts b/Engine/Source/Programs/Horde/HordeDashboard/src/backend/Api.ts
--- a/Engine/Source/Programs/Horde/HordeDashboard/src/backend/Api.ts
+++ b/Engine/Source/Programs/Horde/HordeDashboard/src/backend/Api.ts
@@ -806,15 +806,12 @@ export type CreateJobRequest = {
 	/** Whether to automatically submit the preflighted change on completion */
 	autoSubmit?: boolean;
 
+	/** Whether to update issues based on the outcome of this job */
+	updateIssues?: boolean;
+
 	/** Nodes for the new job */
 	groups?: CreateGroupRequest[];
 
-	/** Aggregates for the new job */
-	aggregates?: CreateAggregateRequest[];
-
-	/** Labels for the new job*/
-	labels?: CreateLabelRequest[];
-
 	/** Arguments for the job */
 	arguments?: string[];
 }
@@ -947,6 +944,9 @@ export type GetJobResponse = {
 	/**The last update time for this job*/
 	updateTime: Date | string;
 
+	/** Whether to update issues based on the outcome of this job */
+	updateIssues?: boolean;
+
 	/**  Custom permissions for this object */
 	acl?: GetAclResponse;
 
@@ -1790,6 +1790,9 @@ export type CreateTemplateRequest = {
 	/**Whether to allow preflights of this template */
 	allowPreflights: boolean;
 
+	/**Whether always update issues regardless of how job was created */
+	updateIssues: boolean;
+
 	/**Array of nodes for this job */
 	groups: CreateGroupRequest[];
 
@@ -1953,6 +1956,9 @@ export type GetTemplateResponse = {
 	/**Whether to allow preflights of this template */
 	allowPreflights: boolean;
 
+	/**Whether to always update issues on jobs that use this template */
+	updateIssues: boolean;
+
 	/**List of node groups for this job */
 	groups: CreateGroupRequest[];
 
