Fixed up Prod release scripts for Horde.Storage

diff --git a/Engine/Source/Programs/Horde/HordeStorage/verifyChangelog.sh b/Engine/Source/Programs/Horde/HordeStorage/verifyChangelog.sh
--- a/Engine/Source/Programs/Horde/HordeStorage/verifyChangelog.sh
+++ b/Engine/Source/Programs/Horde/HordeStorage/verifyChangelog.sh
@@ -1,3 +1,4 @@
-current_version=$(yq e '.version' ./version.yaml)
-cat ./changelog.md | grep -q $current_version && ( echo 'Version found' && exit 0 ) || (echo "No changelog section found for release $current_version" >&2 && exit 1)
+dir=$(dirname "${BASH_SOURCE[0]}")
+current_version=$(yq e '.version' $dir/version.yaml)
+cat $dir/changelog.md | grep -q $current_version && ( echo 'Version found' && exit 0 ) || (echo "No changelog section found for release $current_version" >&2 && exit 1)
        
\ No newline at end of file
