Horde: Add a flag to each pool indicating whether any of its workspaces require AutoSDK.

diff --git a/Engine/Source/Programs/Horde/Horde.Build/Fleet/IPool.cs b/Engine/Source/Programs/Horde/Horde.Build/Fleet/IPool.cs
--- a/Engine/Source/Programs/Horde/Horde.Build/Fleet/IPool.cs
+++ b/Engine/Source/Programs/Horde/Horde.Build/Fleet/IPool.cs
@@ -40,6 +40,11 @@ namespace HordeServer.Models
 		/// </summary>
 		public IReadOnlyList<AgentWorkspace> Workspaces { get; }
 
+		/// <summary>
+		/// Whether to sync autosdk alongside the workspaces for this pool
+		/// </summary>
+		public bool UseAutoSdk { get; }
+
 		/// <summary>
 		/// Arbitrary properties related to this pool
 		/// </summary>
diff --git a/Engine/Source/Programs/Horde/Horde.Build/Fleet/IPoolCollection.cs b/Engine/Source/Programs/Horde/Horde.Build/Fleet/IPoolCollection.cs
--- a/Engine/Source/Programs/Horde/Horde.Build/Fleet/IPoolCollection.cs
+++ b/Engine/Source/Programs/Horde/Horde.Build/Fleet/IPoolCollection.cs
@@ -72,10 +72,11 @@ namespace HordeServer.Collections
 		/// <param name="NewMinAgents">Minimum number of agents in the pool</param>
 		/// <param name="NewNumReserveAgents">Minimum number of idle agents to maintain</param>
 		/// <param name="NewWorkspaces">New workspaces for the pool</param>
+		/// <param name="NewUseAutoSdk">New setting for whether to use autosdk</param>
 		/// <param name="NewProperties">New properties for the pool</param>
 		/// <param name="LastScaleUpTime">New time for last (auto) scale up</param>
 		/// <param name="LastScaleDownTime">New time for last (auto) scale down</param>
 		/// <returns>Async task</returns>
-		Task<IPool?> TryUpdateAsync(IPool Pool, string? NewName = null, Condition? NewCondition = null, bool? NewEnableAutoscaling = null, int? NewMinAgents = null, int? NewNumReserveAgents = null, List<AgentWorkspace>? NewWorkspaces = null, Dictionary<string, string?>? NewProperties = null, DateTime? LastScaleUpTime = null, DateTime? LastScaleDownTime = null);
+		Task<IPool?> TryUpdateAsync(IPool Pool, string? NewName = null, Condition? NewCondition = null, bool? NewEnableAutoscaling = null, int? NewMinAgents = null, int? NewNumReserveAgents = null, List<AgentWorkspace>? NewWorkspaces = null, bool? NewUseAutoSdk = null, Dictionary<string, string?>? NewProperties = null, DateTime? LastScaleUpTime = null, DateTime? LastScaleDownTime = null);
 	}
 }
diff --git a/Engine/Source/Programs/Horde/Horde.Build/Fleet/Impl/PoolCollection.cs b/Engine/Source/Programs/Horde/Horde.Build/Fleet/Impl/PoolCollection.cs
--- a/Engine/Source/Programs/Horde/Horde.Build/Fleet/Impl/PoolCollection.cs
+++ b/Engine/Source/Programs/Horde/Horde.Build/Fleet/Impl/PoolCollection.cs
@@ -33,6 +33,7 @@ namespace HordeServer.Collections.Impl
 			[BsonIgnoreIfNull]
 			public Condition? Condition { get; set; }
 			public List<AgentWorkspace> Workspaces { get; set; } = new List<AgentWorkspace>();
+			public bool UseAutoSdk { get; set; } = true;
 			public Dictionary<string, string> Properties { get; set; } = new Dictionary<string, string>();
 			[BsonIgnoreIfDefault(true)]
 			public bool EnableAutoscaling { get; set; } = true;
@@ -56,6 +57,7 @@ namespace HordeServer.Collections.Impl
 				Name = Other.Name;
 				Condition = Other.Condition;
 				Workspaces.AddRange(Other.Workspaces);
+				UseAutoSdk = Other.UseAutoSdk;
 				Properties = new Dictionary<string, string>(Other.Properties);
 				EnableAutoscaling = Other.EnableAutoscaling;
 				MinAgents = Other.MinAgents;
@@ -171,7 +173,7 @@ namespace HordeServer.Collections.Impl
 		}
 
 		/// <inheritdoc/>
-		public Task<IPool?> TryUpdateAsync(IPool Pool, string? NewName, Condition? NewCondition, bool? NewEnableAutoscaling, int? NewMinAgents, int? NewNumReserveAgents, List<AgentWorkspace>? NewWorkspaces, Dictionary<string, string?>? NewProperties, DateTime? LastScaleUpTime, DateTime? LastScaleDownTime)
+		public Task<IPool?> TryUpdateAsync(IPool Pool, string? NewName, Condition? NewCondition, bool? NewEnableAutoscaling, int? NewMinAgents, int? NewNumReserveAgents, List<AgentWorkspace>? NewWorkspaces, bool? NewUseAutoSdk, Dictionary<string, string?>? NewProperties, DateTime? LastScaleUpTime, DateTime? LastScaleDownTime)
 		{
 			TransactionBuilder<PoolDocument> Transaction = new TransactionBuilder<PoolDocument>();
 			if (NewName != null)
@@ -226,6 +228,10 @@ namespace HordeServer.Collections.Impl
 			{
 				Transaction.Set(x => x.Workspaces, NewWorkspaces);
 			}
+			if (NewUseAutoSdk != null)
+			{
+				Transaction.Set(x => x.UseAutoSdk, NewUseAutoSdk.Value);
+			}
 			if (NewProperties != null)
 			{
 				Transaction.UpdateDictionary(x => x.Properties, NewProperties);
diff --git a/Engine/Source/Programs/Horde/Horde.Build/Fleet/PoolService.cs b/Engine/Source/Programs/Horde/Horde.Build/Fleet/PoolService.cs
--- a/Engine/Source/Programs/Horde/Horde.Build/Fleet/PoolService.cs
+++ b/Engine/Source/Programs/Horde/Horde.Build/Fleet/PoolService.cs
@@ -98,7 +98,7 @@ namespace HordeServer.Services
 		{
 			for (; Pool != null; Pool = await Pools.GetAsync(Pool.Id))
 			{
-				IPool? NewPool = await Pools.TryUpdateAsync(Pool, NewName, NewCondition, NewEnableAutoscaling, NewMinAgents, NewNumReserveAgents, null, NewProperties, null, null);
+				IPool? NewPool = await Pools.TryUpdateAsync(Pool, NewName, NewCondition, NewEnableAutoscaling, NewMinAgents, NewNumReserveAgents, NewProperties: NewProperties);
 				if (NewPool != null)
 				{
 					return NewPool;
@@ -134,6 +134,7 @@ namespace HordeServer.Services
 		/// <returns>List of workspaces</returns>
 		public async Task<HashSet<AgentWorkspace>> GetWorkspacesAsync(IAgent Agent, DateTime ValidAtTime)
 		{
+			bool bAddAutoSdkWorkspace = false;
 			HashSet<AgentWorkspace> Workspaces = new HashSet<AgentWorkspace>();
 
 			Dictionary<PoolId, IPool> PoolMapping = await GetPoolLookupAsync(ValidAtTime);
@@ -143,11 +144,15 @@ namespace HordeServer.Services
 				if (PoolMapping.TryGetValue(PoolId, out Pool))
 				{
 					Workspaces.UnionWith(Pool.Workspaces);
+					bAddAutoSdkWorkspace |= Pool.UseAutoSdk;
 				}
 			}
 
-			Globals Globals = await DatabaseService.GetGlobalsAsync();
-			Workspaces.UnionWith(Agent.GetAutoSdkWorkspaces(Globals, Workspaces.ToList()));
+			if (bAddAutoSdkWorkspace)
+			{
+				Globals Globals = await DatabaseService.GetGlobalsAsync();
+				Workspaces.UnionWith(Agent.GetAutoSdkWorkspaces(Globals, Workspaces.ToList()));
+			}
 
 			return Workspaces;
 		}
diff --git a/Engine/Source/Programs/Horde/Horde.Build/Jobs/Impl/JobTaskSource.cs b/Engine/Source/Programs/Horde/Horde.Build/Jobs/Impl/JobTaskSource.cs
--- a/Engine/Source/Programs/Horde/Horde.Build/Jobs/Impl/JobTaskSource.cs
+++ b/Engine/Source/Programs/Horde/Horde.Build/Jobs/Impl/JobTaskSource.cs
@@ -488,13 +488,13 @@ namespace HordeServer.Tasks.Impl
 					{
 						NewJob = await SkipBatchAsync(NewJob, Batch.Id, Graph, JobStepBatchError.NoAgentsOnline);
 					}
-					else if (!Stream.TryGetAgentWorkspace(AgentType, out AgentWorkspace? Workspace))
+					else if (!Stream.TryGetAgentWorkspace(AgentType, out (AgentWorkspace, bool)? Workspace))
 					{
 						NewJob = await SkipBatchAsync(NewJob, Batch.Id, Graph, JobStepBatchError.UnknownWorkspace);
 					}
 					else
 					{
-						QueueItem NewQueueItem = new QueueItem(Stream, NewJob, BatchIdx, AgentType.Pool, Workspace);
+						QueueItem NewQueueItem = new QueueItem(Stream, NewJob, BatchIdx, AgentType.Pool, Workspace.Value.Item1);
 						NewQueue.Add(NewQueueItem);
 						NewBatchIdToQueueItem[(NewJob.Id, Batch.Id)] = NewQueueItem;
 					}
@@ -672,10 +672,10 @@ namespace HordeServer.Tasks.Impl
 						AgentType? AgentType;
 						if (Stream.AgentTypes.TryGetValue(Group.AgentType, out AgentType))
 						{
-							AgentWorkspace? Workspace;
+							(AgentWorkspace, bool)? Workspace;
 							if (Stream.TryGetAgentWorkspace(AgentType, out Workspace))
 							{
-								InsertQueueItem(Stream, Job, BatchIdx, AgentType.Pool, Workspace);
+								InsertQueueItem(Stream, Job, BatchIdx, AgentType.Pool, Workspace.Value.Item1);
 							}
 						}
 					}
diff --git a/Engine/Source/Programs/Horde/Horde.Build/Projects/IStream.cs b/Engine/Source/Programs/Horde/Horde.Build/Projects/IStream.cs
--- a/Engine/Source/Programs/Horde/Horde.Build/Projects/IStream.cs
+++ b/Engine/Source/Programs/Horde/Horde.Build/Projects/IStream.cs
@@ -585,13 +585,13 @@ namespace HordeServer.Models
 		/// <param name="AgentType">The agent type</param>
 		/// <param name="Workspace">Receives the agent workspace definition</param>
 		/// <returns>True if the agent type was valid, and an agent workspace could be created</returns>
-		public static bool TryGetAgentWorkspace(this IStream Stream, AgentType AgentType, [NotNullWhen(true)] out AgentWorkspace? Workspace)
+		public static bool TryGetAgentWorkspace(this IStream Stream, AgentType AgentType, [NotNullWhen(true)] out (AgentWorkspace, bool)? Workspace)
 		{
 			// Get the workspace settings
 			if (AgentType.Workspace == null)
 			{
 				// Use the default settings (fast switching workspace, clean 
-				Workspace = new AgentWorkspace(null, null, GetDefaultWorkspaceIdentifier(Stream), Stream.Name, null, false);
+				Workspace = (new AgentWorkspace(null, null, GetDefaultWorkspaceIdentifier(Stream), Stream.Name, null, false), true);
 				return true;
 			}
 			else
@@ -620,7 +620,8 @@ namespace HordeServer.Models
 				}
 
 				// Create the new workspace
-				Workspace = new AgentWorkspace(WorkspaceType.Cluster, WorkspaceType.UserName, Identifier, WorkspaceType.Stream ?? Stream.Name, WorkspaceType.View, WorkspaceType.Incremental);
+				const bool bUseAutoSdk = true;
+				Workspace = (new AgentWorkspace(WorkspaceType.Cluster, WorkspaceType.UserName, Identifier, WorkspaceType.Stream ?? Stream.Name, WorkspaceType.View, WorkspaceType.Incremental), bUseAutoSdk);
 				return true;
 			}
 		}
