Horde: Missing file.

diff --git a/Engine/Source/Programs/Horde/Horde.Build/Fleet/PoolUpdateService.cs b/Engine/Source/Programs/Horde/Horde.Build/Fleet/PoolUpdateService.cs
--- a/Engine/Source/Programs/Horde/Horde.Build/Fleet/PoolUpdateService.cs
+++ b/Engine/Source/Programs/Horde/Horde.Build/Fleet/PoolUpdateService.cs
@@ -48,6 +48,8 @@ namespace HordeServer.Services
 			this.Logger = Logger;
 		}
 
+
+
 		/// <summary>
 		/// Execute the background task
 		/// </summary>
@@ -71,6 +73,7 @@ namespace HordeServer.Services
 				List<IPool> CurrentPools = await Pools.GetAsync();
 
 				// Lookup table of pool id to workspaces
+				HashSet<PoolId> PoolsWithAutoSdk = new HashSet<PoolId>();
 				Dictionary<PoolId, List<AgentWorkspace>> PoolToAgentWorkspaces = new Dictionary<PoolId, List<AgentWorkspace>>();
 
 				// Populate the workspace list from the current stream
@@ -80,9 +83,10 @@ namespace HordeServer.Services
 					foreach (KeyValuePair<string, AgentType> AgentTypePair in ActiveStream.AgentTypes)
 					{
 						// Create the new agent workspace
-						AgentWorkspace? AgentWorkspace;
-						if (ActiveStream.TryGetAgentWorkspace(AgentTypePair.Value, out AgentWorkspace))
+						(AgentWorkspace, bool)? Result;
+						if (ActiveStream.TryGetAgentWorkspace(AgentTypePair.Value, out Result))
 						{
+							(AgentWorkspace AgentWorkspace, bool UseAutoSdk) = Result.Value;
 							AgentType AgentType = AgentTypePair.Value;
 
 							// Find or add a list of workspaces for this pool
@@ -98,6 +102,10 @@ namespace HordeServer.Services
 							{
 								AgentWorkspaces.Add(AgentWorkspace);
 							}
+							if (UseAutoSdk)
+							{
+								PoolsWithAutoSdk.Add(AgentType.Pool);
+							}
 						}
 					}
 				}
@@ -113,11 +121,12 @@ namespace HordeServer.Services
 					}
 
 					// Update the pools document
-					if (!AgentWorkspace.SetEquals(CurrentPool.Workspaces, NewWorkspaces) || CurrentPool.Workspaces.Count != NewWorkspaces.Count)
+					bool UseAutoSdk = PoolsWithAutoSdk.Contains(CurrentPool.Id);
+					if (!AgentWorkspace.SetEquals(CurrentPool.Workspaces, NewWorkspaces) || CurrentPool.Workspaces.Count != NewWorkspaces.Count || CurrentPool.UseAutoSdk != UseAutoSdk)
 					{
 						Logger.LogInformation("New workspaces for pool {Pool}:{Workspaces}", CurrentPool.Id, String.Join("", NewWorkspaces.Select(x => $"\n  Identifier=\"{x.Identifier}\", Stream={x.Stream}")));
 
-						IPool? Result = await Pools.TryUpdateAsync(CurrentPool, NewWorkspaces: NewWorkspaces);
+						IPool? Result = await Pools.TryUpdateAsync(CurrentPool, NewWorkspaces: NewWorkspaces, NewUseAutoSdk: UseAutoSdk);
 						if (Result == null)
 						{
 							Logger.LogInformation("Pool modified; will retry");
