Add support for ignoring the AutoSDK.

#rnx

diff --git a/Engine/Source/Programs/Horde/Horde.Build/Api/Streams.cs b/Engine/Source/Programs/Horde/Horde.Build/Api/Streams.cs
--- a/Engine/Source/Programs/Horde/Horde.Build/Api/Streams.cs
+++ b/Engine/Source/Programs/Horde/Horde.Build/Api/Streams.cs
@@ -159,6 +159,11 @@ namespace HordeServer.Api
 		/// Whether to use an incrementally synced workspace
 		/// </summary>
 		public bool Incremental { get; set; }
+
+		/// <summary>
+		/// Whether to use the AutoSDK
+		/// </summary>
+		public bool UseAutoSdk { get; set; } = true;
 	}
 
 	/// <summary>
@@ -311,6 +316,11 @@ namespace HordeServer.Api
 		/// </summary>
 		public bool Incremental { get; set; }
 
+		/// <summary>
+		/// Whether to use the AutoSDK
+		/// </summary>
+		public bool UseAutoSdk { get; set; }
+
 		/// <summary>
 		/// Default constructor
 		/// </summary>
@@ -328,7 +338,8 @@ namespace HordeServer.Api
 		/// <param name="Stream">Override for the stream to sync</param>
 		/// <param name="View">Custom view for the workspace</param>
 		/// <param name="bIncremental">Whether to use an incrementally synced workspace</param>
-		public GetWorkspaceTypeResponse(string? Cluster, string? ServerAndPort, string? UserName, string? Identifier, string? Stream, List<string>? View, bool bIncremental)
+		/// <param name="bUseAutoSdk">Whether to use the AutoSDK</param>
+		public GetWorkspaceTypeResponse(string? Cluster, string? ServerAndPort, string? UserName, string? Identifier, string? Stream, List<string>? View, bool bIncremental, bool bUseAutoSdk)
 		{
 			this.Cluster = Cluster;
 			this.ServerAndPort = ServerAndPort;
@@ -337,6 +348,7 @@ namespace HordeServer.Api
 			this.Stream = Stream;
 			this.View = View;
 			this.Incremental = bIncremental;
+			this.UseAutoSdk = bUseAutoSdk;
 		}
 	}
 
diff --git a/Engine/Source/Programs/Horde/Horde.Build/Jobs/Impl/JobTaskSource.cs b/Engine/Source/Programs/Horde/Horde.Build/Jobs/Impl/JobTaskSource.cs
--- a/Engine/Source/Programs/Horde/Horde.Build/Jobs/Impl/JobTaskSource.cs
+++ b/Engine/Source/Programs/Horde/Horde.Build/Jobs/Impl/JobTaskSource.cs
@@ -77,6 +77,11 @@ namespace HordeServer.Tasks.Impl
 			/// </summary>
 			public AgentWorkspace Workspace;
 
+			/// <summary>
+			/// Whether or not to use the AutoSDK.
+			/// </summary>
+			public bool UseAutoSDK;
+
 			/// <summary>
 			/// Task for creating a lease and assigning to a waiter
 			/// </summary>
@@ -106,13 +111,15 @@ namespace HordeServer.Tasks.Impl
 			/// <param name="BatchIdx">The batch index to execute</param>
 			/// <param name="PoolId">Unique id of the pool of machines to allocate from</param>
 			/// <param name="Workspace">The workspace that this job should run in</param>
-			public QueueItem(IStream Stream, IJob Job, int BatchIdx, PoolId PoolId, AgentWorkspace Workspace)
+			/// <param name="UseAutoSDK">Whether or not to use the AutoSDK</param>
+			public QueueItem(IStream Stream, IJob Job, int BatchIdx, PoolId PoolId, AgentWorkspace Workspace, bool UseAutoSDK)
 			{
 				this.Stream = Stream;
 				this.Job = Job;
 				this.BatchIdx = BatchIdx;
 				this.PoolId = PoolId;
 				this.Workspace = Workspace;
+				this.UseAutoSDK = UseAutoSDK;
 			}
 		}
 
@@ -488,13 +495,14 @@ namespace HordeServer.Tasks.Impl
 					{
 						NewJob = await SkipBatchAsync(NewJob, Batch.Id, Graph, JobStepBatchError.NoAgentsOnline);
 					}
-					else if (!Stream.TryGetAgentWorkspace(AgentType, out (AgentWorkspace, bool)? Workspace))
+					else if (!Stream.TryGetAgentWorkspace(AgentType, out (AgentWorkspace, bool)? WorkspaceResult))
 					{
 						NewJob = await SkipBatchAsync(NewJob, Batch.Id, Graph, JobStepBatchError.UnknownWorkspace);
 					}
 					else
 					{
-						QueueItem NewQueueItem = new QueueItem(Stream, NewJob, BatchIdx, AgentType.Pool, Workspace.Value.Item1);
+						(AgentWorkspace Workspace, bool UseAutoSDK) = WorkspaceResult.Value;
+						QueueItem NewQueueItem = new QueueItem(Stream, NewJob, BatchIdx, AgentType.Pool, Workspace, UseAutoSDK);
 						NewQueue.Add(NewQueueItem);
 						NewBatchIdToQueueItem[(NewJob.Id, Batch.Id)] = NewQueueItem;
 					}
@@ -659,7 +667,7 @@ namespace HordeServer.Tasks.Impl
 								else
 								{
 									RemoveQueueItem(ExistingItem);
-									InsertQueueItem(Stream, Job, BatchIdx, ExistingItem.PoolId, ExistingItem.Workspace);
+									InsertQueueItem(Stream, Job, BatchIdx, ExistingItem.PoolId, ExistingItem.Workspace, ExistingItem.UseAutoSDK);
 								}
 							}
 							continue;
@@ -672,10 +680,11 @@ namespace HordeServer.Tasks.Impl
 						AgentType? AgentType;
 						if (Stream.AgentTypes.TryGetValue(Group.AgentType, out AgentType))
 						{
-							(AgentWorkspace, bool)? Workspace;
-							if (Stream.TryGetAgentWorkspace(AgentType, out Workspace))
+							(AgentWorkspace, bool)? Result;
+							if (Stream.TryGetAgentWorkspace(AgentType, out Result))
 							{
-								InsertQueueItem(Stream, Job, BatchIdx, AgentType.Pool, Workspace.Value.Item1);
+								(AgentWorkspace AgentWorkspace, bool UseAutoSdk) = Result.Value;
+								InsertQueueItem(Stream, Job, BatchIdx, AgentType.Pool, AgentWorkspace, UseAutoSdk);
 							}
 						}
 					}
@@ -777,7 +786,7 @@ namespace HordeServer.Tasks.Impl
 				Globals Globals = await DatabaseService.GetGlobalsAsync();
 
 				// Encode the payload
-				ExecuteJobTask? Task = await CreateExecuteJobTaskAsync(Item.Stream, Job, Batch, Agent, Item.Workspace, LogId);
+				ExecuteJobTask? Task = await CreateExecuteJobTaskAsync(Item.Stream, Job, Batch, Agent, Item.Workspace, Item.UseAutoSDK, LogId);
 				if (Task != null)
 				{
 					byte[] Payload = Any.Pack(Task).ToByteArray();
@@ -826,7 +835,7 @@ namespace HordeServer.Tasks.Impl
 			return null;
 		}
 
-		async Task<ExecuteJobTask?> CreateExecuteJobTaskAsync(IStream Stream, IJob Job, IJobStepBatch Batch, IAgent Agent, AgentWorkspace Workspace, LogId LogId)
+		async Task<ExecuteJobTask?> CreateExecuteJobTaskAsync(IStream Stream, IJob Job, IJobStepBatch Batch, IAgent Agent, AgentWorkspace Workspace, bool bUseAutoSDK, LogId LogId)
 		{
 			// Get the lease name
 			StringBuilder LeaseName = new StringBuilder($"{Stream.Name} - ");
@@ -858,7 +867,7 @@ namespace HordeServer.Tasks.Impl
 				return null;
 			}
 
-			AgentWorkspace? AutoSdkWorkspace = Agent.GetAutoSdkWorkspace(Cluster);
+			AgentWorkspace? AutoSdkWorkspace = bUseAutoSDK ? Agent.GetAutoSdkWorkspace(Cluster) : null;
 			if (AutoSdkWorkspace != null)
 			{
 				if (!await Agent.TryAddWorkspaceMessage(AutoSdkWorkspace, Cluster, PerforceLoadBalancer, Workspaces))
@@ -1082,12 +1091,13 @@ namespace HordeServer.Tasks.Impl
 		/// <param name="BatchIdx"></param>
 		/// <param name="PoolId">The pool to use</param>
 		/// <param name="Workspace">The workspace for this item to run in</param>
+		/// <param name="UseAutoSDK">Whether or not to use the AutoSDK</param>
 		/// <returns></returns>
-		void InsertQueueItem(IStream Stream, IJob Job, int BatchIdx, PoolId PoolId, AgentWorkspace Workspace)
+		void InsertQueueItem(IStream Stream, IJob Job, int BatchIdx, PoolId PoolId, AgentWorkspace Workspace, bool UseAutoSDK)
 		{
 			Logger.LogDebug("Adding queued job {JobId}, batch {BatchId} [Pool: {Pool}, Workspace: {Workspace}]", Job.Id, Job.Batches[BatchIdx].Id, PoolId, Workspace.Identifier);
 
-			QueueItem NewItem = new QueueItem(Stream, Job, BatchIdx, PoolId, Workspace);
+			QueueItem NewItem = new QueueItem(Stream, Job, BatchIdx, PoolId, Workspace, UseAutoSDK);
 			BatchIdToQueueItem[NewItem.Id] = NewItem;
 			Queue.Add(NewItem);
 
diff --git a/Engine/Source/Programs/Horde/Horde.Build/Projects/IStream.cs b/Engine/Source/Programs/Horde/Horde.Build/Projects/IStream.cs
--- a/Engine/Source/Programs/Horde/Horde.Build/Projects/IStream.cs
+++ b/Engine/Source/Programs/Horde/Horde.Build/Projects/IStream.cs
@@ -197,6 +197,11 @@ namespace HordeServer.Models
 		/// </summary>
 		public bool Incremental { get; set; }
 
+		/// <summary>
+		/// Whether to use the AutoSDK
+		/// </summary>
+		public bool UseAutoSdk { get; set; } = true;
+
 		/// <summary>
 		/// Default constructor
 		/// </summary>
@@ -218,6 +223,7 @@ namespace HordeServer.Models
 			Stream = Request.Stream;
 			View = Request.View;
 			Incremental = Request.Incremental;
+			UseAutoSdk = Request.UseAutoSdk;
 		}
 
 		/// <summary>
@@ -237,7 +243,7 @@ namespace HordeServer.Models
 		/// <returns>The response object</returns>
 		public Api.GetWorkspaceTypeResponse ToApiResponse()
 		{
-			return new Api.GetWorkspaceTypeResponse(Cluster, ServerAndPort, UserName, Identifier, Stream, View, Incremental);
+			return new Api.GetWorkspaceTypeResponse(Cluster, ServerAndPort, UserName, Identifier, Stream, View, Incremental, UseAutoSdk);
 		}
 	}
 
@@ -620,8 +626,7 @@ namespace HordeServer.Models
 				}
 
 				// Create the new workspace
-				const bool bUseAutoSdk = true;
-				Workspace = (new AgentWorkspace(WorkspaceType.Cluster, WorkspaceType.UserName, Identifier, WorkspaceType.Stream ?? Stream.Name, WorkspaceType.View, WorkspaceType.Incremental), bUseAutoSdk);
+				Workspace = (new AgentWorkspace(WorkspaceType.Cluster, WorkspaceType.UserName, Identifier, WorkspaceType.Stream ?? Stream.Name, WorkspaceType.View, WorkspaceType.Incremental), WorkspaceType.UseAutoSdk);
 				return true;
 			}
 		}
