Horde: Add support for generic OpenID Connect authentication

The Okta flavored OIDC auth method is still there and it's possible when can remove that one later.
I wanted to decouple adding this new generic auth method from switching the current one in use internally at Epic.

diff --git a/Engine/Source/Programs/Horde/Horde.Build/Api/Config.cs b/Engine/Source/Programs/Horde/Horde.Build/Api/Config.cs
--- a/Engine/Source/Programs/Horde/Horde.Build/Api/Config.cs
+++ b/Engine/Source/Programs/Horde/Horde.Build/Api/Config.cs
@@ -323,7 +323,6 @@ namespace HordeServer.Api
 			this.JwtSecret = Settings.JwtSecret;
 			this.JwtExpiryTimeHours = Settings.JwtExpiryTimeHours;
 
-			this.DisableAuth = Settings.DisableAuth;
 			this.CorsEnabled = Settings.CorsEnabled;
 			this.CorsOrigin = Settings.CorsOrigin;
 			this.EnableScheduleInTestData = Settings.EnableScheduleInTestData;
diff --git a/Engine/Source/Programs/Horde/Horde.Build/Authentication/HordeOpenIdConnectHandler.cs b/Engine/Source/Programs/Horde/Horde.Build/Authentication/HordeOpenIdConnectHandler.cs
new file mode 100644
--- /dev/null
+++ b/Engine/Source/Programs/Horde/Horde.Build/Authentication/HordeOpenIdConnectHandler.cs
@@ -0,0 +1,134 @@
+// Copyright Epic Games, Inc. All Rights Reserved.
+
+using HordeServer.Collections;
+using HordeServer.Models;
+using HordeServer.Utilities;
+using Microsoft.AspNetCore.Authentication;
+using Microsoft.AspNetCore.Authentication.OAuth.Claims;
+using Microsoft.AspNetCore.Authentication.OpenIdConnect;
+using Microsoft.Extensions.DependencyInjection;
+using Microsoft.Extensions.DependencyInjection.Extensions;
+using Microsoft.Extensions.Logging;
+using Microsoft.Extensions.Options;
+using Microsoft.IdentityModel.Protocols.OpenIdConnect;
+using System;
+using System.Linq;
+using System.Security.Claims;
+using System.Text.Encodings.Web;
+using System.Text.Json;
+using System.Threading.Tasks;
+
+namespace HordeServer.Authentication
+{
+	class HordeOpenIdConnectHandler : OpenIdConnectHandler
+	{
+		IUserCollection UserCollection;
+
+		public HordeOpenIdConnectHandler(IOptionsMonitor<OpenIdConnectOptions> Options, ILoggerFactory Logger, HtmlEncoder HtmlEncoder, UrlEncoder Encoder, ISystemClock Clock, IUserCollection UserCollection)
+			: base(Options, Logger, HtmlEncoder, Encoder, Clock)
+		{
+			this.UserCollection = UserCollection;
+		}
+
+		/// <summary>
+		/// Try to map a field found in user info as a claim
+		///
+		/// If the claim is already set, it is skipped.
+		/// </summary>
+		/// <param name="UserInfo">User info from the OIDC /userinfo endpoint </param>
+		/// <param name="Identity">Identity which claims to modify</param>
+		/// <param name="ClaimName">Name of claim</param>
+		/// <param name="UserInfoFields">List of field names to try when mapping</param>
+		/// <exception cref="Exception">Thrown if no field name matches</exception>
+		private static void MapUserInfoFieldToClaim(JsonElement UserInfo, ClaimsIdentity Identity, string ClaimName, string[] UserInfoFields)
+		{
+			if (Identity.FindFirst(ClaimName) != null)
+			{
+				return;
+			}
+
+			foreach (string UserDataField in UserInfoFields)
+			{
+				if (UserInfo.TryGetProperty(UserDataField, out JsonElement FieldElement))
+				{
+					Identity.AddClaim(new Claim(ClaimName, FieldElement.ToString()));
+					return;
+				}
+			}
+
+			string Message = $"Unable to map a field from user info to claim '{ClaimName}' using list [{string.Join(", ", UserInfoFields)}].";
+			Message += " UserInfo: " + UserInfo;
+			throw new Exception(Message);
+		}
+
+		public static void AddUserInfoClaims(ServerSettings Settings, JsonElement UserInfo, ClaimsIdentity Identity)
+		{
+			MapUserInfoFieldToClaim(UserInfo, Identity, ClaimTypes.Name, Settings.OidcClaimNameMapping);
+			MapUserInfoFieldToClaim(UserInfo, Identity, ClaimTypes.Email, Settings.OidcClaimEmailMapping);
+			MapUserInfoFieldToClaim(UserInfo, Identity, HordeClaimTypes.User, Settings.OidcClaimHordeUserMapping);
+			MapUserInfoFieldToClaim(UserInfo, Identity, HordeClaimTypes.PerforceUser, Settings.OidcClaimHordePerforceUserMapping);
+
+			if (UserInfo.TryGetProperty("groups", out JsonElement GroupsElement) && GroupsElement.ValueKind == JsonValueKind.Array)
+			{
+				for (int Idx = 0; Idx < GroupsElement.GetArrayLength(); Idx++)
+				{
+					Identity.AddClaim(new Claim(ClaimTypes.Role, GroupsElement[Idx].ToString()!));
+				}
+			}
+		}
+
+		protected override async Task<HandleRequestResult> HandleRemoteAuthenticateAsync()
+		{
+			// Authenticate with the OIDC provider
+			HandleRequestResult Result = await base.HandleRemoteAuthenticateAsync();
+			if (!Result.Succeeded)
+			{
+				return Result;
+			}
+
+			ClaimsIdentity? Identity = (ClaimsIdentity?)Result.Principal?.Identity;
+			if (Identity == null)
+			{
+				return HandleRequestResult.Fail("No identity specified");
+			}
+
+			string Login = Identity.FindFirst(ClaimTypes.Name)!.Value;
+			string? Name = Identity.FindFirst("name")?.Value;
+			string? Email = Identity.FindFirst(ClaimTypes.Email)?.Value;
+
+			IUser User = await UserCollection.FindOrAddUserByLoginAsync(Login, Name, Email);
+			Identity.AddClaim(new Claim(HordeClaimTypes.UserId, User.Id.ToString()));
+
+			await UserCollection.UpdateClaimsAsync(User.Id, Identity.Claims.Select(x => new UserClaim(x.Type, x.Value)));
+
+			return Result;
+		}
+	}
+
+	static class OpenIdConnectHandlerExtensions
+	{
+		class MapRolesClaimAction : ClaimAction
+		{
+			private ServerSettings Settings;
+			public MapRolesClaimAction(ServerSettings Settings) : base(ClaimTypes.Role, ClaimTypes.Role) { this.Settings = Settings; }
+			public override void Run(JsonElement UserData, ClaimsIdentity Identity, string Issuer) { HordeOpenIdConnectHandler.AddUserInfoClaims(Settings, UserData, Identity); }
+		}
+
+		static void ApplyDefaultOptions(ServerSettings Settings, OpenIdConnectOptions Options, Action<OpenIdConnectOptions> Handler)
+		{
+			Options.ResponseType = OpenIdConnectResponseType.Code;
+			Options.GetClaimsFromUserInfoEndpoint = true;
+			Options.SaveTokens = true;
+			Options.TokenValidationParameters.NameClaimType = "name";
+			Options.ClaimActions.Add(new MapRolesClaimAction(Settings));
+
+			Handler(Options);
+		}
+
+		public static void AddHordeOpenId(this AuthenticationBuilder Builder, ServerSettings Settings, string AuthenticationScheme, string DisplayName, Action<OpenIdConnectOptions> Handler)
+		{
+			Builder.Services.TryAddEnumerable(ServiceDescriptor.Singleton<IPostConfigureOptions<OpenIdConnectOptions>, OpenIdConnectPostConfigureOptions>());
+			Builder.AddRemoteScheme<OpenIdConnectOptions, HordeOpenIdConnectHandler>(AuthenticationScheme, DisplayName, Options => ApplyDefaultOptions(Settings, Options, Handler));
+		}
+	}
+}
diff --git a/Engine/Source/Programs/Horde/Horde.Build/Controllers/AccountController.cs b/Engine/Source/Programs/Horde/Horde.Build/Controllers/AccountController.cs
--- a/Engine/Source/Programs/Horde/Horde.Build/Controllers/AccountController.cs
+++ b/Engine/Source/Programs/Horde/Horde.Build/Controllers/AccountController.cs
@@ -15,6 +15,7 @@ using System.Linq;
 using System.Net;
 using System.Text;
 using System.Threading.Tasks;
+using Microsoft.Extensions.Options;
 
 namespace HordeServer.Controllers
 {
@@ -25,8 +26,6 @@ namespace HordeServer.Controllers
 	[Route("[controller]")]
 	public class AccountController : Controller
 	{
-		const string DefaultAuthenticationScheme = OktaDefaults.AuthenticationScheme;
-
 		/// <summary>
 		/// Style sheet for HTML responses
 		/// </summary>
@@ -41,14 +40,37 @@ namespace HordeServer.Controllers
 		/// The ACL service singleton
 		/// </summary>
 		AclService AclService;
+		
+		/// <summary>
+		/// Authentication scheme in use
+		/// </summary>
+		string AuthenticationScheme;
 
 		/// <summary>
 		/// Constructor
 		/// </summary>
 		/// <param name="AclService">ACL service instance</param>
-		public AccountController(AclService AclService)
+		/// <param name="ServerSettings">Server settings</param>
+		public AccountController(AclService AclService, IOptionsMonitor<ServerSettings> ServerSettings)
 		{
 			this.AclService = AclService;
+			this.AuthenticationScheme = GetAuthScheme(ServerSettings.CurrentValue.AuthMethod);
+		}
+
+		/// <summary>
+		/// Get auth scheme name for a given auth method
+		/// </summary>
+		/// <param name="Method">Authentication method</param>
+		/// <returns>Name of authentication scheme</returns>
+		public static string GetAuthScheme(AuthMethod Method)
+		{
+			return Method switch
+			{
+				AuthMethod.Anonymous => AnonymousAuthenticationHandler.AuthenticationScheme,
+				AuthMethod.Okta => OktaDefaults.AuthenticationScheme,
+				AuthMethod.OpenIdConnect => OpenIdConnectDefaults.AuthenticationScheme,
+				_ => throw new ArgumentOutOfRangeException(nameof(Method), Method, null)
+			};
 		}
 
 		/// <summary>
@@ -102,7 +124,7 @@ namespace HordeServer.Controllers
 		[Route("/account/login")]
 		public IActionResult Login()
 		{
-			return new ChallengeResult(DefaultAuthenticationScheme, new AuthenticationProperties { RedirectUri = "/account" });
+			return new ChallengeResult(AuthenticationScheme, new AuthenticationProperties { RedirectUri = "/account" });
 		}
 
 		/// <summary>
@@ -116,7 +138,7 @@ namespace HordeServer.Controllers
 			await HttpContext.SignOutAsync(CookieAuthenticationDefaults.AuthenticationScheme);
 			try
 			{
-				await HttpContext.SignOutAsync(DefaultAuthenticationScheme);
+				await HttpContext.SignOutAsync(AuthenticationScheme);
 			}
 			catch
 			{
diff --git a/Engine/Source/Programs/Horde/Horde.Build/Controllers/DashboardController.cs b/Engine/Source/Programs/Horde/Horde.Build/Controllers/DashboardController.cs
--- a/Engine/Source/Programs/Horde/Horde.Build/Controllers/DashboardController.cs
+++ b/Engine/Source/Programs/Horde/Horde.Build/Controllers/DashboardController.cs
@@ -8,6 +8,9 @@ using Microsoft.AspNetCore.Authentication.Cookies;
 using Microsoft.AspNetCore.Authorization;
 using Microsoft.AspNetCore.Mvc;
 using System.Threading.Tasks;
+using HordeServer;
+using HordeServer.Controllers;
+using Microsoft.Extensions.Options;
 
 /// <summary>	
 /// Dashboard authorization challenge controller	
@@ -16,8 +19,20 @@ using System.Threading.Tasks;
 [Route("[controller]")]
 public class DashboardChallengeController : Controller
 {
-	const string DefaultAuthenticationScheme = OktaDefaults.AuthenticationScheme;
+	/// <summary>
+	/// Authentication scheme in use
+	/// </summary>
+	string AuthenticationScheme;
 
+	/// <summary>
+	/// Constructor
+	/// </summary>
+	/// <param name="ServerSettings">Server settings</param>
+	public DashboardChallengeController(IOptionsMonitor<ServerSettings> ServerSettings)
+	{
+		AuthenticationScheme = AccountController.GetAuthScheme(ServerSettings.CurrentValue.AuthMethod);
+	}
+	
 	/// <summary>	
 	/// Challenge endpoint for the dashboard, using cookie authentication scheme	
 	/// </summary>	
@@ -39,7 +54,7 @@ public class DashboardChallengeController : Controller
 	[Route("/api/v1/dashboard/login")]
 	public IActionResult Login([FromQuery] string? Redirect)
 	{
-		return new ChallengeResult(DefaultAuthenticationScheme, new AuthenticationProperties { RedirectUri = Redirect ?? "/" });
+		return new ChallengeResult(AuthenticationScheme, new AuthenticationProperties { RedirectUri = Redirect ?? "/" });
 	}
 
 	/// <summary>
@@ -53,7 +68,7 @@ public class DashboardChallengeController : Controller
 		await HttpContext.SignOutAsync(CookieAuthenticationDefaults.AuthenticationScheme);
 		try
 		{
-			await HttpContext.SignOutAsync(DefaultAuthenticationScheme);
+			await HttpContext.SignOutAsync(AuthenticationScheme);
 		}
 #pragma warning disable CA1031 // Do not catch general exception types
 		catch
diff --git a/Engine/Source/Programs/Horde/Horde.Build/ServerSettings.cs b/Engine/Source/Programs/Horde/Horde.Build/ServerSettings.cs
--- a/Engine/Source/Programs/Horde/Horde.Build/ServerSettings.cs
+++ b/Engine/Source/Programs/Horde/Horde.Build/ServerSettings.cs
@@ -59,6 +59,27 @@ namespace HordeServer
 		/// </summary>
 		Aws,
 	}
+	
+	/// <summary>
+	/// Authentication method used for logging users in
+	/// </summary>
+	public enum AuthMethod
+	{
+		/// <summary>
+		/// No authentication enabled, mainly for demo and testing purposes
+		/// </summary>
+		Anonymous,
+
+		/// <summary>
+		/// OpenID Connect authentication, tailored for Okta
+		/// </summary>
+		Okta,
+		
+		/// <summary>
+		/// Generic OpenID Connect authentication, recommended for most
+		/// </summary>
+		OpenIdConnect,
+	}
 
 	/// <summary>
 	/// Settings for a remote execution instance
@@ -179,6 +200,11 @@ namespace HordeServer
 		/// </summary>
 		public string? ServerPrivateCert { get; set; }
 
+		/// <summary>
+		/// Issuer for tokens from the auth provider
+		/// </summary>
+		public AuthMethod AuthMethod { get; set; } = AuthMethod.Anonymous;
+		
 		/// <summary>
 		/// Issuer for tokens from the auth provider
 		/// </summary>
@@ -188,11 +214,41 @@ namespace HordeServer
 		/// Client id for the OIDC authority
 		/// </summary>
 		public string? OidcClientId { get; set; }
+		
+		/// <summary>
+		/// Client secret for the OIDC authority
+		/// </summary>
+		public string? OidcClientSecret { get; set; }
 
 		/// <summary>
 		/// Optional redirect url provided to OIDC login
 		/// </summary>
-		public string? OidcSigninRedirect { get; set; }		
+		public string? OidcSigninRedirect { get; set; }
+
+		/// <summary>
+		/// OpenID Connect scopes to request when signing in
+		/// </summary>
+		public string[] OidcRequestedScopes { get; set; } = { "profile", "email", "openid" };
+		
+		/// <summary>
+		/// List of fields in /userinfo endpoint to try map to the standard name claim (see System.Security.Claims.ClaimTypes.Name)
+		/// </summary>
+		public string[] OidcClaimNameMapping { get; set; } = { "preferred_username", "email" };
+		
+		/// <summary>
+		/// List of fields in /userinfo endpoint to try map to the standard email claim (see System.Security.Claims.ClaimTypes.Email)
+		/// </summary>
+		public string[] OidcClaimEmailMapping { get; set; } = { "email" };
+		
+		/// <summary>
+		/// List of fields in /userinfo endpoint to try map to the Horde user claim (see HordeClaimTypes.User)
+		/// </summary>
+		public string[] OidcClaimHordeUserMapping { get; set; } = { "preferred_username", "email" };
+		
+		/// <summary>
+		/// List of fields in /userinfo endpoint to try map to the Horde Perforce user claim (see HordeClaimTypes.PerforceUser)
+		/// </summary>
+		public string[] OidcClaimHordePerforceUserMapping { get; set; } = { "preferred_username", "email" };
 
 		/// <summary>
 		/// Name of the issuer in bearer tokens from the server
@@ -209,11 +265,6 @@ namespace HordeServer
 		/// </summary>
 		public int JwtExpiryTimeHours { get; set; } = 4;
 
-		/// <summary>
-		/// Disable authentication for debugging purposes
-		/// </summary>
-		public bool DisableAuth { get; set; }
-
 		/// <summary>
 		/// Whether to enable Cors, generally for development purposes
 		/// </summary>
diff --git a/Engine/Source/Programs/Horde/Horde.Build/Startup.cs b/Engine/Source/Programs/Horde/Horde.Build/Startup.cs
--- a/Engine/Source/Programs/Horde/Horde.Build/Startup.cs
+++ b/Engine/Source/Programs/Horde/Horde.Build/Startup.cs
@@ -448,20 +448,34 @@ namespace HordeServer
 
 			AuthenticationBuilder AuthBuilder = Services.AddAuthentication(Options =>
 				{
-					if (Settings.DisableAuth)
+					switch (Settings.AuthMethod)
 					{
-						Options.DefaultAuthenticateScheme = AnonymousAuthenticationHandler.AuthenticationScheme;
-						Options.DefaultSignInScheme = AnonymousAuthenticationHandler.AuthenticationScheme;
-						Options.DefaultChallengeScheme = AnonymousAuthenticationHandler.AuthenticationScheme;
-					}
-					else
-					{
-						// If an authentication cookie is present, use it to get authentication information
-						Options.DefaultAuthenticateScheme = CookieAuthenticationDefaults.AuthenticationScheme;
-						Options.DefaultSignInScheme = CookieAuthenticationDefaults.AuthenticationScheme;
-
-						// If authentication is required, and no cookie is present, use OIDC to sign in
-						Options.DefaultChallengeScheme = OktaDefaults.AuthenticationScheme;
+						case AuthMethod.Anonymous:
+							Options.DefaultAuthenticateScheme = AnonymousAuthenticationHandler.AuthenticationScheme;
+							Options.DefaultSignInScheme = AnonymousAuthenticationHandler.AuthenticationScheme;
+							Options.DefaultChallengeScheme = AnonymousAuthenticationHandler.AuthenticationScheme;
+							break;
+						
+						case AuthMethod.Okta:
+							// If an authentication cookie is present, use it to get authentication information
+							Options.DefaultAuthenticateScheme = CookieAuthenticationDefaults.AuthenticationScheme;
+							Options.DefaultSignInScheme = CookieAuthenticationDefaults.AuthenticationScheme;
+
+							// If authentication is required, and no cookie is present, use OIDC to sign in
+							Options.DefaultChallengeScheme = OktaDefaults.AuthenticationScheme;
+							break;
+						
+						case AuthMethod.OpenIdConnect:
+							// If an authentication cookie is present, use it to get authentication information
+							Options.DefaultAuthenticateScheme = CookieAuthenticationDefaults.AuthenticationScheme;
+							Options.DefaultSignInScheme = CookieAuthenticationDefaults.AuthenticationScheme;
+
+							// If authentication is required, and no cookie is present, use OIDC to sign in
+							Options.DefaultChallengeScheme = OpenIdConnectDefaults.AuthenticationScheme;
+							break;
+						
+						default:
+							throw new ArgumentException($"Invalid auth method {Settings.AuthMethod}");
 					}
 				});
 
@@ -489,36 +503,67 @@ namespace HordeServer
 			AuthBuilder.AddServiceAccount(Options => { });
 			Schemes.Add(ServiceAccountAuthHandler.AuthenticationScheme);
 
-			if (Settings.DisableAuth)
-			{
-				AuthBuilder.AddAnonymous(Options =>
-				{
-					Options.AdminClaimType = Settings.AdminClaimType;
-					Options.AdminClaimValue = Settings.AdminClaimValue;
-				});
-				Schemes.Add(AnonymousAuthenticationHandler.AuthenticationScheme);
-			}
-			else if (Settings.OidcClientId != null && Settings.OidcAuthority != null)
+			
+			switch (Settings.AuthMethod)
 			{
-				AuthBuilder.AddOkta(OktaDefaults.AuthenticationScheme, OpenIdConnectDefaults.DisplayName, 
-					Options =>
+				case AuthMethod.Anonymous:
+					AuthBuilder.AddAnonymous(Options =>
 					{
-						Options.Authority = Settings.OidcAuthority;
-						Options.ClientId = Settings.OidcClientId;
+						Options.AdminClaimType = Settings.AdminClaimType;
+						Options.AdminClaimValue = Settings.AdminClaimValue;
+					});
+					Schemes.Add(AnonymousAuthenticationHandler.AuthenticationScheme);
+					break;
+						
+				case AuthMethod.Okta:
+					AuthBuilder.AddOkta(OktaDefaults.AuthenticationScheme, OpenIdConnectDefaults.DisplayName, Options =>
+						{
+							Options.Authority = Settings.OidcAuthority;
+							Options.ClientId = Settings.OidcClientId;
 
-						if (!String.IsNullOrEmpty(Settings.OidcSigninRedirect))
+							if (!String.IsNullOrEmpty(Settings.OidcSigninRedirect))
+							{
+								Options.Events = new OpenIdConnectEvents
+								{
+									OnRedirectToIdentityProvider = async RedirectContext =>
+									{
+										RedirectContext.ProtocolMessage.RedirectUri = Settings.OidcSigninRedirect;
+										await Task.CompletedTask;
+									}
+								};
+							}
+						});
+					Schemes.Add(OktaDefaults.AuthenticationScheme);
+					break;
+						
+				case AuthMethod.OpenIdConnect:
+					AuthBuilder.AddHordeOpenId(Settings, OpenIdConnectDefaults.AuthenticationScheme, OpenIdConnectDefaults.DisplayName, Options =>
 						{
-							Options.Events = new OpenIdConnectEvents
+							Options.Authority = Settings.OidcAuthority;
+							Options.ClientId = Settings.OidcClientId;
+							Options.ClientSecret = Settings.OidcClientSecret;
+							foreach (string Scope in Settings.OidcRequestedScopes)
+							{
+								Options.Scope.Add(Scope);
+							}
+
+							if (!String.IsNullOrEmpty(Settings.OidcSigninRedirect))
 							{
-								OnRedirectToIdentityProvider = async RedirectContext =>
+								Options.Events = new OpenIdConnectEvents
 								{
-									RedirectContext.ProtocolMessage.RedirectUri = Settings.OidcSigninRedirect;
-									await Task.CompletedTask;
-								}
-							};
-						}
-					});
-				Schemes.Add(OktaDefaults.AuthenticationScheme);
+									OnRedirectToIdentityProvider = async RedirectContext =>
+									{
+										RedirectContext.ProtocolMessage.RedirectUri = Settings.OidcSigninRedirect;
+										await Task.CompletedTask;
+									}
+								};
+							}
+						});
+					Schemes.Add(OpenIdConnectDefaults.AuthenticationScheme);
+					break;
+						
+				default:
+					throw new ArgumentException($"Invalid auth method {Settings.AuthMethod}");
 			}
 
 			AuthBuilder.AddScheme<JwtBearerOptions, HordeJwtBearerHandler>(HordeJwtBearerHandler.AuthenticationScheme, Options => { });
@@ -774,6 +819,13 @@ namespace HordeServer
 			App.UseForwardedHeaders();
 			App.UseExceptionHandler("/api/v1/error");
 
+			// Used for allowing auth cookies in combination with OpenID Connect auth (for example, Google Auth did not work with these unset)
+			App.UseCookiePolicy(new CookiePolicyOptions()
+			{
+				MinimumSameSitePolicy = SameSiteMode.None,
+				CheckConsentNeeded = _ => true
+			});
+
 			if (Settings.Value.CorsEnabled)
 			{
 				App.UseCors("CorsPolicy");
