Horde Storage - Added more information when objects mismatch during finalize step

diff --git a/Engine/Source/Programs/Horde/HordeStorage/Horde.Storage/Implementation/Objects/ScyllaReferencesStore.cs b/Engine/Source/Programs/Horde/HordeStorage/Horde.Storage/Implementation/Objects/ScyllaReferencesStore.cs
--- a/Engine/Source/Programs/Horde/HordeStorage/Horde.Storage/Implementation/Objects/ScyllaReferencesStore.cs
+++ b/Engine/Source/Programs/Horde/HordeStorage/Horde.Storage/Implementation/Objects/ScyllaReferencesStore.cs
@@ -94,7 +94,7 @@ namespace Horde.Storage.Implementation
             using Scope _ = Tracer.Instance.StartActive("scylla.finalize");
             ObjectRecord o = await Get(ns, bucket, name);
             if (!o.BlobIdentifier.Equals(blobHash))
-                throw new ObjectHashMismatchException(ns, bucket, name, blobHash);
+                throw new ObjectHashMismatchException(ns, bucket, name, blobHash, o.BlobIdentifier);
 
             AppliedInfo<ScyllaObject> info = await _mapper.UpdateIfAsync<ScyllaObject>("SET is_finalized=true WHERE namespace=? AND bucket=? AND name=?", ns.ToString(), bucket.ToString(), name.ToString());
             if (!info.Applied)
@@ -277,7 +277,7 @@ namespace Horde.Storage.Implementation
 
     public class ObjectHashMismatchException : Exception
     {
-        public ObjectHashMismatchException(NamespaceId ns, BucketId bucket, KeyId name, BlobIdentifier blobHash) : base($"Object {name} in bucket {bucket} and namespace {ns} did not reference hash {blobHash}")
+        public ObjectHashMismatchException(NamespaceId ns, BucketId bucket, KeyId name, BlobIdentifier suppliedHash, BlobIdentifier actualHash) : base($"Object {name} in bucket {bucket} and namespace {ns} did not reference hash {suppliedHash} was referencing {actualHash}")
         {
         }
     }
