Horde: Add optimized endpoint in JobsController for fetching jobs in a specific stream

Also adds an index hint parameter to FindAsync in JobCollection. The new specialized method in that collection uses that parameter to optimize the database query.
When the dashboard can use this endpoint, load times can hopefully be improved for listing jobs in the browser.

diff --git a/Engine/Source/Programs/Horde/Horde.Build/Controllers/JobsController.cs b/Engine/Source/Programs/Horde/Horde.Build/Controllers/JobsController.cs
--- a/Engine/Source/Programs/Horde/Horde.Build/Controllers/JobsController.cs
+++ b/Engine/Source/Programs/Horde/Horde.Build/Controllers/JobsController.cs
@@ -700,6 +700,78 @@ namespace HordeServer.Controllers
 			}
 			return Responses;
 		}
+		
+		/// <summary>
+		/// Find jobs for a stream with given templates, sorted by creation date
+		/// </summary>
+		/// <param name="StreamId">The stream to search for</param>
+		/// <param name="Templates">List of templates to find</param>
+		/// <param name="PreflightStartedByUserId">User id for which to include preflight jobs</param>
+		/// <param name="Filter">Filter for properties to return</param>
+		/// <param name="Index">Index of the first result to be returned</param>
+		/// <param name="Count">Number of results to return</param>
+		/// <param name="ConsistentRead">If a read to the primary database is required, for read consistency. Usually not required.</param>
+		/// <returns>List of jobs</returns>
+		[HttpGet]
+		[Route("/api/v1/jobs/streams/{StreamId}")]
+		[ProducesResponseType(typeof(List<GetJobResponse>), 200)]
+		public async Task<ActionResult<List<object>>> FindJobsByStreamWithTemplatesAsync(
+			string StreamId,
+			[FromQuery(Name = "template")] string[] Templates,
+			[FromQuery] string? PreflightStartedByUserId = null,
+			[FromQuery] PropertyFilter? Filter = null,
+			[FromQuery] int Index = 0,
+			[FromQuery] int Count = 100,
+			[FromQuery] bool ConsistentRead = true)
+		{
+			StreamId StreamIdValue = new StreamId(StreamId);
+			TemplateRefId[] TemplateRefIds = Templates.Select(x => new TemplateRefId(x)).ToArray();
+			UserId? PreflightStartedByUserIdValue = PreflightStartedByUserId != null ? new UserId(PreflightStartedByUserId) : null;
+			Count = Math.Min(1000, Count);
+
+			List<IJob> Jobs = await JobService.FindJobsByStreamWithTemplatesAsync(StreamIdValue, TemplateRefIds, PreflightStartedByUserIdValue, Index, Count, ConsistentRead);
+			return await CreateAuthorizedJobResponses(Jobs, Filter);
+		}
+
+		private async Task<List<object>> CreateAuthorizedJobResponses(List<IJob> Jobs, PropertyFilter? Filter = null)
+		{
+			StreamPermissionsCache PermissionsCache = new ();
+			List<object> Responses = new ();
+			foreach (IJob Job in Jobs)
+			{
+				using IScope JobScope = GlobalTracer.Instance.BuildSpan("JobIteration").StartActive();
+				JobScope.Span.SetTag("jobId", Job.Id.ToString());
+				
+				bool ViewJobAuthorized;
+				using (IScope _ = GlobalTracer.Instance.BuildSpan("AuthorizeViewJob").StartActive())
+				{
+					ViewJobAuthorized = await JobService.AuthorizeAsync(Job, AclAction.ViewJob, User, PermissionsCache);
+				}
+				
+				if (ViewJobAuthorized)
+				{
+					IGraph Graph;
+					using (IScope _ = GlobalTracer.Instance.BuildSpan("GetGraph").StartActive())
+					{
+						Graph = await JobService.GetGraphAsync(Job);
+					}
+
+					bool bIncludeAcl;
+					using (IScope _ = GlobalTracer.Instance.BuildSpan("AuthorizeViewPermissions").StartActive())
+					{
+						bIncludeAcl = await JobService.AuthorizeAsync(Job, AclAction.ViewPermissions, User, PermissionsCache);
+					}
+
+					using (IScope _ = GlobalTracer.Instance.BuildSpan("CreateResponse").StartActive())
+					{
+						Responses.Add(await CreateJobResponseAsync(Job, Graph, bIncludeAcl, Filter));
+					}
+				}
+			}
+
+			
+			return Responses;
+		}
 
 		/// <summary>
 		/// Adds an array of nodes to be executed for a job
diff --git a/Engine/Source/Programs/Horde/Horde.Build/Jobs/IJobCollection.cs b/Engine/Source/Programs/Horde/Horde.Build/Jobs/IJobCollection.cs
--- a/Engine/Source/Programs/Horde/Horde.Build/Jobs/IJobCollection.cs
+++ b/Engine/Source/Programs/Horde/Horde.Build/Jobs/IJobCollection.cs
@@ -105,8 +105,20 @@ namespace HordeServer.Collections
 		/// <param name="Index">Index of the first result to return</param>
 		/// <param name="Count">Number of results to return</param>
 		/// <param name="ConsistentRead">If the database read should be made to the replica server</param>
+		/// <param name="IndexHint">Name of index to be specified as a hint to the database query planner</param>
 		/// <returns>List of jobs matching the given criteria</returns>
-		Task<List<IJob>> FindAsync(JobId[]? JobIds = null, StreamId? StreamId = null, string? Name = null, TemplateRefId[]? Templates = null, int? MinChange = null, int? MaxChange = null, int? PreflightChange = null, UserId? PreflightStartedByUser = null, UserId? StartedByUser = null, DateTimeOffset? MinCreateTime = null, DateTimeOffset? MaxCreateTime = null, DateTimeOffset? ModifiedBefore = null, DateTimeOffset? ModifiedAfter = null, int? Index = null, int? Count = null, bool ConsistentRead = true);
+		Task<List<IJob>> FindAsync(JobId[]? JobIds = null, StreamId? StreamId = null, string? Name = null, TemplateRefId[]? Templates = null, int? MinChange = null, int? MaxChange = null, int? PreflightChange = null, UserId? PreflightStartedByUser = null, UserId? StartedByUser = null, DateTimeOffset? MinCreateTime = null, DateTimeOffset? MaxCreateTime = null, DateTimeOffset? ModifiedBefore = null, DateTimeOffset? ModifiedAfter = null, int? Index = null, int? Count = null, bool ConsistentRead = true, string? IndexHint = null);
+
+		/// <summary>
+		/// Searches for jobs in a specified stream and templates
+		/// </summary>
+		/// <param name="StreamId">The stream containing the job</param>
+		/// <param name="Templates">Templates to look for</param>
+		/// <param name="PreflightStartedByUser">User for which to include preflight jobs</param>
+		/// <param name="Index">Index of the first result to return</param>
+		/// <param name="Count">Number of results to return</param>
+		/// <param name="ConsistentRead">If the database read should be made to the replica server</param>
+		Task<List<IJob>> FindLatestByStreamWithTemplatesAsync(StreamId StreamId, TemplateRefId[] Templates, UserId? PreflightStartedByUser = null, int? Index = null, int? Count = null, bool ConsistentRead = false);
 
 		/// <summary>
 		/// Updates a new job
diff --git a/Engine/Source/Programs/Horde/Horde.Build/Jobs/Impl/JobCollection.cs b/Engine/Source/Programs/Horde/Horde.Build/Jobs/Impl/JobCollection.cs
--- a/Engine/Source/Programs/Horde/Horde.Build/Jobs/Impl/JobCollection.cs
+++ b/Engine/Source/Programs/Horde/Horde.Build/Jobs/Impl/JobCollection.cs
@@ -312,6 +312,51 @@ namespace HordeServer.Collections.Impl
 			}
 		}
 
+		class DatabaseIndexes
+		{
+			public readonly string StreamId;
+			public readonly string Change;
+			public readonly string PreflightChange;
+			public readonly string CreateTimeUtc;
+			public readonly string UpdateTimeUtc;
+			public readonly string Name;
+			public readonly string StartedByUserId;
+			public readonly string TemplateId;
+			public readonly string SchedulePriority;
+			private bool DatabaseReadOnlyMode;
+
+			public DatabaseIndexes(IMongoCollection<JobDocument> Jobs, bool DatabaseReadOnlyMode)
+			{
+				this.DatabaseReadOnlyMode = DatabaseReadOnlyMode;
+				
+				StreamId = CreateOrGetIndex(Jobs, "StreamId_1", Builders<JobDocument>.IndexKeys.Ascending(x => x.StreamId));
+				Change = CreateOrGetIndex(Jobs, "Change_1", Builders<JobDocument>.IndexKeys.Ascending(x => x.Change));
+				PreflightChange = CreateOrGetIndex(Jobs, "PreflightChange_1", Builders<JobDocument>.IndexKeys.Ascending(x => x.PreflightChange));
+				CreateTimeUtc = CreateOrGetIndex(Jobs, "CreateTimeUtc_-1", Builders<JobDocument>.IndexKeys.Descending(x => x.CreateTimeUtc));
+				UpdateTimeUtc = CreateOrGetIndex(Jobs, "UpdateTimeUtc_-1", Builders<JobDocument>.IndexKeys.Descending(x => x.UpdateTimeUtc));
+				Name = CreateOrGetIndex(Jobs, "Name_1", Builders<JobDocument>.IndexKeys.Ascending(x => x.Name));
+				StartedByUserId = CreateOrGetIndex(Jobs, "StartedByUserId_1", Builders<JobDocument>.IndexKeys.Ascending(x => x.StartedByUserId));
+				TemplateId = CreateOrGetIndex(Jobs, "TemplateId_1", Builders<JobDocument>.IndexKeys.Ascending(x => x.TemplateId));
+				SchedulePriority = CreateOrGetIndex(Jobs, "SchedulePriority_-1", Builders<JobDocument>.IndexKeys.Descending(x => x.SchedulePriority));
+			}
+
+			private string CreateOrGetIndex<T>(IMongoCollection<T> Collection, string Name, IndexKeysDefinition<T> IndexDef)
+			{
+				if (DatabaseReadOnlyMode)
+				{
+					if (Collection.Indexes.List().ToList().Find(x => x.TryGetString("name", out string? n) && n == Name) == null)
+					{
+						throw new ArgumentException("Index with name {Name} not found (running in database read-only mode)");
+					}
+					return Name;
+				}
+
+				CreateIndexOptions IndexOptions = new() { Name = Name, Background = true };
+				Collection.Indexes.CreateOne(new CreateIndexModel<T>(IndexDef, IndexOptions));
+				return Name;
+			}
+		}
+
 		/// <summary>
 		/// Projection of a job definition to just include permissions info
 		/// </summary>
@@ -344,6 +389,11 @@ namespace HordeServer.Collections.Impl
 		/// </summary>
 		ILogger<JobCollection> Logger;
 
+		/// <summary>
+		/// Creation and lookup of indexes
+		/// </summary>
+		DatabaseIndexes Indexes;
+
 		/// <summary>
 		/// Constructor
 		/// </summary>
@@ -356,19 +406,7 @@ namespace HordeServer.Collections.Impl
 			this.Logger = Logger;
 
 			Jobs = DatabaseService.GetCollection<JobDocument>("Jobs");
-
-			if (!DatabaseService.ReadOnlyMode)
-			{
-				Jobs.Indexes.CreateOne(new CreateIndexModel<JobDocument>(Builders<JobDocument>.IndexKeys.Ascending(x => x.StreamId)));
-				Jobs.Indexes.CreateOne(new CreateIndexModel<JobDocument>(Builders<JobDocument>.IndexKeys.Ascending(x => x.Change)));
-				Jobs.Indexes.CreateOne(new CreateIndexModel<JobDocument>(Builders<JobDocument>.IndexKeys.Ascending(x => x.PreflightChange)));
-				Jobs.Indexes.CreateOne(new CreateIndexModel<JobDocument>(Builders<JobDocument>.IndexKeys.Descending(x => x.CreateTimeUtc)));
-				Jobs.Indexes.CreateOne(new CreateIndexModel<JobDocument>(Builders<JobDocument>.IndexKeys.Descending(x => x.UpdateTimeUtc)));
-				Jobs.Indexes.CreateOne(new CreateIndexModel<JobDocument>(Builders<JobDocument>.IndexKeys.Ascending(x => x.Name)));
-				Jobs.Indexes.CreateOne(new CreateIndexModel<JobDocument>(Builders<JobDocument>.IndexKeys.Ascending(x => x.StartedByUserId)));
-				Jobs.Indexes.CreateOne(new CreateIndexModel<JobDocument>(Builders<JobDocument>.IndexKeys.Ascending(x => x.TemplateId)));
-				Jobs.Indexes.CreateOne(new CreateIndexModel<JobDocument>(Builders<JobDocument>.IndexKeys.Descending(x => x.SchedulePriority)));
-			}
+			Indexes = new DatabaseIndexes(Jobs, DatabaseService.ReadOnlyMode);
 		}
 
 		async Task PostLoadAsync(JobDocument Job)
@@ -496,7 +534,7 @@ namespace HordeServer.Collections.Impl
 		}
 
 		/// <inheritdoc/>
-		public async Task<List<IJob>> FindAsync(JobId[]? JobIds, StreamId? StreamId, string? Name, TemplateRefId[]? Templates, int? MinChange, int? MaxChange, int? PreflightChange, UserId? PreflightStartedByUser, UserId? StartedByUser, DateTimeOffset? MinCreateTime, DateTimeOffset? MaxCreateTime, DateTimeOffset? ModifiedBefore, DateTimeOffset? ModifiedAfter, int? Index, int? Count, bool ConsistentRead)
+		public async Task<List<IJob>> FindAsync(JobId[]? JobIds, StreamId? StreamId, string? Name, TemplateRefId[]? Templates, int? MinChange, int? MaxChange, int? PreflightChange, UserId? PreflightStartedByUser, UserId? StartedByUser, DateTimeOffset? MinCreateTime, DateTimeOffset? MaxCreateTime, DateTimeOffset? ModifiedBefore, DateTimeOffset? ModifiedAfter, int? Index, int? Count, bool ConsistentRead, string? IndexHint)
 		{
 			FilterDefinitionBuilder<JobDocument> FilterBuilder = Builders<JobDocument>.Filter;
 
@@ -558,7 +596,14 @@ namespace HordeServer.Collections.Impl
 			using (IScope Scope = GlobalTracer.Instance.BuildSpan("Jobs.Find").StartActive())
 			{
 				IMongoCollection<JobDocument> Collection = ConsistentRead ? Jobs : Jobs.WithReadPreference(ReadPreference.SecondaryPreferred);
-				IFindFluent<JobDocument, JobDocument> Query = Collection.Find<JobDocument>(Filter).SortByDescending(x => x.CreateTimeUtc!);
+
+				FindOptions? FindOptions = null;
+				if (IndexHint != null)
+				{
+					FindOptions = new FindOptions { Hint = new BsonString(IndexHint) };
+				}
+				
+				IFindFluent<JobDocument, JobDocument> Query = Collection.Find<JobDocument>(Filter, FindOptions).SortByDescending(x => x.CreateTimeUtc!);
 				
 				if (Index != null)
 				{
@@ -577,6 +622,16 @@ namespace HordeServer.Collections.Impl
 			return Results.ConvertAll<JobDocument, IJob>(x => x);
 		}
 
+		/// <inheritdoc/>
+		public async Task<List<IJob>> FindLatestByStreamWithTemplatesAsync(StreamId StreamId, TemplateRefId[] Templates, UserId? PreflightStartedByUser, int? Index, int? Count, bool ConsistentRead)
+		{
+			// This find call uses an index hint. Modifying the parameter passed to FindAsync can affect execution time a lot as the query planner is forced to use the specified index.
+			// Casting to interface to benefit from default parameter values
+			return await (this as IJobCollection).FindAsync(
+				StreamId: StreamId, Templates: Templates, PreflightStartedByUser: PreflightStartedByUser,
+				Index: Index, Count: Count, IndexHint: Indexes.CreateTimeUtc, ConsistentRead: ConsistentRead);
+		}
+
 		/// <inheritdoc/>
 		public async Task<IJob?> TryUpdateJobAsync(IJob InJob, IGraph Graph, string? Name, Priority? Priority, bool? AutoSubmit, int? AutoSubmitChange, string? AutoSubmitMessage, UserId? AbortedByUserId, ObjectId? NotificationTriggerId, List<Report>? Reports, List<string>? Arguments, KeyValuePair<int, ObjectId>? LabelIdxToTriggerId, KeyValuePair<TemplateRefId, JobId>? JobTrigger)
 		{
diff --git a/Engine/Source/Programs/Horde/Horde.Build/Jobs/Impl/JobService.cs b/Engine/Source/Programs/Horde/Horde.Build/Jobs/Impl/JobService.cs
--- a/Engine/Source/Programs/Horde/Horde.Build/Jobs/Impl/JobService.cs
+++ b/Engine/Source/Programs/Horde/Horde.Build/Jobs/Impl/JobService.cs
@@ -605,6 +605,28 @@ namespace HordeServer.Services
 				return Results;
 			}
 		}
+		
+		/// <summary>
+		/// Searches for jobs matching a stream with given templates
+		/// </summary>
+		/// <param name="StreamId">The stream containing the job</param>
+		/// <param name="Templates">Templates to look for</param>
+		/// <param name="PreflightStartedByUser">User for which to include preflight jobs</param>		
+		/// <param name="Index">Index of the first result to return</param>
+		/// <param name="Count">Number of results to return</param>
+		/// <param name="ConsistentRead">If the database read should be made to the replica server</param>
+		/// <returns>List of jobs matching the given criteria</returns>
+		public async Task<List<IJob>> FindJobsByStreamWithTemplatesAsync(StreamId StreamId, TemplateRefId[] Templates, UserId? PreflightStartedByUser = null, int? Index = null, int? Count = null, bool ConsistentRead = true)
+		{
+			using IScope Scope = GlobalTracer.Instance.BuildSpan("JobService.FindJobsByStreamWithTemplatesAsync").StartActive();
+			Scope.Span.SetTag("StreamId", StreamId);
+			Scope.Span.SetTag("Templates", Templates);
+			Scope.Span.SetTag("PreflightStartedByUser", PreflightStartedByUser);
+			Scope.Span.SetTag("Index", Index);
+			Scope.Span.SetTag("Count", Count);
+			
+			return await Jobs.FindLatestByStreamWithTemplatesAsync(StreamId, Templates, PreflightStartedByUser, Index, Count, ConsistentRead);
+		}
 
 		/// <summary>
 		/// Attempts to update the node groups to be executed for a job. Fails if another write happens in the meantime.
