Refactored blob storage in Horde.Storage to allow for central handling of all blob store requests.

This lets us add verification of the blob identifier before submitting content to the blob store.
I also added a blob index, a database tracking of which blobs exists, and a optional toggle to use this instead of checking S3 during existence checks (which we think will be faster). Could also support batched requests but that is not yet implemented.
Also promoted the current site setting into the common (JupiterSettings) settings, as we now also use this for the blob index, which we intend to use for in line replication of missing blobs in the future.

Removed the hirearchal blob stores and moved this into the blob service (the new central place to handle all blob store requests). Change this to process blob stores in parallel which will be a small speedup as we start the s3 requests earlier.
Changed GetOldObjects in the IBlobStore to just fetch AllObjects (as we used it for both of those purposes) and now we include the last modified time with the object for the filtering cases.

Added IBufferedPayload for reusable control of how we buffer the payloads we need to handle. This means we will now handle large (>2GB) payloads better in more places. Also added a setting to control when we start buffering to disk (defaults to 2GB payloads, e.g. payloads to large to fit into memory). Could be used to reduce the amount of memory required by [REDACTED_PERSON] at the cost of requiring more temp drive space, and likely some performance as we then end up writing to disk. This buffering is required as we need to hash all objects before we store them, requiring us to have the entire payload before it can be sent anywhere.

diff --git a/Engine/Source/Programs/Horde/HordeStorage/Callisto/CallistoStartup.cs b/Engine/Source/Programs/Horde/HordeStorage/Callisto/CallistoStartup.cs
--- a/Engine/Source/Programs/Horde/HordeStorage/Callisto/CallistoStartup.cs
+++ b/Engine/Source/Programs/Horde/HordeStorage/Callisto/CallistoStartup.cs
@@ -92,9 +92,6 @@ namespace Callisto
 
         [Required] public TransactionLogImplementations TransactionLogImplementation { get; set; }
 
-        // Name of the current site, has to be globally unique across all deployments, used to avoid replication loops
-        [Required] public string CurrentSite { get; set; } = "";
-
         [Required]
         public string TransactionLogRoot
         {
diff --git a/Engine/Source/Programs/Horde/HordeStorage/Callisto/Controllers/TransactionLogController.cs b/Engine/Source/Programs/Horde/HordeStorage/Callisto/Controllers/TransactionLogController.cs
--- a/Engine/Source/Programs/Horde/HordeStorage/Callisto/Controllers/TransactionLogController.cs
+++ b/Engine/Source/Programs/Horde/HordeStorage/Callisto/Controllers/TransactionLogController.cs
@@ -24,16 +24,16 @@ namespace Callisto.Controllers
         private readonly IDiagnosticContext _diagnosticContext;
         private readonly IAuthorizationService _authorizationService;
 
-        private readonly CallistoSettings _settings;
+        private readonly IOptionsMonitor<JupiterSettings> _settings;
         private readonly ILogger _logger = Log.ForContext<TransactionLogController>();
 
-        public TransactionLogController(ITransactionLogs transactionLogs, IOptionsMonitor<CallistoSettings> options, IDiagnosticContext diagnosticContext, IAuthorizationService authorizationService)
+        public TransactionLogController(ITransactionLogs transactionLogs, IOptionsMonitor<JupiterSettings> options, IDiagnosticContext diagnosticContext, IAuthorizationService authorizationService)
         {
             _transactionLogs = transactionLogs;
             _diagnosticContext = diagnosticContext;
             _authorizationService = authorizationService;
 
-            _settings = options.CurrentValue;
+            _settings = options;
         }
 
 
@@ -162,14 +162,14 @@ namespace Callisto.Controllers
                 return Forbid();
             }
 
-            if (string.IsNullOrEmpty(_settings.CurrentSite))
+            if (string.IsNullOrEmpty(_settings.CurrentValue.CurrentSite))
             {
                 throw new Exception(
                     "Current site setting is missing, this has to be part of the deployment and has to be globally unique");
             }
 
             // add the current site to the list of sites that have seen this event
-            request.Locations.Add(_settings.CurrentSite);
+            request.Locations.Add(_settings.CurrentValue.CurrentSite);
             long id;
             switch (request)
             {
diff --git a/Engine/Source/Programs/Horde/HordeStorage/Callisto/appsettings.Development.json b/Engine/Source/Programs/Horde/HordeStorage/Callisto/appsettings.Development.json
--- a/Engine/Source/Programs/Horde/HordeStorage/Callisto/appsettings.Development.json
+++ b/Engine/Source/Programs/Horde/HordeStorage/Callisto/appsettings.Development.json
@@ -18,8 +18,11 @@
     "Method": "Disabled"
   },
 
+  "Jupiter": {
+    "CurrentSite": "Dev"
+  },
+
   "Callisto": {
-    "CurrentSite": "Dev",
     "TransactionLogRoot": "%Temp%/Callisto",
     "TransactionLogImplementation": "File"
   }
diff --git a/Engine/Source/Programs/Horde/HordeStorage/Callisto/appsettings.Testing.json b/Engine/Source/Programs/Horde/HordeStorage/Callisto/appsettings.Testing.json
--- a/Engine/Source/Programs/Horde/HordeStorage/Callisto/appsettings.Testing.json
+++ b/Engine/Source/Programs/Horde/HordeStorage/Callisto/appsettings.Testing.json
@@ -24,11 +24,14 @@
     }
   },
 
+  "Jupiter": {
+    "CurrentSite": "Test"
+  },
+
   "Callisto": {
-    "CurrentSite": "Test",
     "TransactionLogRoot": "./callisto-logs",
     "transactionLogImplementation": "file",
-    "VerifySerialization":  "true"
+    "VerifySerialization": "true"
   },
 
   "AllowedHosts": "*"
diff --git a/Engine/Source/Programs/Horde/HordeStorage/Horde.Storage/Controllers/CompressedBlobController.cs b/Engine/Source/Programs/Horde/HordeStorage/Horde.Storage/Controllers/CompressedBlobController.cs
--- a/Engine/Source/Programs/Horde/HordeStorage/Horde.Storage/Controllers/CompressedBlobController.cs
+++ b/Engine/Source/Programs/Horde/HordeStorage/Horde.Storage/Controllers/CompressedBlobController.cs
@@ -11,9 +11,10 @@ using System.Threading.Tasks;
 using Datadog.Trace;
 using Horde.Storage.Implementation;
 using Jupiter;
+using Jupiter.Common.Implementation;
 using Jupiter.Implementation;
+using Jupiter.Utils;
 using Microsoft.AspNetCore.Authorization;
-using Microsoft.AspNetCore.Http;
 using Microsoft.AspNetCore.Mvc;
 using Serilog;
 
@@ -23,21 +24,23 @@ namespace Horde.Storage.Controllers
     [Route("api/v1/compressed-blobs")]
     public class CompressedBlobController : ControllerBase
     {
-        private readonly IBlobStore _storage;
+        private readonly IBlobService _storage;
         private readonly IContentIdStore _contentIdStore;
         private readonly IDiagnosticContext _diagnosticContext;
         private readonly IAuthorizationService _authorizationService;
         private readonly CompressedBufferUtils _compressedBufferUtils;
+        private readonly BufferedPayloadFactory _bufferedPayloadFactory;
 
         private readonly ILogger _logger = Log.ForContext<CompressedBlobController>();
 
-        public CompressedBlobController(IBlobStore storage, IContentIdStore contentIdStore, IDiagnosticContext diagnosticContext, IAuthorizationService authorizationService, CompressedBufferUtils compressedBufferUtils)
+        public CompressedBlobController(IBlobService storage, IContentIdStore contentIdStore, IDiagnosticContext diagnosticContext, IAuthorizationService authorizationService, CompressedBufferUtils compressedBufferUtils, BufferedPayloadFactory bufferedPayloadFactory)
         {
             _storage = storage;
             _contentIdStore = contentIdStore;
             _diagnosticContext = diagnosticContext;
             _authorizationService = authorizationService;
             _compressedBufferUtils = compressedBufferUtils;
+            _bufferedPayloadFactory = bufferedPayloadFactory;
         }
 
 
@@ -197,56 +200,54 @@ namespace Horde.Storage.Controllers
                 return Forbid();
             }
 
-            byte[] blob;
+            _diagnosticContext.Set("Content-Length", Request.ContentLength ?? -1);
+
+            using IBufferedPayload payload = await _bufferedPayloadFactory.CreateFromRequest(Request);
+
             try
             {
-                blob = await RequestUtil.ReadRawBody(Request);
+                BlobIdentifier identifier = await PutImpl(ns, id, payload);
+
+                return Ok(new
+                {
+                    Identifier = identifier.ToString()
+                });
             }
-            catch (BadHttpRequestException e)
+            catch (HashMismatchException e)
             {
-                const string msg = "Partial content transfer when reading request body.";
-                _logger.Warning(e, msg);
-                return BadRequest(msg);
+                return BadRequest(new ProblemDetails
+                {
+                    Title = $"Incorrect hash, got hash \"{e.SuppliedHash}\" but hash of content was determined to be \"{e.ContentHash}\""
+                });
             }
-            _diagnosticContext.Set("Content-Length", Request.ContentLength ?? -1);
-
-            BlobIdentifier identifier = await PutImpl(ns, id, blob);
-            return Ok(new
-            {
-                Identifier = identifier.ToString()
-            });
         }
 
-        private async Task<BlobIdentifier> PutImpl(NamespaceId ns, BlobIdentifier? id, byte[] content)
+        private async Task<BlobIdentifier> PutImpl(NamespaceId ns, BlobIdentifier id, IBufferedPayload payload)
         {
-            BlobIdentifier identifier;
-            {
-                using Scope _ = Tracer.Instance.StartActive("web.hash");
-                identifier = BlobIdentifier.FromBlob(content);
-            }
             // decompress the content and generate a identifier from it to verify the identifier we got
-            byte[] decompressedContent = _compressedBufferUtils.DecompressContent(content);
+            await using Stream decompressStream = payload.GetStream();
+            // TODO: we should add a overload for decompress content that can work on streams, otherwise we are still limited to 2GB compressed blobs
+            byte[] decompressedContent = _compressedBufferUtils.DecompressContent(await decompressStream.ToByteArray());
 
-            BlobIdentifier identifierDecompressedPayload;
-            {
-                using Scope _ = Tracer.Instance.StartActive("web.hash");
-                identifierDecompressedPayload = BlobIdentifier.FromBlob(decompressedContent);
-            }
+            MemoryStream decompressedStream = new MemoryStream(decompressedContent);
+            BlobIdentifier identifierDecompressedPayload = await _storage.VerifyContentMatchesHash(decompressedStream, id);
 
-            if (id != null && !id.Equals(identifierDecompressedPayload))
+            BlobIdentifier identifierCompressedPayload;
             {
-                _logger.Debug("ID {@Id} was not the same as identifier {@Identifier} {Content}", id, identifierDecompressedPayload,
-                    content);
-
-                throw new ArgumentException("ID was not a hash of the content uploaded.", paramName: nameof(id));
+                using Scope _ = Tracer.Instance.StartActive("web.hash");
+                await using Stream hashStream = payload.GetStream();
+                identifierCompressedPayload = await BlobIdentifier.FromStream(hashStream);
             }
 
             // commit the mapping from the decompressed hash to the compressed hash, we run this in parallel with the blob store submit
             // TODO: let users specify weight of the blob compared to previously submitted content ids
-            Task contentIdStoreTask = _contentIdStore.Put(ns, identifierDecompressedPayload, identifier, content.Length);
+            int contentIdWeight = (int)payload.Length;
+            Task contentIdStoreTask = _contentIdStore.Put(ns, identifierDecompressedPayload, identifierCompressedPayload, contentIdWeight);
 
-            // we still commit the compressed buffer to the object store
-            await _storage.PutObject(ns, content, identifier);
+            // we still commit the compressed buffer to the object store using the hash of the compressed content
+            {
+                await _storage.PutObjectKnownHash(ns, payload, identifierCompressedPayload);
+            }
 
             await contentIdStoreTask;
 
diff --git a/Engine/Source/Programs/Horde/HordeStorage/Horde.Storage/Controllers/DDCRefController.cs b/Engine/Source/Programs/Horde/HordeStorage/Horde.Storage/Controllers/DDCRefController.cs
--- a/Engine/Source/Programs/Horde/HordeStorage/Horde.Storage/Controllers/DDCRefController.cs
+++ b/Engine/Source/Programs/Horde/HordeStorage/Horde.Storage/Controllers/DDCRefController.cs
@@ -474,7 +474,7 @@ namespace Horde.Storage.Controllers
             {
                 return BadRequest(new ProblemDetails
                 {
-                    Title = $"Incorrect hash, got hash \"{e.NewHash}\" but hash of content was determined to be \"{e.ExpectedHash}\""
+                    Title = $"Incorrect hash, got hash \"{e.SuppliedHash}\" but hash of content was determined to be \"{e.ContentHash}\""
                 });
             }
             catch (MissingBlobsException e)
@@ -790,7 +790,7 @@ namespace Horde.Storage.Controllers
                     ContentHash blobHash = ContentHash.FromBlob(op.Content);
                     if (!blobHash.Equals(op.ContentHash))
                     {
-                        throw new HashMismatchException(blobHash, op.ContentHash);
+                        throw new HashMismatchException(op.ContentHash, blobHash);
                     }
                     return _ddcRefService.Put(op.Namespace.Value, op.Bucket.Value, op.Id.Value, blobHash, op.Content).ContinueWith(t => (object?)t.Result);
                 }
@@ -910,13 +910,13 @@ namespace Horde.Storage.Controllers
 
     public class HashMismatchException : Exception
     {
-        public ContentHash NewHash { get; }
-        public ContentHash ExpectedHash { get; }
+        public ContentHash SuppliedHash { get; }
+        public ContentHash ContentHash { get; }
 
-        public HashMismatchException(ContentHash newHash, ContentHash expectedHash)
+        public HashMismatchException(ContentHash suppliedHash, ContentHash contentHash) : base($"ID was not a hash of the content uploaded. Supplied hash was: {suppliedHash} but hash of content was {contentHash}")
         {
-            NewHash = newHash;
-            ExpectedHash = expectedHash;
+            SuppliedHash = suppliedHash;
+            ContentHash = contentHash;
         }
     }
 
diff --git a/Engine/Source/Programs/Horde/HordeStorage/Horde.Storage/Controllers/ObjectController.cs b/Engine/Source/Programs/Horde/HordeStorage/Horde.Storage/Controllers/ObjectController.cs
--- a/Engine/Source/Programs/Horde/HordeStorage/Horde.Storage/Controllers/ObjectController.cs
+++ b/Engine/Source/Programs/Horde/HordeStorage/Horde.Storage/Controllers/ObjectController.cs
@@ -4,6 +4,7 @@ using System;
 using System.Collections.Concurrent;
 using System.Collections.Generic;
 using System.ComponentModel.DataAnnotations;
+using System.IO;
 using System.Linq;
 using System.Net.Mime;
 using System.Threading.Tasks;
@@ -11,6 +12,7 @@ using async_enumerable_dotnet;
 using Datadog.Trace;
 using Horde.Storage.Implementation;
 using Jupiter;
+using Jupiter.Common.Implementation;
 using Jupiter.Implementation;
 using Jupiter.Utils;
 using Microsoft.AspNetCore.Authorization;
@@ -25,19 +27,21 @@ namespace Horde.Storage.Controllers
     [Produces(CustomMediaTypeNames.UnrealCompactBinary, MediaTypeNames.Application.Json)]
     public class ObjectController : ControllerBase
     {
-        private readonly IBlobStore _storage;
+        private readonly IBlobService _storage;
         private readonly IDiagnosticContext _diagnosticContext;
         private readonly IAuthorizationService _authorizationService;
         private readonly IReferenceResolver _referenceResolver;
+        private readonly BufferedPayloadFactory _bufferedPayloadFactory;
 
         private readonly ILogger _logger = Log.ForContext<ObjectController>();
 
-        public ObjectController(IBlobStore storage, IDiagnosticContext diagnosticContext, IAuthorizationService authorizationService, IReferenceResolver referenceResolver)
+        public ObjectController(IBlobService storage, IDiagnosticContext diagnosticContext, IAuthorizationService authorizationService, IReferenceResolver referenceResolver, BufferedPayloadFactory bufferedPayloadFactory)
         {
             _storage = storage;
             _diagnosticContext = diagnosticContext;
             _authorizationService = authorizationService;
             _referenceResolver = referenceResolver;
+            _bufferedPayloadFactory = bufferedPayloadFactory;
         }
 
 
@@ -131,36 +135,10 @@ namespace Horde.Storage.Controllers
                 return Forbid();
             }
 
-            byte[] blob;
-            try
-            {
-                blob = await RequestUtil.ReadRawBody(Request);
-            }
-            catch (BadHttpRequestException e)
-            {
-                const string msg = "Partial content transfer when reading request body.";
-                _logger.Warning(e, msg);
-                return BadRequest(msg);
-            }
-
             _diagnosticContext.Set("Content-Length", Request.ContentLength ?? -1);
+            using IBufferedPayload payload = await _bufferedPayloadFactory.CreateFromRequest(Request);
 
-            BlobIdentifier blobHash;
-            {
-                using Scope _ = Tracer.Instance.StartActive("web.hash");
-                blobHash = BlobIdentifier.FromBlob(blob);
-            }
-
-            if (!id.Equals(blobHash))
-            {
-                _logger.Debug("ID {@Id} was not the same as identifier {@Identifier} {Content}", id, blobHash,
-                    blob);
-
-                throw new ArgumentException("ID was not a hash of the content uploaded.", paramName: nameof(id));
-            }
-
-            await _storage.PutObject(ns, blob, blobHash);
-            BlobIdentifier identifier = blobHash;
+            BlobIdentifier identifier = await _storage.PutObject(ns, payload, id);
             return Ok(new
             {
                 Identifier = identifier
diff --git a/Engine/Source/Programs/Horde/HordeStorage/Horde.Storage/Controllers/ReferencesController.cs b/Engine/Source/Programs/Horde/HordeStorage/Horde.Storage/Controllers/ReferencesController.cs
--- a/Engine/Source/Programs/Horde/HordeStorage/Horde.Storage/Controllers/ReferencesController.cs
+++ b/Engine/Source/Programs/Horde/HordeStorage/Horde.Storage/Controllers/ReferencesController.cs
@@ -4,6 +4,7 @@ using System;
 using System.Collections.Concurrent;
 using System.Collections.Generic;
 using System.ComponentModel.DataAnnotations;
+using System.IO;
 using System.Linq;
 using System.Net.Mime;
 using System.Text;
@@ -12,6 +13,7 @@ using Dasync.Collections;
 using Datadog.Trace;
 using Horde.Storage.Implementation;
 using Jupiter;
+using Jupiter.Common.Implementation;
 using Jupiter.Implementation;
 using Jupiter.Utils;
 using Microsoft.AspNetCore.Authorization;
@@ -33,12 +35,13 @@ namespace Horde.Storage.Controllers
         private readonly IAuthorizationService _authorizationService;
         private readonly IOptionsMonitor<JupiterSettings> _jupiterSettings;
         private readonly FormatResolver _formatResolver;
+        private readonly BufferedPayloadFactory _bufferedPayloadFactory;
 
         private readonly ILogger _logger = Log.ForContext<ReferencesController>();
         private readonly IObjectService _objectService;
-        private readonly IBlobStore _blobStore;
+        private readonly IBlobService _blobStore;
 
-        public ReferencesController(IObjectService objectService, IBlobStore blobStore, IDiagnosticContext diagnosticContext, IAuthorizationService authorizationService, IOptionsMonitor<JupiterSettings> jupiterSettings, FormatResolver formatResolver)
+        public ReferencesController(IObjectService objectService, IBlobService blobStore, IDiagnosticContext diagnosticContext, IAuthorizationService authorizationService, IOptionsMonitor<JupiterSettings> jupiterSettings, FormatResolver formatResolver, BufferedPayloadFactory bufferedPayloadFactory)
         {
             _objectService = objectService;
             _blobStore = blobStore;
@@ -46,6 +49,7 @@ namespace Horde.Storage.Controllers
             _authorizationService = authorizationService;
             _jupiterSettings = jupiterSettings;
             _formatResolver = formatResolver;
+            _bufferedPayloadFactory = bufferedPayloadFactory;
         }
 
         /// <summary>
@@ -306,21 +310,11 @@ namespace Horde.Storage.Controllers
                 }
             }
 
-            byte[] blob;
-            try
-            {
-                blob = await RequestUtil.ReadRawBody(Request);
-            }
-            catch (BadHttpRequestException e)
-            {
-                const string msg = "Partial content transfer when reading request body.";
-                _logger.Warning(e, msg);
-                return BadRequest(msg);
-            }
             _diagnosticContext.Set("Content-Length", Request.ContentLength ?? -1);
+            
+            using IBufferedPayload payload = await _bufferedPayloadFactory.CreateFromRequest(Request);
 
             BlobIdentifier headerHash;
-            BlobIdentifier blobHeader = BlobIdentifier.FromBlob(blob);
             if (Request.Headers.ContainsKey(CommonHeaders.HashHeaderName))
             {
                 headerHash = new BlobIdentifier(Request.Headers[CommonHeaders.HashHeaderName]);
@@ -333,56 +327,63 @@ namespace Horde.Storage.Controllers
                 });
             }
 
-            if (!blobHeader.Equals(headerHash))
-            {
-                return BadRequest(new ProblemDetails
-                {
-                    Title = $"Incorrect hash, got hash \"{headerHash}\" but hash of content was determined to be \"{blobHeader}\""
-                });
-            }
-
             CompactBinaryObject payloadObject;
-            switch (Request.ContentType)
+            BlobIdentifier blobHeader = headerHash;
+            try
             {
-                case MediaTypeNames.Application.Json:
+                switch (Request.ContentType)
                 {
-                    // TODO: define a scheme for how a json object specifies references
+                    case MediaTypeNames.Application.Json:
+                    {
+                        // TODO: define a scheme for how a json object specifies references
 
-                    await _blobStore.PutObject(ns, blob, blobHeader);
+                        blobHeader = await _blobStore.PutObject(ns, payload, headerHash);
 
-                    // TODO: convert the json object into a compact binary instead
-                    CompactBinaryWriter compactBinaryWriter = new CompactBinaryWriter();
-                    compactBinaryWriter.BeginObject();
-                    compactBinaryWriter.AddBinaryAttachment(blobHeader);
-                    compactBinaryWriter.EndObject();
+                        // TODO: convert the json object into a compact binary instead
+                        CompactBinaryWriter compactBinaryWriter = new CompactBinaryWriter();
+                        compactBinaryWriter.BeginObject();
+                        compactBinaryWriter.AddBinaryAttachment(blobHeader);
+                        compactBinaryWriter.EndObject();
 
-                    blob = compactBinaryWriter.Save();
-                    payloadObject = CompactBinaryObject.Load(blob);
-                    blobHeader = BlobIdentifier.FromBlob(blob);
-                    break;
-                }
-                case CustomMediaTypeNames.UnrealCompactBinary:
-                {
-                    payloadObject = CompactBinaryObject.Load(blob);
-                    break;
-                }
-                case MediaTypeNames.Application.Octet:
-                {
-                    await _blobStore.PutObject(ns, blob, blobHeader);
+                        byte[] blob = compactBinaryWriter.Save();
+                        payloadObject = CompactBinaryObject.Load(blob);
+                        blobHeader = BlobIdentifier.FromBlob(blob);
+                        break;
+                    }
+                    case CustomMediaTypeNames.UnrealCompactBinary:
+                    {
+                        MemoryStream ms = new MemoryStream();
+                        await using Stream payloadStream = payload.GetStream();
+                        await payloadStream.CopyToAsync(ms);
+                        payloadObject = CompactBinaryObject.Load(ms.ToArray());
+                        break;
+                    }
+                    case MediaTypeNames.Application.Octet:
+                    {
+                        blobHeader = await _blobStore.PutObject(ns, payload, headerHash);
 
-                    CompactBinaryWriter compactBinaryWriter = new CompactBinaryWriter();
-                    compactBinaryWriter.BeginObject();
-                    compactBinaryWriter.AddBinaryAttachment(blobHeader);
-                    compactBinaryWriter.EndObject();
+                        CompactBinaryWriter compactBinaryWriter = new CompactBinaryWriter();
+                        compactBinaryWriter.BeginObject();
+                        compactBinaryWriter.AddBinaryAttachment(blobHeader);
+                        compactBinaryWriter.EndObject();
 
-                    blob = compactBinaryWriter.Save();
-                    payloadObject = CompactBinaryObject.Load(blob);
-                    blobHeader = BlobIdentifier.FromBlob(blob);
-                    break;
+                        byte[] blob = compactBinaryWriter.Save();
+                        payloadObject = CompactBinaryObject.Load(blob);
+                        blobHeader = BlobIdentifier.FromBlob(blob);
+                        break;
+                    }
+                    default:
+                        throw new Exception($"Unknown request type {Request.ContentType}, if submitting a blob please use {MediaTypeNames.Application.Octet}");
                 }
-                default:
-                    throw new Exception($"Unknown request type {Request.ContentType}, if submitting a blob please use {MediaTypeNames.Application.Octet}");
             }
+            catch (HashMismatchException e)
+            {
+                return BadRequest(new ProblemDetails
+                {
+                    Title = $"Incorrect hash, got hash \"{e.SuppliedHash}\" but hash of content was determined to be \"{e.ContentHash}\""
+                });
+            }
+
 
             PutObjectResult result = await _objectService.Put(ns, bucket, key, blobHeader, payloadObject);
             return Ok(new PutObjectResponse(result.MissingReferences));
diff --git a/Engine/Source/Programs/Horde/HordeStorage/Horde.Storage/Controllers/StatusController.cs b/Engine/Source/Programs/Horde/HordeStorage/Horde.Storage/Controllers/StatusController.cs
--- a/Engine/Source/Programs/Horde/HordeStorage/Horde.Storage/Controllers/StatusController.cs
+++ b/Engine/Source/Programs/Horde/HordeStorage/Horde.Storage/Controllers/StatusController.cs
@@ -10,6 +10,7 @@ using System.Threading;
 using System.Threading.Tasks;
 using Datadog.Trace.DuckTyping;
 using Horde.Storage.Implementation;
+using Jupiter;
 using Jupiter.Implementation;
 using Microsoft.AspNetCore.Authorization;
 using Microsoft.AspNetCore.Mvc;
@@ -23,12 +24,12 @@ namespace Horde.Storage.Controllers
     public class StatusController : Controller
     {
         private readonly VersionFile _versionFile;
-        private readonly IOptionsMonitor<ReplicationSettings> _replicationSettings;
+        private readonly IOptionsMonitor<JupiterSettings> _jupiterSettings;
 
-        public StatusController(VersionFile versionFile, IOptionsMonitor<ReplicationSettings> replicationSettings)
+        public StatusController(VersionFile versionFile, IOptionsMonitor<JupiterSettings> jupiterSettings)
         {
             _versionFile = versionFile;
-            _replicationSettings = replicationSettings;
+            _jupiterSettings = jupiterSettings;
         }
 
         /// <summary>
@@ -53,7 +54,7 @@ namespace Horde.Storage.Controllers
             AssemblyMetadataAttribute? p4ChangeAttribute = attrs.FirstOrDefault(attr => attr.Key == "PerforceChangelist");
             if (p4ChangeAttribute?.Value != null && !string.IsNullOrEmpty(p4ChangeAttribute.Value))
                 srcControlIdentifier = p4ChangeAttribute.Value;
-            return Ok(new StatusResponse(_versionFile.VersionString ?? "Unknown", srcControlIdentifier, GetCapabilities(), _replicationSettings.CurrentValue.CurrentSite));
+            return Ok(new StatusResponse(_versionFile.VersionString ?? "Unknown", srcControlIdentifier, GetCapabilities(), _jupiterSettings.CurrentValue.CurrentSite));
         }
 
         private string[] GetCapabilities()
diff --git a/Engine/Source/Programs/Horde/HordeStorage/Horde.Storage/Controllers/StorageController.cs b/Engine/Source/Programs/Horde/HordeStorage/Horde.Storage/Controllers/StorageController.cs
--- a/Engine/Source/Programs/Horde/HordeStorage/Horde.Storage/Controllers/StorageController.cs
+++ b/Engine/Source/Programs/Horde/HordeStorage/Horde.Storage/Controllers/StorageController.cs
@@ -11,6 +11,7 @@ using System.Threading.Tasks;
 using Datadog.Trace;
 using Horde.Storage.Implementation;
 using Jupiter;
+using Jupiter.Common.Implementation;
 using Jupiter.Implementation;
 using Microsoft.AspNetCore.Authorization;
 using Microsoft.AspNetCore.Http;
@@ -25,18 +26,20 @@ namespace Horde.Storage.Controllers
     [Route("api/v1/blobs", Order = 0)]
     public class StorageController : ControllerBase
     {
-        private readonly IBlobStore _storage;
+        private readonly IBlobService _storage;
         private readonly IDiagnosticContext _diagnosticContext;
         private readonly IAuthorizationService _authorizationService;
+        private readonly BufferedPayloadFactory _bufferedPayloadFactory;
         private readonly HordeStorageSettings _settings;
 
         private readonly ILogger _logger = Log.ForContext<StorageController>();
 
-        public StorageController(IBlobStore storage, IOptions<HordeStorageSettings> settings, IDiagnosticContext diagnosticContext, IAuthorizationService authorizationService)
+        public StorageController(IBlobService storage, IOptions<HordeStorageSettings> settings, IDiagnosticContext diagnosticContext, IAuthorizationService authorizationService, BufferedPayloadFactory bufferedPayloadFactory)
         {
             _storage = storage;
             _diagnosticContext = diagnosticContext;
             _authorizationService = authorizationService;
+            _bufferedPayloadFactory = bufferedPayloadFactory;
             _settings = settings.Value;
         }
 
@@ -162,53 +165,10 @@ namespace Horde.Storage.Controllers
                 return Forbid();
             }
 
-            BlobIdentifier identifier;
-            // blob is small enough to fit into memory so we just stream it into memory
-            if (Request.ContentLength is < Int32.MaxValue)
-            {
-                byte[] blob;
-                try
-                {
-                    blob = await RequestUtil.ReadRawBody(Request);
-                }
-                catch (BadHttpRequestException e)
-                {
-                    const string msg = "Partial content transfer when reading request body.";
-                    _logger.Warning(e, msg);
-                    return BadRequest(msg);
-                }
-
-                identifier = await PutImpl(ns, id, blob);
-            }
-            else
-            {
-                FileInfo? tempFile = null;
-
-                Stream s = Request.Body;
-                try
-                {
-                    // stream to a temporary file on disk if the stream is not seekable
-                    if (!Request.Body.CanSeek)
-                    {
-                        tempFile = new FileInfo(Path.GetTempFileName());
-
-                        {
-                            await using FileStream fs = tempFile.OpenWrite();
-                            await Request.Body.CopyToAsync(fs);
-                        }
-
-                        s = tempFile.OpenRead();
-                    }
-                    identifier = await PutImpl(ns, id, s);
-                }
-                finally
-                {
-                    s.Close();
-                    if (tempFile != null && tempFile.Exists)
-                        tempFile.Delete();
-                }
-            }
             _diagnosticContext.Set("Content-Length", Request.ContentLength ?? -1);
+            using IBufferedPayload payload = await _bufferedPayloadFactory.CreateFromRequest(Request);
+
+            BlobIdentifier identifier = await _storage.PutObject(ns, payload, id);
 
             return Ok(new
             {
@@ -216,48 +176,6 @@ namespace Horde.Storage.Controllers
             });
         }
 
-        private async Task<BlobIdentifier> PutImpl(NamespaceId ns, BlobIdentifier id, byte[] content)
-        {
-            BlobIdentifier identifier;
-            {
-                using Scope _ = Tracer.Instance.StartActive("web.hash");
-                identifier = BlobIdentifier.FromBlob(content);
-            }
-
-            if (!id.Equals(identifier))
-            {
-                _logger.Debug("ID {@Id} was not the same as identifier {@Identifier} {Content}", id, identifier,
-                    content);
-
-                throw new ArgumentException("ID was not a hash of the content uploaded.", paramName: nameof(id));
-            }
-
-            await _storage.PutObject(ns, content, identifier);
-            return identifier;
-        }
-
-        private async Task<BlobIdentifier> PutImpl(NamespaceId ns, BlobIdentifier id, Stream content)
-        {
-            BlobIdentifier identifier;
-            {
-                using Scope _ = Tracer.Instance.StartActive("web.hash");
-                identifier = await BlobIdentifier.FromStream(content);
-            }
-
-            if (!id.Equals(identifier))
-            {
-                _logger.Debug("ID {@Id} was not the same as identifier {@Identifier} {Content}", id, identifier,
-                    content);
-
-                throw new ArgumentException("ID was not a hash of the content uploaded.", paramName: nameof(id));
-            }
-
-            // seek back to the beginning
-            content.Position = 0;
-            await _storage.PutObject(ns, content, identifier);
-            return identifier;
-        }
-
         [HttpDelete("{ns}/{id}")]
         [Authorize("Storage.delete")]
         public async Task<IActionResult> Delete(
@@ -413,7 +331,8 @@ namespace Horde.Storage.Controllers
                             return BadRequest();
                         }
 
-                        tasks[index] = PutImpl(op.Namespace.Value, op.Id, op.Content).ContinueWith(t => (object?) t.Result);
+                        using MemoryBufferedPayload payload = new MemoryBufferedPayload(op.Content);
+                        tasks[index] = _storage.PutObject(op.Namespace.Value, payload, op.Id).ContinueWith(t => (object?) t.Result);
                         break;
                     }
                     case BatchOp.Operation.DELETE:
diff --git a/Engine/Source/Programs/Horde/HordeStorage/Horde.Storage/Horde.Storage.Startup.cs b/Engine/Source/Programs/Horde/HordeStorage/Horde.Storage/Horde.Storage.Startup.cs
--- a/Engine/Source/Programs/Horde/HordeStorage/Horde.Storage/Horde.Storage.Startup.cs
+++ b/Engine/Source/Programs/Horde/HordeStorage/Horde.Storage/Horde.Storage.Startup.cs
@@ -13,6 +13,7 @@ using Amazon.SecretsManager;
 using Cassandra;
 using Horde.Storage.Controllers;
 using Horde.Storage.Implementation;
+using Horde.Storage.Implementation.Blob;
 using Horde.Storage.Implementation.LeaderElection;
 using Jupiter;
 using Jupiter.Common.Implementation;
@@ -189,7 +190,9 @@ namespace Horde.Storage
                 return AWSCredentialsHelper.GetCredentials(awsSettings, "Horde.Storage");
             });
             services.AddSingleton(typeof(IAmazonDynamoDB), AddDynamo);
-            
+
+            services.AddSingleton<BufferedPayloadFactory>();
+
             services.AddSingleton<BlobCleanupService>();
             services.AddHostedService<BlobCleanupService>(p => p.GetService<BlobCleanupService>()!);
 
@@ -202,13 +205,19 @@ namespace Horde.Storage
             services.AddSingleton(serviceType: typeof(IContentIdStore), ContentIdStoreFactory);
 
             services.AddSingleton(serviceType: typeof(ITransactionLogWriter), TransactionLogWriterFactory);
+
+            services.AddSingleton(serviceType: typeof(IBlobIndex), BlobIndexFactory);
             services.AddSingleton(serviceType: typeof(IAmazonS3), CreateS3);
-            services.AddSingleton(serviceType: typeof(IBlobStore), StorageBackendFactory);
 
             services.AddSingleton<AmazonS3Store>();
             services.AddSingleton<FileSystemStore>();
             services.AddSingleton<AzureBlobStore>();
             services.AddSingleton<MemoryCacheBlobStore>();
+            services.AddSingleton<MemoryBlobStore>();
+
+            services.AddSingleton<OrphanBlobCleanup>();
+
+            services.AddSingleton(typeof(IBlobService), typeof(BlobService));
 
             services.AddSingleton(serviceType: typeof(IScyllaSessionManager), ScyllaFactory);
 
@@ -261,6 +270,20 @@ namespace Horde.Storage
             }
         }
 
+        private IBlobIndex BlobIndexFactory(IServiceProvider provider)
+        {
+            HordeStorageSettings settings = provider.GetService<IOptionsMonitor<HordeStorageSettings>>()!.CurrentValue!;
+            switch (settings.BlobIndexImplementation)
+            {
+                case HordeStorageSettings.BlobIndexImplementations.Memory:
+                    return ActivatorUtilities.CreateInstance<MemoryBlobIndex>(provider);
+                case HordeStorageSettings.BlobIndexImplementations.Scylla:
+                    return ActivatorUtilities.CreateInstance<ScyllaBlobIndex>(provider);
+                default:
+                    throw new ArgumentOutOfRangeException();
+            }
+        }
+
         private IAmazonSecretsManager CreateAWSSecretsManager(IServiceProvider provider)
         {
             AWSCredentials awsCredentials = provider.GetService<AWSCredentials>()!;
@@ -498,56 +521,6 @@ namespace Horde.Storage
             }
         }
 
-        private HordeStorageSettings.StorageBackendImplementations[] GetStorageImplementationEnums(string[]? storageImpls)
-        {
-            storageImpls ??= new[] {HordeStorageSettings.StorageBackendImplementations.Memory.ToString()};
-            return storageImpls.Select(x =>
-            {
-                if (Enum.TryParse(typeof(HordeStorageSettings.StorageBackendImplementations), x, true, out var backendType) && backendType is HordeStorageSettings.StorageBackendImplementations type)
-                {
-                    return type;
-                }
-
-                throw new ArgumentException($"Unable to find storage implementation for specified value '{x}'");
-            }).ToArray();
-        }
-        
-        private IBlobStore StorageBackendFactory(IServiceProvider provider)
-        {
-            IOptionsMonitor<HordeStorageSettings> settings = provider.GetService<IOptionsMonitor<HordeStorageSettings>>()!;
-            IOptionsMonitor<GCSettings> gcSettings = provider.GetService<IOptionsMonitor<GCSettings>>()!;
-            BlobCleanupService blobCleanupService = provider.GetService<BlobCleanupService>()!;
-
-            List<IBlobStore> storageBackends = GetStorageImplementationEnums(settings.CurrentValue.StorageImplementations)
-                    .Select<HordeStorageSettings.StorageBackendImplementations, IBlobStore?>(impl =>
-                    {
-                        switch (impl)
-                        {
-                            case HordeStorageSettings.StorageBackendImplementations.S3: return provider.GetService<AmazonS3Store>();
-                            case HordeStorageSettings.StorageBackendImplementations.Azure: return provider.GetService<AzureBlobStore>();
-                            case HordeStorageSettings.StorageBackendImplementations.Memory: return provider.GetService<MemoryCacheBlobStore>();
-                            case HordeStorageSettings.StorageBackendImplementations.FileSystem:
-                                FileSystemStore store = ActivatorUtilities.CreateInstance<FileSystemStore>(provider)!;
-                                blobCleanupService.RegisterCleanup(store);
-                                return store;
-                            default: throw new NotImplementedException();
-                        }
-                    })
-                    .Select(x => x ?? throw new ArgumentException("Unable to resolve all stores!"))
-                    .ToList();
-
-            // Only use the hierarchical backend if more than one backend is specified 
-            IBlobStore blobStore = storageBackends.Count == 1 ? storageBackends[0] : new HierarchicalBlobStore(storageBackends);
-            
-            if (gcSettings.CurrentValue.CleanOldBlobs)
-            {
-                OrphanBlobCleanup orphanBlobCleanup = ActivatorUtilities.CreateInstance<OrphanBlobCleanup>(provider, blobStore);
-                blobCleanupService.RegisterCleanup(orphanBlobCleanup);    
-            }
-
-            return blobStore;
-        }
-
         private static IRefsStore RefStoreFactory(IServiceProvider provider)
         {
             HordeStorageSettings settings = provider.GetService<IOptionsMonitor<HordeStorageSettings>>()!.CurrentValue;
@@ -605,7 +578,7 @@ namespace Horde.Storage
                     throw new ArgumentOutOfRangeException();
             }
 
-            foreach (HordeStorageSettings.StorageBackendImplementations impl in GetStorageImplementationEnums(settings.StorageImplementations))
+            foreach (HordeStorageSettings.StorageBackendImplementations impl in settings.GetStorageImplementations())
             {
                 switch (impl)
                 {
@@ -635,6 +608,7 @@ namespace Horde.Storage
                         });
                         break;
                     case HordeStorageSettings.StorageBackendImplementations.Memory:
+                    case HordeStorageSettings.StorageBackendImplementations.MemoryBlobStore:
                         break;
                     default:
                         throw new ArgumentOutOfRangeException("Unhandled storage impl " + impl);
diff --git a/Engine/Source/Programs/Horde/HordeStorage/Horde.Storage/Implementation/Blob/AmazonS3Store.cs b/Engine/Source/Programs/Horde/HordeStorage/Horde.Storage/Implementation/Blob/AmazonS3Store.cs
--- a/Engine/Source/Programs/Horde/HordeStorage/Horde.Storage/Implementation/Blob/AmazonS3Store.cs
+++ b/Engine/Source/Programs/Horde/HordeStorage/Horde.Storage/Implementation/Blob/AmazonS3Store.cs
@@ -7,6 +7,7 @@ using System.Net;
 using System.Threading.Tasks;
 using Amazon.S3;
 using Amazon.S3.Model;
+using Horde.Storage.Implementation.Blob;
 using Jupiter.Common.Implementation;
 using Jupiter.Implementation;
 using Microsoft.Extensions.Options;
@@ -18,13 +19,15 @@ namespace Horde.Storage.Implementation
     {
         private readonly ILogger _logger = Log.ForContext<AmazonS3Store>();
         private readonly IAmazonS3 _amazonS3;
+        private readonly IBlobIndex _blobIndex;
         private readonly S3Settings _settings;
         private readonly HashSet<string> _bucketAccessPolicyApplied = new HashSet<string>();
         private readonly HashSet<string> _bucketExistenceChecked = new HashSet<string>();
 
-        public AmazonS3Store(IAmazonS3 amazonS3, IOptionsMonitor<S3Settings> settings)
+        public AmazonS3Store(IAmazonS3 amazonS3, IOptionsMonitor<S3Settings> settings, IBlobIndex blobIndex)
         {
             _amazonS3 = amazonS3;
+            _blobIndex = blobIndex;
             _settings = settings.CurrentValue;
         }
 
@@ -147,6 +150,11 @@ namespace Horde.Storage.Implementation
 
         public async Task<bool> Exists(NamespaceId ns, BlobIdentifier blobIdentifier)
         {
+            if (_settings.UseBlobIndexForExistsCheck)
+            {
+                return await _blobIndex.BlobExistsInRegion(ns, blobIdentifier);
+            }
+
             string bucketName = GetBucketName(ns);
             try
             {
@@ -186,7 +194,7 @@ namespace Horde.Storage.Implementation
             }
         }
 
-        public async IAsyncEnumerable<BlobIdentifier> ListOldObjects(NamespaceId ns, DateTime cutoff)
+        public async IAsyncEnumerable<(BlobIdentifier, DateTime)> ListObjects(NamespaceId ns)
         {
             ListObjectsV2Request request = new ListObjectsV2Request
             {
@@ -202,13 +210,7 @@ namespace Horde.Storage.Implementation
                     string key = obj.Key;
                     string identifierString = key.Substring(key.LastIndexOf("/", StringComparison.Ordinal)+1);
 
-                    // if the object was updated after the cutoff time we do not consider it old
-                    if (obj.LastModified > cutoff)
-                    {
-                        continue;
-                    }
-
-                    yield return new BlobIdentifier(identifierString);
+                    yield return (new BlobIdentifier(identifierString), obj.LastModified);
                 }
 
                 request.ContinuationToken = response.NextContinuationToken;
diff --git a/Engine/Source/Programs/Horde/HordeStorage/Horde.Storage/Implementation/Blob/AzureBlobStore.cs b/Engine/Source/Programs/Horde/HordeStorage/Horde.Storage/Implementation/Blob/AzureBlobStore.cs
--- a/Engine/Source/Programs/Horde/HordeStorage/Horde.Storage/Implementation/Blob/AzureBlobStore.cs
+++ b/Engine/Source/Programs/Horde/HordeStorage/Horde.Storage/Implementation/Blob/AzureBlobStore.cs
@@ -146,7 +146,7 @@ namespace Horde.Storage.Implementation
             }
         }
 
-        public async IAsyncEnumerable<BlobIdentifier> ListOldObjects(NamespaceId ns, DateTime cutoff)
+        public async IAsyncEnumerable<(BlobIdentifier, DateTime)> ListObjects(NamespaceId ns)
         {
             string fixedNamespace = SanitizeNamespace(ns);
             BlobContainerClient container = new BlobContainerClient(_connectionString, fixedNamespace);
@@ -156,13 +156,7 @@ namespace Horde.Storage.Implementation
 
             await foreach (var item in container.GetBlobsAsync(BlobTraits.Metadata))
             {
-                // ignore any recent blob
-                if ((item.Properties?.LastModified ?? DateTime.Now) > cutoff)
-                {
-                    continue;
-                }
-
-                yield return new BlobIdentifier(item.Name);
+                yield return (new BlobIdentifier(item.Name), item.Properties?.LastModified?.DateTime ?? DateTime.Now);
             }
         }
 
diff --git a/Engine/Source/Programs/Horde/HordeStorage/Horde.Storage/Implementation/Blob/ConsistencyCheckService.cs b/Engine/Source/Programs/Horde/HordeStorage/Horde.Storage/Implementation/Blob/ConsistencyCheckService.cs
--- a/Engine/Source/Programs/Horde/HordeStorage/Horde.Storage/Implementation/Blob/ConsistencyCheckService.cs
+++ b/Engine/Source/Programs/Horde/HordeStorage/Horde.Storage/Implementation/Blob/ConsistencyCheckService.cs
@@ -82,11 +82,10 @@ namespace Horde.Storage.Implementation
             }
 
             long countOfBlobsChecked = 0;
-            DateTime objectCutoff = DateTime.Now;
             // technically this does not need to be run per namespace but per s3 bucket
             await foreach (NamespaceId ns in _refsStore.GetNamespaces())
             {
-                await foreach (BlobIdentifier blob in s3Store.ListOldObjects(ns, objectCutoff))
+                await foreach ((BlobIdentifier blob, DateTime lastModified) in s3Store.ListObjects(ns))
                 {
                     if (countOfBlobsChecked % 100 == 0)
                         _logger.Information("Consistency check running on S3 blob store, count of blobs processed so far: {CountOfBlobs}", countOfBlobsChecked);
diff --git a/Engine/Source/Programs/Horde/HordeStorage/Horde.Storage/Implementation/Blob/FileSystemStore.cs b/Engine/Source/Programs/Horde/HordeStorage/Horde.Storage/Implementation/Blob/FileSystemStore.cs
--- a/Engine/Source/Programs/Horde/HordeStorage/Horde.Storage/Implementation/Blob/FileSystemStore.cs
+++ b/Engine/Source/Programs/Horde/HordeStorage/Horde.Storage/Implementation/Blob/FileSystemStore.cs
@@ -21,9 +21,10 @@ namespace Horde.Storage.Implementation
 
         internal const int DefaultBufferSize = 4096;
 
-        public FileSystemStore(IOptionsMonitor<FilesystemSettings> settings)
+        public FileSystemStore(IOptionsMonitor<FilesystemSettings> settings, BlobCleanupService cleanupService)
         {
             _settings = settings;
+            cleanupService.RegisterCleanup(this);
         }
 
         public static FileInfo GetFilesystemPath(string rootDir, NamespaceId ns, BlobIdentifier blob)
@@ -160,12 +161,12 @@ namespace Horde.Storage.Implementation
             return Task.CompletedTask;
         }
 
-        public IAsyncEnumerable<BlobIdentifier> ListOldObjects(NamespaceId ns, DateTime cutoff)
+        public IAsyncEnumerable<(BlobIdentifier,DateTime)> ListObjects(NamespaceId ns)
         {
-            return DoListOldObjects(ns, cutoff).ToAsyncEnumerable();
+            return DoListOldObjects(ns).ToAsyncEnumerable();
         }
 
-        private IEnumerable<BlobIdentifier> DoListOldObjects(NamespaceId ns, DateTime cutoff)
+        private IEnumerable<(BlobIdentifier, DateTime)> DoListOldObjects(NamespaceId ns)
         {
             DirectoryInfo di = new DirectoryInfo(Path.Combine(_settings.CurrentValue.RootDir, ns.ToString()));
             if (!di.Exists)
@@ -180,14 +181,7 @@ namespace Horde.Storage.Implementation
                     continue;
                 }
 
-                if (file.LastWriteTime > cutoff)
-                {
-                    _logger.Information("Ignoring {Blob} as it is new. Last modified: {LastModified} Cutoff: {CutOff}", file.Name, file.LastWriteTime, cutoff);
-                    continue;
-                }
-                _logger.Information("Old {Blob} found. Last modified: {LastModified} Cutoff: {CutOff}", file.Name, file.LastWriteTime, cutoff);
-
-                yield return new BlobIdentifier(file.Name);
+                yield return (new BlobIdentifier(file.Name), file.LastWriteTime);
             }
         }
         
@@ -226,15 +220,13 @@ namespace Horde.Storage.Implementation
                     return blobsRemoved;
                 }
                 
-                FileInfo[] fileInfos = await GetLeastRecentlyAccessedObjects(maxResults: batchSize);
-                if (fileInfos.Length == 0)
-                {
-                    return blobsRemoved;
-                }
-                
+                IEnumerable<FileInfo> fileInfos = GetLeastRecentlyAccessedObjects(maxResults: batchSize);
+
+                bool hadFiles = false;
                 long totalBytesDeleted = 0;
                 foreach (FileInfo fi in fileInfos)
                 {
+                    hadFiles = true;
                     try
                     {
                         totalBytesDeleted += fi.Length;
@@ -252,6 +244,11 @@ namespace Horde.Storage.Implementation
                         // if the file was gced while running we can just ignore it
                     }
                 }
+
+                if (!hadFiles)
+                {
+                    return blobsRemoved;
+                }
             }
             
             return blobsRemoved;
@@ -264,7 +261,7 @@ namespace Horde.Storage.Implementation
         /// <param name="ns">Namespace, if set to null all namespaces will be scanned</param>
         /// <param name="maxResults">Max results to return. Note that the entire namespace will be scanned no matter what.</param>
         /// <returns>Enumerable of least recently accessed objects as FileInfos</returns>
-        public async Task<FileInfo[]> GetLeastRecentlyAccessedObjects(NamespaceId? ns = null, int maxResults = 10000)
+        public IEnumerable<FileInfo> GetLeastRecentlyAccessedObjects(NamespaceId? ns = null, int maxResults = 10000)
         {
             string path = ns != null ? Path.Combine(_settings.CurrentValue.RootDir, ns.ToString()!) : _settings.CurrentValue.RootDir;
             DirectoryInfo di = new DirectoryInfo(path);
@@ -272,8 +269,8 @@ namespace Horde.Storage.Implementation
             {
                 return Array.Empty<FileInfo>();
             }
-            
-            return await Task.Run(() => di.EnumerateFiles("*", SearchOption.AllDirectories).OrderBy(x => x.LastWriteTime).Take(maxResults).ToArray());
+
+            return di.EnumerateFiles("*", SearchOption.AllDirectories).OrderBy(x => x.LastWriteTime).Take(maxResults);
         }
         
         /// <summary>
diff --git a/Engine/Source/Programs/Horde/HordeStorage/Horde.Storage/Implementation/Blob/HierarchicalBlobStore.cs b/Engine/Source/Programs/Horde/HordeStorage/Horde.Storage/Implementation/Blob/HierarchicalBlobStore.cs
deleted file mode 100644
diff --git a/Engine/Source/Programs/Horde/HordeStorage/Horde.Storage/Implementation/Blob/IBlobIndex.cs b/Engine/Source/Programs/Horde/HordeStorage/Horde.Storage/Implementation/Blob/IBlobIndex.cs
new file mode 100644
--- /dev/null
+++ b/Engine/Source/Programs/Horde/HordeStorage/Horde.Storage/Implementation/Blob/IBlobIndex.cs
@@ -0,0 +1,23 @@
+using System;
+using System.Collections.Generic;
+using System.Threading.Tasks;
+using AutoMapper;
+using Jupiter.Implementation;
+
+namespace Horde.Storage.Implementation.Blob
+{
+    public interface IBlobIndex
+    {
+        public class BlobInfo
+        {
+            public HashSet<string> Regions = new HashSet<string>();
+        }
+
+        Task AddBlobToIndex(NamespaceId ns, BlobIdentifier id);
+
+        Task<BlobInfo?> GetBlobInfo(NamespaceId ns, BlobIdentifier id);
+
+        Task<bool> RemoveBlobFromIndex(NamespaceId ns, BlobIdentifier id);
+        Task<bool> BlobExistsInRegion(NamespaceId ns, BlobIdentifier blobIdentifier);
+    }
+}
diff --git a/Engine/Source/Programs/Horde/HordeStorage/Horde.Storage/Implementation/Blob/IBlobService.cs b/Engine/Source/Programs/Horde/HordeStorage/Horde.Storage/Implementation/Blob/IBlobService.cs
new file mode 100644
--- /dev/null
+++ b/Engine/Source/Programs/Horde/HordeStorage/Horde.Storage/Implementation/Blob/IBlobService.cs
@@ -0,0 +1,426 @@
+﻿// Copyright Epic Games, Inc. All Rights Reserved.
+
+using System;
+using System.Collections.Concurrent;
+using System.Collections.Generic;
+using System.IO;
+using System.Linq;
+using System.Threading.Tasks;
+using Dasync.Collections;
+using Datadog.Trace;
+using Horde.Storage.Controllers;
+using Horde.Storage.Implementation.Blob;
+using Jupiter.Common.Implementation;
+using Jupiter.Implementation;
+using Microsoft.Extensions.DependencyInjection;
+using Microsoft.Extensions.Options;
+using Nito.AsyncEx;
+using Serilog;
+
+namespace Horde.Storage.Implementation;
+
+public interface IBlobService
+{
+    Task<BlobIdentifier> VerifyContentMatchesHash(Stream content, BlobIdentifier identifier);
+    Task<BlobIdentifier> PutObjectKnownHash(NamespaceId ns, IBufferedPayload content, BlobIdentifier identifier);
+    Task<BlobIdentifier> PutObject(NamespaceId ns, IBufferedPayload payload, BlobIdentifier identifier);
+    Task<BlobIdentifier> PutObject(NamespaceId ns, byte[] payload, BlobIdentifier identifier);
+
+    Task<BlobContents> GetObject(NamespaceId ns, BlobIdentifier blob);
+    Task<bool> Exists(NamespaceId ns, BlobIdentifier blob);
+
+    // Delete a object
+    Task DeleteObject(NamespaceId ns, BlobIdentifier blob);
+
+
+    // delete the whole namespace
+    Task DeleteNamespace(NamespaceId ns);
+
+    IAsyncEnumerable<(BlobIdentifier,DateTime)> ListObjects(NamespaceId ns);
+    Task<BlobIdentifier[]> FilterOutKnownBlobs(NamespaceId ns, BlobIdentifier[] blobs);
+    Task<BlobIdentifier[]> FilterOutKnownBlobs(NamespaceId ns, IAsyncEnumerable<BlobIdentifier> blobs);
+    Task<BlobContents> GetObjects(NamespaceId ns, BlobIdentifier[] refRequestBlobReferences);
+
+}
+
+public class BlobService : IBlobService
+{
+    private List<IBlobStore> _blobStores;
+    private readonly IBlobIndex _blobIndex;
+    private readonly ILogger _logger = Log.ForContext<BlobService>();
+
+    internal IEnumerable<IBlobStore> BlobStore
+    {
+        get { return _blobStores; }
+        set { _blobStores = value.ToList(); } 
+    }
+
+    public BlobService(IServiceProvider provider, IOptionsMonitor<HordeStorageSettings> settings, IBlobIndex blobIndex)
+    {
+        _blobStores = GetBlobStores(provider, settings).ToList();
+        _blobIndex = blobIndex;
+    }
+
+    private IEnumerable<IBlobStore> GetBlobStores(IServiceProvider provider, IOptionsMonitor<HordeStorageSettings> settings)
+    {
+        return settings.CurrentValue.GetStorageImplementations().Select(impl => ToStorageImplementation(provider, impl));
+    }
+
+    private IBlobStore ToStorageImplementation(IServiceProvider provider, HordeStorageSettings.StorageBackendImplementations impl)
+    {
+        IBlobStore? store = impl switch
+        {
+            HordeStorageSettings.StorageBackendImplementations.S3 => provider.GetService<AmazonS3Store>(),
+            HordeStorageSettings.StorageBackendImplementations.Azure => provider.GetService<AzureBlobStore>(),
+            HordeStorageSettings.StorageBackendImplementations.FileSystem => provider.GetService<FileSystemStore>(),
+            HordeStorageSettings.StorageBackendImplementations.Memory => provider.GetService<MemoryCacheBlobStore>(),
+            HordeStorageSettings.StorageBackendImplementations.MemoryBlobStore => provider.GetService<MemoryBlobStore>(),
+            _ => throw new ArgumentOutOfRangeException()
+        };
+        if (store == null)
+            throw new ArgumentException($"Failed to find a provider service for type: {impl}");
+
+        return store;
+    }
+
+    public async Task<BlobIdentifier> VerifyContentMatchesHash(Stream content, BlobIdentifier identifier)
+    {
+        BlobIdentifier blobHash;
+        {
+            using Scope _ = Tracer.Instance.StartActive("web.hash");
+            blobHash = await BlobIdentifier.FromStream(content);
+        }
+
+        if (!identifier.Equals(blobHash))
+        {
+            throw new HashMismatchException(identifier, blobHash);
+        }
+
+        return identifier;
+    }
+
+    public async Task<BlobIdentifier> PutObject(NamespaceId ns, IBufferedPayload payload, BlobIdentifier identifier)
+    {
+        await using Stream hashStream = payload.GetStream();
+        BlobIdentifier id = await VerifyContentMatchesHash(hashStream, identifier);
+
+        Task<BlobIdentifier> putObjectTask = PutObjectToStores(ns, payload, id);
+        Task addToBlobIndexTask = _blobIndex.AddBlobToIndex(ns, id);
+
+        await Task.WhenAll(putObjectTask, addToBlobIndexTask);
+
+        return await putObjectTask;
+
+    }
+
+    public async Task<BlobIdentifier> PutObject(NamespaceId ns, byte[] payload, BlobIdentifier identifier)
+    {
+        await using Stream hashStream = new MemoryStream(payload);
+        BlobIdentifier id = await VerifyContentMatchesHash(hashStream, identifier);
+
+        Task<BlobIdentifier> putObjectTask = PutObjectToStores(ns, payload, id);
+        Task addToBlobIndexTask = _blobIndex.AddBlobToIndex(ns, id);
+
+        await Task.WhenAll(putObjectTask, addToBlobIndexTask);
+
+        return await putObjectTask;
+    }
+
+    public async Task<BlobIdentifier> PutObjectKnownHash(NamespaceId ns, IBufferedPayload content,
+        BlobIdentifier identifier)
+    {
+        Task<BlobIdentifier> putObjectTask = PutObjectToStores(ns, content, identifier);
+        Task addToBlobIndexTask = _blobIndex.AddBlobToIndex(ns, identifier);
+
+        await Task.WhenAll(putObjectTask, addToBlobIndexTask);
+
+        return await putObjectTask;
+    }
+
+    private async Task<BlobIdentifier> PutObjectToStores(NamespaceId ns, IBufferedPayload bufferedPayload, BlobIdentifier identifier)
+    {
+        await Parallel.ForEachAsync(_blobStores, async (store, cancellationToken) =>
+        {
+            await using Stream s = bufferedPayload.GetStream();
+            await store.PutObject(ns, s, identifier);
+        });
+
+        return identifier;
+    }
+
+    private async Task<BlobIdentifier> PutObjectToStores(NamespaceId ns, byte[] payload, BlobIdentifier identifier)
+    {
+        await Parallel.ForEachAsync(_blobStores, async (store, cancellationToken) =>
+        {
+            await store.PutObject(ns, payload, identifier);
+        });
+
+        return identifier;
+    }
+
+    public async Task<BlobContents> GetObject(NamespaceId ns, BlobIdentifier blob)
+    {
+        bool seenBlobNotFound = false;
+        bool seenNamespaceNotFound = false;
+
+        List<Task<(BlobContents?, IBlobStore?)>> getTasks = new();
+
+        foreach (IBlobStore store in _blobStores)
+        {
+            getTasks.Add(Task.Run( async () =>
+            {
+                using Scope scope = Tracer.Instance.StartActive("HierarchicalStore.GetObject");
+                scope.Span.SetTag("BlobStore", store.GetType().Name);
+                scope.Span.SetTag("ObjectFound", false.ToString());
+                try
+                {
+                    BlobContents? blobContents = await store.GetObject(ns, blob);
+                    scope.Span.SetTag("ObjectFound", true.ToString());
+                    return ((BlobContents?)blobContents, (IBlobStore?)store);
+                }
+                catch (BlobNotFoundException)
+                {
+                    seenBlobNotFound = true;
+                }
+                catch (NamespaceNotFoundException)
+                {
+                    seenNamespaceNotFound = true;
+                }
+
+                return (null, null);
+            }));
+        }
+
+        BlobContents? blobContents = null;
+        IBlobStore? storeWithBlob = null;
+        while (getTasks.Count != 0)
+        {
+            Task<(BlobContents?, IBlobStore?)> completedTask = await getTasks.WhenAny();
+            (BlobContents? contents, IBlobStore? store)= await completedTask;
+            if (contents == null)
+            {
+                // the task did not find a object, we wait for the other tasks to run
+                getTasks.Remove(completedTask);
+                continue;
+            }
+
+            blobContents = contents;
+            storeWithBlob = store;
+            // object found, we can break
+            break;
+        }
+
+        if (seenBlobNotFound && blobContents == null)
+        {
+            throw new BlobNotFoundException(ns, blob);
+        }
+
+        if (seenNamespaceNotFound && blobContents == null)
+        {
+            throw new NamespaceNotFoundException(ns);
+        }
+        
+        if (blobContents == null)
+        {
+            // Should not happen but exists to safeguard against the null pointer
+            throw new Exception("blobContents is null");
+        }
+        if (storeWithBlob == null)
+        {
+            // Should not happen but exists to safeguard against the null pointer
+            throw new Exception("storeWithBlob is null");
+        }
+
+        int indexOfStoreWithBlob = _blobStores.IndexOf(storeWithBlob);
+        // if it wasn't the first store we need to populate the stores before it with the blob
+        if (indexOfStoreWithBlob != 0)
+        {
+            using Scope _ = Tracer.Instance.StartActive("HierarchicalStore.Populate");
+            // buffer the blob contents as it could be to large for us to buffer into memory before forwarding to the blob stores
+            using FilesystemBufferedPayload payload = await FilesystemBufferedPayload.Create(blobContents.Stream);
+
+            // Don't populate the last store, as that is where we got the hit
+            // Populate each store traversed that did not have the content found lower in the hierarchy
+            await Parallel.ForEachAsync(_blobStores.GetRange(0, indexOfStoreWithBlob), async (store, token) =>
+            {
+                using Scope scope = Tracer.Instance.StartActive("HierarchicalStore.PopulateStore");
+                scope.Span.SetTag("BlobStore", store.GetType().Name);
+
+                await using Stream s = payload.GetStream();
+                await store.PutObject(ns, s, blob);
+            });
+
+            // we just pushed the content to all stores, thus consuming it, so we just refetch it from the store that just reported it had it
+            blobContents = await storeWithBlob.GetObject(ns, blob);
+        }
+        
+        return blobContents;
+    }
+
+    public async Task<bool> Exists(NamespaceId ns, BlobIdentifier blob)
+    {
+        List<Task<bool>> existsTasks = new();
+
+        foreach (IBlobStore store in _blobStores)
+        {
+            existsTasks.Add(Task.Run( async () =>
+            {
+                using Scope scope = Tracer.Instance.StartActive("HierarchicalStore.ObjectExists");
+                scope.Span.SetTag("BlobStore", store.GetType().Name);
+                if (await store.Exists(ns, blob))
+                {
+                    scope.Span.SetTag("ObjectFound", true.ToString());
+                    return true;
+                }
+
+                scope.Span.SetTag("ObjectFound", false.ToString());
+                return false;
+            }));
+        }
+
+        while (existsTasks.Count != 0)
+        {
+            Task<bool> completedTask = await existsTasks.WhenAny();
+            bool exist = await completedTask;
+            if (!exist)
+            {
+                // the task did not find a object, we wait for the other tasks to run
+                existsTasks.Remove(completedTask);
+                continue;
+            }
+
+            // object found, we can break
+            return true;
+        }
+        return false;
+    }
+
+    public async Task DeleteObject(NamespaceId ns, BlobIdentifier blob)
+    {
+        bool blobNotFound = false;
+        bool deletedAtLeastOnce = false;
+        Task removeFromIndexTask = _blobIndex.RemoveBlobFromIndex(ns, blob);
+
+        List<Task> tasksToWaitFor = new List<Task>();
+        tasksToWaitFor.Add(removeFromIndexTask);
+
+        foreach (IBlobStore store in _blobStores)
+        {
+            try
+            {
+                using Scope scope = Tracer.Instance.StartActive("HierarchicalStore.DeleteObject");
+                scope.Span.SetTag("BlobStore", store.GetType().Name);
+                Task deleteTask = store.DeleteObject(ns, blob);
+                tasksToWaitFor.Add(deleteTask);
+                deletedAtLeastOnce = true;
+            }
+            catch (NamespaceNotFoundException)
+            {
+                // Ignore
+            }
+            catch (BlobNotFoundException)
+            {
+                blobNotFound = true;
+            }
+        }
+
+        await Task.WhenAll(tasksToWaitFor);
+
+        if (deletedAtLeastOnce)
+            return;
+
+        if (blobNotFound)
+            throw new BlobNotFoundException(ns, blob);
+
+        throw new NamespaceNotFoundException(ns);
+    }
+
+    public async Task DeleteNamespace(NamespaceId ns)
+    {
+        bool deletedAtLeastOnce = false;
+        foreach (IBlobStore store in _blobStores)
+        {
+            using Scope scope = Tracer.Instance.StartActive("HierarchicalStore.DeleteNamespace");
+            scope.Span.SetTag("BlobStore", store.GetType().Name);
+            try
+            {
+                await store.DeleteNamespace(ns);
+                deletedAtLeastOnce = true;
+            }
+            catch (NamespaceNotFoundException)
+            {
+                // Ignore
+            }
+        }
+
+        if (deletedAtLeastOnce)
+            return;
+
+        throw new NamespaceNotFoundException(ns);
+    }
+
+    public IAsyncEnumerable<(BlobIdentifier,DateTime)> ListObjects(NamespaceId ns)
+    {
+        // as this is a hierarchy of blob stores the last blob store should contain the superset of all stores
+        return _blobStores.Last().ListObjects(ns);
+    }
+
+    public async Task<BlobIdentifier[]> FilterOutKnownBlobs(NamespaceId ns, BlobIdentifier[] blobs)
+    {
+        var tasks = blobs.Select(async blobIdentifier => new { BlobIdentifier = blobIdentifier, Exists = await Exists(ns, blobIdentifier) });
+        var blobResults = await Task.WhenAll(tasks);
+        var filteredBlobs = blobResults.Where(ac => !ac.Exists).Select(ac => ac.BlobIdentifier);
+        return filteredBlobs.ToArray();
+    }
+
+    public async Task<BlobIdentifier[]> FilterOutKnownBlobs(NamespaceId ns, IAsyncEnumerable<BlobIdentifier> blobs)
+    {
+        ConcurrentBag<BlobIdentifier> missingBlobs = new ConcurrentBag<BlobIdentifier>();
+
+        try
+        {
+            await blobs.ParallelForEachAsync(async identifier =>
+            {
+                bool exists = await Exists(ns, identifier);
+
+                if (!exists)
+                {
+                    missingBlobs.Add(identifier);
+                }
+            });
+        }
+        catch (ParallelForEachException e)
+        {
+            if (e.InnerException is PartialReferenceResolveException)
+                throw e.InnerException;
+
+            throw;
+        }
+
+        return missingBlobs.ToArray();
+    }
+
+    public async Task<BlobContents> GetObjects(NamespaceId ns, BlobIdentifier[] blobs)
+    {
+        using Scope _ = Tracer.Instance.StartActive("blob.combine");
+        Task<BlobContents>[] tasks = new Task<BlobContents>[blobs.Length];
+        for (int i = 0; i < blobs.Length; i++)
+        {
+            tasks[i] = GetObject(ns, blobs[i]);
+        }
+
+        MemoryStream ms = new MemoryStream();
+        foreach (Task<BlobContents> task in tasks)
+        {
+            BlobContents blob = await task;
+            await using Stream s = blob.Stream;
+            await s.CopyToAsync(ms);
+        }
+
+        ms.Seek(0, SeekOrigin.Begin);
+
+        return new BlobContents(ms, ms.Length);
+    }
+
+
+}
diff --git a/Engine/Source/Programs/Horde/HordeStorage/Horde.Storage/Implementation/Blob/IBlobStore.cs b/Engine/Source/Programs/Horde/HordeStorage/Horde.Storage/Implementation/Blob/IBlobStore.cs
--- a/Engine/Source/Programs/Horde/HordeStorage/Horde.Storage/Implementation/Blob/IBlobStore.cs
+++ b/Engine/Source/Programs/Horde/HordeStorage/Horde.Storage/Implementation/Blob/IBlobStore.cs
@@ -29,7 +29,7 @@ namespace Horde.Storage.Implementation
         // delete the whole namespace
         Task DeleteNamespace(NamespaceId ns);
 
-        IAsyncEnumerable<BlobIdentifier> ListOldObjects(NamespaceId ns, DateTime cutoff);
+        IAsyncEnumerable<(BlobIdentifier,DateTime)> ListObjects(NamespaceId ns);
     }
 
     public class BlobNotFoundException : Exception
@@ -66,65 +66,4 @@ namespace Horde.Storage.Implementation
 
     }
 
-    public static class BlobStoreUtils
-    {
-        public static async Task<BlobIdentifier[]> FilterOutKnownBlobs(this IBlobStore blobStore, NamespaceId ns, BlobIdentifier[] blobs)
-        {
-            var tasks = blobs.Select(async blobIdentifier => new { BlobIdentifier = blobIdentifier, Exists = await blobStore.Exists(ns, blobIdentifier) });
-            var blobResults = await Task.WhenAll(tasks);
-            var filteredBlobs = blobResults.Where(ac => !ac.Exists).Select(ac => ac.BlobIdentifier);
-            return filteredBlobs.ToArray();
-        }
-
-        public static async Task<BlobIdentifier[]> FilterOutKnownBlobs(this IBlobStore blobStore, NamespaceId ns, IAsyncEnumerable<BlobIdentifier> blobs)
-        {
-            ConcurrentBag<BlobIdentifier> missingBlobs = new ConcurrentBag<BlobIdentifier>();
-
-            try
-            {
-                await blobs.ParallelForEachAsync(async identifier =>
-                {
-                    bool exists = await blobStore.Exists(ns, identifier);
-
-                    if (!exists)
-                    {
-                        missingBlobs.Add(identifier);
-                    }
-                });
-            }
-            catch (ParallelForEachException e)
-            {
-                if (e.InnerException is PartialReferenceResolveException)
-                    throw e.InnerException;
-
-                throw;
-            }
-
-            return missingBlobs.ToArray();
-        }
-
-
-        public static async Task<BlobContents> GetObjects(this IBlobStore blobStore, NamespaceId ns, BlobIdentifier[] blobs)
-        {
-            using Scope _ = Tracer.Instance.StartActive("blob.combine");
-            Task<BlobContents>[] tasks = new Task<BlobContents>[blobs.Length];
-            for (int i = 0; i < blobs.Length; i++)
-            {
-                tasks[i] = blobStore.GetObject(ns, blobs[i]);
-            }
-
-            MemoryStream ms = new MemoryStream();
-            foreach (Task<BlobContents> task in tasks)
-            {
-                BlobContents blob = await task;
-                await using Stream s = blob.Stream;
-                await s.CopyToAsync(ms);
-            }
-
-            ms.Seek(0, SeekOrigin.Begin);
-
-            return new BlobContents(ms, ms.Length);
-        }
-    }
-
 }
diff --git a/Engine/Source/Programs/Horde/HordeStorage/Horde.Storage/Implementation/Blob/MemoryBlobIndex.cs b/Engine/Source/Programs/Horde/HordeStorage/Horde.Storage/Implementation/Blob/MemoryBlobIndex.cs
new file mode 100644
--- /dev/null
+++ b/Engine/Source/Programs/Horde/HordeStorage/Horde.Storage/Implementation/Blob/MemoryBlobIndex.cs
@@ -0,0 +1,63 @@
+// Copyright Epic Games, Inc. All Rights Reserved.
+
+using System.Collections.Concurrent;
+using System.Collections.Generic;
+using System.Threading.Tasks;
+using Jupiter;
+using Jupiter.Implementation;
+using Microsoft.Extensions.Options;
+
+namespace Horde.Storage.Implementation.Blob;
+
+public class MemoryBlobIndex : IBlobIndex
+{
+    private class MemoryBlobInfo : IBlobIndex.BlobInfo
+    {
+    }
+
+    private readonly ConcurrentDictionary<NamespaceId, ConcurrentDictionary<BlobIdentifier, MemoryBlobInfo>> _index = new ();
+    private readonly IOptionsMonitor<JupiterSettings> _jupiterSettings;
+
+    public MemoryBlobIndex(IOptionsMonitor<JupiterSettings> settings)
+    {
+        _jupiterSettings = settings;
+    }
+
+    private ConcurrentDictionary<BlobIdentifier, MemoryBlobInfo> GetNamespaceContainer(NamespaceId ns)
+    {
+        return _index.GetOrAdd(ns, id => new ConcurrentDictionary<BlobIdentifier, MemoryBlobInfo>());
+    }
+
+    public Task AddBlobToIndex(NamespaceId ns, BlobIdentifier id)
+    {
+        ConcurrentDictionary<BlobIdentifier, MemoryBlobInfo> index = GetNamespaceContainer(ns);
+        index[id] = new MemoryBlobInfo
+        {
+            Regions = new HashSet<string> { _jupiterSettings.CurrentValue.CurrentSite }
+        };
+        return Task.CompletedTask;
+    }
+
+    public Task<IBlobIndex.BlobInfo?> GetBlobInfo(NamespaceId ns, BlobIdentifier id)
+    {
+        ConcurrentDictionary<BlobIdentifier, MemoryBlobInfo> index = GetNamespaceContainer(ns);
+
+        if (!index.TryGetValue(id, out MemoryBlobInfo? blobInfo))
+            return Task.FromResult<IBlobIndex.BlobInfo?>(null);
+
+        return Task.FromResult<IBlobIndex.BlobInfo?>(blobInfo);
+    }
+
+    public Task<bool> RemoveBlobFromIndex(NamespaceId ns, BlobIdentifier id)
+    {
+        ConcurrentDictionary<BlobIdentifier, MemoryBlobInfo> index = GetNamespaceContainer(ns);
+
+        return Task.FromResult(index.TryRemove(id, out MemoryBlobInfo? _));
+    }
+
+    public async Task<bool> BlobExistsInRegion(NamespaceId ns, BlobIdentifier blobIdentifier)
+    {
+        IBlobIndex.BlobInfo? blobInfo = await GetBlobInfo(ns, blobIdentifier);
+        return blobInfo?.Regions.Contains(_jupiterSettings.CurrentValue.CurrentSite) ?? false;
+    }
+}
diff --git a/Engine/Source/Programs/Horde/HordeStorage/Horde.Storage/Implementation/Blob/MemoryBlobStore.cs b/Engine/Source/Programs/Horde/HordeStorage/Horde.Storage/Implementation/Blob/MemoryBlobStore.cs
--- a/Engine/Source/Programs/Horde/HordeStorage/Horde.Storage/Implementation/Blob/MemoryBlobStore.cs
+++ b/Engine/Source/Programs/Horde/HordeStorage/Horde.Storage/Implementation/Blob/MemoryBlobStore.cs
@@ -53,7 +53,7 @@ namespace Horde.Storage.Implementation
             }
 
             var namespaceContainer = _blobs.GetOrAdd(ns, new ConcurrentDictionary<BlobIdentifier, BlobContainer>());
-            if (namespaceContainer.ContainsKey(identifier))
+            if (throwOnOverwrite && namespaceContainer.ContainsKey(identifier))
             {
                 throw new Exception($"Blob {identifier} already exists in {ns}");
             }
@@ -124,7 +124,7 @@ namespace Horde.Storage.Implementation
             return Task.CompletedTask;
         }
 
-        public async IAsyncEnumerable<BlobIdentifier> ListOldObjects(NamespaceId ns, DateTime cutoff)
+        public async IAsyncEnumerable<(BlobIdentifier,DateTime)> ListObjects(NamespaceId ns)
         {
             if (!_blobs.TryGetValue(ns, value: out var namespaceContainer))
             {
@@ -134,11 +134,7 @@ namespace Horde.Storage.Implementation
             await Task.CompletedTask;
             foreach (BlobContainer blobContainer in namespaceContainer.Values)
             {
-                if (blobContainer.LastModified > cutoff)
-                {
-                    continue;
-                }
-                yield return blobContainer.BlobIdentifier;
+                yield return (blobContainer.BlobIdentifier, blobContainer.LastModified);
             }
         }
 
diff --git a/Engine/Source/Programs/Horde/HordeStorage/Horde.Storage/Implementation/Blob/MemoryCacheBlobStore.cs b/Engine/Source/Programs/Horde/HordeStorage/Horde.Storage/Implementation/Blob/MemoryCacheBlobStore.cs
--- a/Engine/Source/Programs/Horde/HordeStorage/Horde.Storage/Implementation/Blob/MemoryCacheBlobStore.cs
+++ b/Engine/Source/Programs/Horde/HordeStorage/Horde.Storage/Implementation/Blob/MemoryCacheBlobStore.cs
@@ -102,10 +102,10 @@ namespace Horde.Storage.Implementation
             return Task.CompletedTask;
         }
 
-        public IAsyncEnumerable<BlobIdentifier> ListOldObjects(NamespaceId ns, DateTime cutoff)
+        public IAsyncEnumerable<(BlobIdentifier,DateTime)> ListObjects(NamespaceId ns)
         {
             // Silently ignore this request as we cannot implement this using MemoryCache
-            return Array.Empty<BlobIdentifier>().ToAsyncEnumerable();
+            return Array.Empty<(BlobIdentifier, DateTime)>().ToAsyncEnumerable();
         }
         
         private string BuildKey(NamespaceId ns, BlobIdentifier blob)
diff --git a/Engine/Source/Programs/Horde/HordeStorage/Horde.Storage/Implementation/Blob/ScyllaBlobIndex.cs b/Engine/Source/Programs/Horde/HordeStorage/Horde.Storage/Implementation/Blob/ScyllaBlobIndex.cs
new file mode 100644
--- /dev/null
+++ b/Engine/Source/Programs/Horde/HordeStorage/Horde.Storage/Implementation/Blob/ScyllaBlobIndex.cs
@@ -0,0 +1,102 @@
+// Copyright Epic Games, Inc. All Rights Reserved.
+
+using System;
+using System.Collections.Generic;
+using System.Linq;
+using System.Threading.Tasks;
+using Cassandra;
+using Cassandra.Mapping;
+using Datadog.Trace;
+using Jupiter;
+using Jupiter.Implementation;
+using Microsoft.Extensions.Options;
+
+namespace Horde.Storage.Implementation.Blob;
+
+public class ScyllaBlobIndex : IBlobIndex
+{
+    private readonly IOptionsMonitor<JupiterSettings> _jupiterSettings;
+    private readonly ISession _session;
+    private readonly Mapper _mapper;
+
+    public ScyllaBlobIndex(IScyllaSessionManager scyllaSessionManager, IOptionsMonitor<JupiterSettings> jupiterSettings)
+    {
+        _jupiterSettings = jupiterSettings;
+        _session = scyllaSessionManager.GetSessionForReplicatedKeyspace();
+        _mapper = new Mapper(_session);
+
+        _session.Execute(new SimpleStatement(@"CREATE TABLE IF NOT EXISTS blob_index (
+            namespace text,
+            blob_id frozen<blob_identifier>,
+            regions set<text>,
+            PRIMARY KEY ((namespace, blob_id))
+        );"
+        ));
+    }
+
+    public async Task AddBlobToIndex(NamespaceId ns, BlobIdentifier id)
+    {
+        using Scope _ = Tracer.Instance.StartActive("scylla.insert_blob_index");
+
+        await _mapper.UpdateAsync<ScyllaBlobIndexTable>("SET regions = regions + ? WHERE namespace = ? AND blob_id = ?", new string[] {_jupiterSettings.CurrentValue.CurrentSite}, ns.ToString(), new ScyllaBlobIdentifier(id));
+    }
+
+    public async Task<IBlobIndex.BlobInfo?> GetBlobInfo(NamespaceId ns, BlobIdentifier id)
+    {
+        using Scope _ = Tracer.Instance.StartActive("scylla.fetch_blob_index");
+
+        ScyllaBlobIndexTable? blobIndex = await _mapper.FirstOrDefaultAsync<ScyllaBlobIndexTable>("WHERE namespace = ? AND blob_id = ?", ns.ToString(), new ScyllaBlobIdentifier(id));
+
+        if (blobIndex == null)
+            return null;
+
+        return new IBlobIndex.BlobInfo()
+        {
+            Regions = blobIndex.Regions
+        };
+    }
+
+    public async Task<bool> RemoveBlobFromIndex(NamespaceId ns, BlobIdentifier id)
+    {
+        // TODO: Should this only remove the current region, and the actual row isnt removed until all regions have been removed? Seems overly complicated
+        using Scope _ = Tracer.Instance.StartActive("scylla.remove_from_blob_index");
+        await _mapper.DeleteAsync<ScyllaBlobIndexTable>("WHERE namespace = ? AND blob_id = ?", ns.ToString(), new ScyllaBlobIdentifier(id));
+
+        return true;
+    }
+
+    public async Task<bool> BlobExistsInRegion(NamespaceId ns, BlobIdentifier blobIdentifier)
+    {
+        IBlobIndex.BlobInfo? blobInfo = await GetBlobInfo(ns, blobIdentifier);
+        return blobInfo?.Regions.Contains(_jupiterSettings.CurrentValue.CurrentSite) ?? false;
+    }
+}
+
+
+[Cassandra.Mapping.Attributes.Table("blob_index")]
+class ScyllaBlobIndexTable
+{
+    public ScyllaBlobIndexTable()
+    {
+        Namespace = null!;
+        BlobId = null!;
+        Regions = null!;
+    }
+
+    public ScyllaBlobIndexTable(string @namespace, BlobIdentifier blobId, HashSet<string> regions)
+    {
+        Namespace = @namespace;
+        BlobId =  new ScyllaBlobIdentifier(blobId);
+        Regions = regions;
+    }
+
+    [Cassandra.Mapping.Attributes.PartitionKey]
+    public string Namespace { get; set; }
+
+    [Cassandra.Mapping.Attributes.PartitionKey]
+    [Cassandra.Mapping.Attributes.Column("blob_id")]
+    public ScyllaBlobIdentifier BlobId { get;set; }
+
+    [Cassandra.Mapping.Attributes.Column("regions")]
+    public HashSet<string> Regions { get; set; }
+}
diff --git a/Engine/Source/Programs/Horde/HordeStorage/Horde.Storage/Implementation/ContentId/MemoryContentIdStore.cs b/Engine/Source/Programs/Horde/HordeStorage/Horde.Storage/Implementation/ContentId/MemoryContentIdStore.cs
--- a/Engine/Source/Programs/Horde/HordeStorage/Horde.Storage/Implementation/ContentId/MemoryContentIdStore.cs
+++ b/Engine/Source/Programs/Horde/HordeStorage/Horde.Storage/Implementation/ContentId/MemoryContentIdStore.cs
@@ -13,10 +13,10 @@ namespace Horde.Storage.Implementation
 {
     public class MemoryContentIdStore : IContentIdStore
     {
-        private readonly IBlobStore _blobStore;
+        private readonly IBlobService _blobStore;
         private readonly ConcurrentDictionary<NamespaceId, ConcurrentDictionary<BlobIdentifier, SortedList<int, BlobIdentifier[]>>> _contentIds = new ConcurrentDictionary<NamespaceId, ConcurrentDictionary<BlobIdentifier, SortedList<int, BlobIdentifier[]>>>();
         
-        public MemoryContentIdStore(IBlobStore blobStore)
+        public MemoryContentIdStore(IBlobService blobStore)
         {
             _blobStore = blobStore;
         }
diff --git a/Engine/Source/Programs/Horde/HordeStorage/Horde.Storage/Implementation/ContentId/ScyllaContentIdStore.cs b/Engine/Source/Programs/Horde/HordeStorage/Horde.Storage/Implementation/ContentId/ScyllaContentIdStore.cs
--- a/Engine/Source/Programs/Horde/HordeStorage/Horde.Storage/Implementation/ContentId/ScyllaContentIdStore.cs
+++ b/Engine/Source/Programs/Horde/HordeStorage/Horde.Storage/Implementation/ContentId/ScyllaContentIdStore.cs
@@ -14,10 +14,10 @@ namespace Horde.Storage.Implementation
     public class ScyllaContentIdStore : IContentIdStore
     {
         private readonly ISession _session;
-        private readonly IBlobStore _blobStore;
+        private readonly IBlobService _blobStore;
         private Mapper _mapper;
 
-        public ScyllaContentIdStore(IScyllaSessionManager scyllaSessionManager, IBlobStore blobStore)
+        public ScyllaContentIdStore(IScyllaSessionManager scyllaSessionManager, IBlobService blobStore)
         {
             _session = scyllaSessionManager.GetSessionForReplicatedKeyspace();
             _blobStore = blobStore;
diff --git a/Engine/Source/Programs/Horde/HordeStorage/Horde.Storage/Implementation/DDCRefService.cs b/Engine/Source/Programs/Horde/HordeStorage/Horde.Storage/Implementation/DDCRefService.cs
--- a/Engine/Source/Programs/Horde/HordeStorage/Horde.Storage/Implementation/DDCRefService.cs
+++ b/Engine/Source/Programs/Horde/HordeStorage/Horde.Storage/Implementation/DDCRefService.cs
@@ -57,7 +57,7 @@ namespace Horde.Storage.Implementation
 
     public class DDCRefService : IDDCRefService
     {
-        private readonly IBlobStore _blobStore;
+        private readonly IBlobService _blobStore;
         private readonly IRefsStore _refsStore;
         private readonly ITransactionLogWriter _transactionLog;
         private readonly ILastAccessTracker<RefRecord> _lastAccessTracker;
@@ -65,7 +65,7 @@ namespace Horde.Storage.Implementation
         private readonly HordeStorageSettings _settings;
         private readonly ILogger _logger = Log.ForContext<DDCRefService>();
 
-        public DDCRefService(IBlobStore blobStore, IRefsStore refsStore, ITransactionLogWriter transactionLog, ILastAccessTracker<RefRecord> lastAccessTracker,
+        public DDCRefService(IBlobService blobStore, IRefsStore refsStore, ITransactionLogWriter transactionLog, ILastAccessTracker<RefRecord> lastAccessTracker,
             IDiagnosticContext diagnosticContext, IOptionsMonitor<HordeStorageSettings> settings)
         {
             _blobStore = blobStore;
@@ -237,7 +237,7 @@ namespace Horde.Storage.Implementation
                 Task[] insertTasks = new Task[countOfBlobs];
                 for (int i = 0; i < countOfBlobs; i++)
                 {
-                    insertTasks[i] = _blobStore.PutObject(ns, slices[i], blobReferences[i]);
+                    insertTasks[i] = _blobStore.PutObject(ns, slices[i].ToArray(), blobReferences[i]);
                 }
 
                 blobInsertTask = Task.WhenAll(insertTasks);
diff --git a/Engine/Source/Programs/Horde/HordeStorage/Horde.Storage/Implementation/GC/BlobCleanupService.cs b/Engine/Source/Programs/Horde/HordeStorage/Horde.Storage/Implementation/GC/BlobCleanupService.cs
--- a/Engine/Source/Programs/Horde/HordeStorage/Horde.Storage/Implementation/GC/BlobCleanupService.cs
+++ b/Engine/Source/Programs/Horde/HordeStorage/Horde.Storage/Implementation/GC/BlobCleanupService.cs
@@ -7,6 +7,7 @@ using System.Threading.Tasks;
 using Jupiter;
 using Microsoft.Extensions.Options;
 using Datadog.Trace;
+using Microsoft.Extensions.DependencyInjection;
 using Serilog;
 
 namespace Horde.Storage.Implementation
@@ -27,11 +28,17 @@ namespace Horde.Storage.Implementation
             public List<IBlobCleanup> BlobCleanups { get; } = new List<IBlobCleanup>();
         }
 
-        public BlobCleanupService(IOptionsMonitor<GCSettings> settings) :
-            base(serviceName: nameof(RefCleanupService), settings.CurrentValue.BlobCleanupPollFrequency,
+        public BlobCleanupService(IServiceProvider provider, IOptionsMonitor<GCSettings> settings) :
+            base(serviceName: nameof(BlobCleanupService), settings.CurrentValue.BlobCleanupPollFrequency,
                 new BlobCleanupState())
         {
             _settings = settings;
+
+            if (settings.CurrentValue.CleanOldBlobs)
+            {
+                OrphanBlobCleanup orphanBlobCleanup = provider.GetService<OrphanBlobCleanup>()!;
+                RegisterCleanup(orphanBlobCleanup);
+            }
         }
 
         public void RegisterCleanup(IBlobCleanup cleanup)
diff --git a/Engine/Source/Programs/Horde/HordeStorage/Horde.Storage/Implementation/GC/OrphanBlobCleanup.cs b/Engine/Source/Programs/Horde/HordeStorage/Horde.Storage/Implementation/GC/OrphanBlobCleanup.cs
--- a/Engine/Source/Programs/Horde/HordeStorage/Horde.Storage/Implementation/GC/OrphanBlobCleanup.cs
+++ b/Engine/Source/Programs/Horde/HordeStorage/Horde.Storage/Implementation/GC/OrphanBlobCleanup.cs
@@ -34,15 +34,15 @@ namespace Horde.Storage.Implementation
 
     public class OrphanBlobCleanup : IBlobCleanup
     {
-        private readonly IBlobStore _blobStore;
+        private readonly IBlobService _blobService;
         private readonly ILeaderElection _leaderElection;
         private readonly IRestClient _client;
         private readonly ILogger _logger = Log.ForContext<OrphanBlobCleanup>();
 
         // ReSharper disable once UnusedMember.Global
-        public OrphanBlobCleanup(IBlobStore blobStore, ILeaderElection leaderElection, IOptionsMonitor<CallistoTransactionLogSettings> callistoSettings, IServiceCredentials serviceCredentials)
+        public OrphanBlobCleanup(IBlobService blobService, ILeaderElection leaderElection, IOptionsMonitor<CallistoTransactionLogSettings> callistoSettings, IServiceCredentials serviceCredentials)
         {
-            _blobStore = blobStore;
+            _blobService = blobService;
             _leaderElection = leaderElection;
 
             _client = new RestClient(callistoSettings.CurrentValue.ConnectionString)
@@ -51,9 +51,9 @@ namespace Horde.Storage.Implementation
             }.UseSerializer(() => new JsonNetSerializer());
         }
 
-        internal OrphanBlobCleanup(IBlobStore blobStore, ILeaderElection leaderElection, IRestClient callistoClient)
+        internal OrphanBlobCleanup(IBlobService blobService, ILeaderElection leaderElection, IRestClient callistoClient)
         {
-            _blobStore = blobStore;
+            _blobService = blobService;
             _leaderElection = leaderElection;
             _client = callistoClient;
         }
@@ -97,8 +97,11 @@ namespace Horde.Storage.Implementation
             {
                 // only consider blobs that have been around for 60 minutes
                 // this due to cases were blobs are uploaded first, and the GC roots specified after the fact for DDC keys
-                await foreach (BlobIdentifier blob in _blobStore.ListOldObjects(@namespace, DateTime.Now.AddMinutes(-60)).WithCancellation(cancellationToken))
+                await foreach ((BlobIdentifier blob, DateTime lastModified) in _blobService.ListObjects(@namespace).WithCancellation(cancellationToken))
                 {
+                    if (lastModified > DateTime.Now.AddMinutes(-60))
+                        continue;
+
                     // TODO: This could be turned into a parallel for as each iteration does not mutate the state
                     // there is no parallel for each for async enumerable yet
                     // we could manually build one, see [REDACTED_PROFILE_URL]/corefx/issues/34233 but lets wait until we know the time this takes is an issue
@@ -140,7 +143,7 @@ namespace Horde.Storage.Implementation
                 {
                     try
                     {
-                        await _blobStore.DeleteObject(ns, blob);
+                        await _blobService.DeleteObject(ns, blob);
                     }
                     catch (Exception e)
                     {
diff --git a/Engine/Source/Programs/Horde/HordeStorage/Horde.Storage/Implementation/IReferenceResolver.cs b/Engine/Source/Programs/Horde/HordeStorage/Horde.Storage/Implementation/IReferenceResolver.cs
--- a/Engine/Source/Programs/Horde/HordeStorage/Horde.Storage/Implementation/IReferenceResolver.cs
+++ b/Engine/Source/Programs/Horde/HordeStorage/Horde.Storage/Implementation/IReferenceResolver.cs
@@ -18,10 +18,10 @@ namespace Horde.Storage.Implementation
 
     public class ReferenceResolver : IReferenceResolver
     {
-        private readonly IBlobStore _blobStore;
+        private readonly IBlobService _blobStore;
         private readonly IContentIdStore _contentIdStore;
 
-        public ReferenceResolver(IBlobStore blobStore, IContentIdStore contentIdStore)
+        public ReferenceResolver(IBlobService blobStore, IContentIdStore contentIdStore)
         {
             _blobStore = blobStore;
             _contentIdStore = contentIdStore;
diff --git a/Engine/Source/Programs/Horde/HordeStorage/Horde.Storage/Implementation/ObjectService.cs b/Engine/Source/Programs/Horde/HordeStorage/Horde.Storage/Implementation/ObjectService.cs
--- a/Engine/Source/Programs/Horde/HordeStorage/Horde.Storage/Implementation/ObjectService.cs
+++ b/Engine/Source/Programs/Horde/HordeStorage/Horde.Storage/Implementation/ObjectService.cs
@@ -2,6 +2,7 @@
 
 using System;
 using System.Collections.Generic;
+using System.IO;
 using System.Linq;
 using System.Threading.Tasks;
 using async_enumerable_dotnet;
@@ -27,14 +28,14 @@ namespace Horde.Storage.Implementation
     public class ObjectService : IObjectService
     {
         private readonly IReferencesStore _referencesStore;
-        private readonly IBlobStore _blobStore;
+        private readonly IBlobService _blobService;
         private readonly IReferenceResolver _referenceResolver;
         private readonly IReplicationLog _replicationLog;
 
-        public ObjectService(IReferencesStore referencesStore, IBlobStore blobStore, IReferenceResolver referenceResolver, IReplicationLog replicationLog)
+        public ObjectService(IReferencesStore referencesStore, IBlobService blobService, IReferenceResolver referenceResolver, IReplicationLog replicationLog)
         {
             _referencesStore = referencesStore;
-            _blobStore = blobStore;
+            _blobService = blobService;
             _referenceResolver = referenceResolver;
             _replicationLog = replicationLog;
         }
@@ -50,7 +51,7 @@ namespace Horde.Storage.Implementation
             }
             else
             {
-                blobContents = await _blobStore.GetObject(ns, o.BlobIdentifier);
+                blobContents = await _blobService.GetObject(ns, o.BlobIdentifier);
             }
             return (o, blobContents);
         }
@@ -64,7 +65,7 @@ namespace Horde.Storage.Implementation
 
             Task objectStorePut = _referencesStore.Put(ns, bucket, key, blobHash, payload.Data, isFinalized);
 
-            Task<BlobIdentifier> blobStorePut = _blobStore.PutObject(ns, payload.Data, blobHash);
+            Task<BlobIdentifier> blobStorePut = _blobService.PutObject(ns, payload.Data, blobHash);
 
             BlobIdentifier[] missingReferences = Array.Empty<BlobIdentifier>();
             if (hasReferences)
@@ -73,7 +74,7 @@ namespace Horde.Storage.Implementation
                 try
                 {
                     IAsyncEnumerable<BlobIdentifier> references = _referenceResolver.ResolveReferences(ns, payload);
-                    missingReferences = await _blobStore.FilterOutKnownBlobs(ns, references);
+                    missingReferences = await _blobService.FilterOutKnownBlobs(ns, references);
                 }
                 catch (PartialReferenceResolveException e)
                 {
@@ -107,7 +108,7 @@ namespace Horde.Storage.Implementation
                 try
                 {
                     IAsyncEnumerable<BlobIdentifier> references = _referenceResolver.ResolveReferences(ns, payload);
-                    missingReferences = await _blobStore.FilterOutKnownBlobs(ns, references);
+                    missingReferences = await _blobService.FilterOutKnownBlobs(ns, references);
                 }
                 catch (PartialReferenceResolveException e)
                 {
diff --git a/Engine/Source/Programs/Horde/HordeStorage/Horde.Storage/Implementation/Replication/RefsReplicator.cs b/Engine/Source/Programs/Horde/HordeStorage/Horde.Storage/Implementation/Replication/RefsReplicator.cs
--- a/Engine/Source/Programs/Horde/HordeStorage/Horde.Storage/Implementation/Replication/RefsReplicator.cs
+++ b/Engine/Source/Programs/Horde/HordeStorage/Horde.Storage/Implementation/Replication/RefsReplicator.cs
@@ -14,6 +14,7 @@ using Dasync.Collections;
 using Datadog.Trace;
 using Horde.Storage.Controllers;
 using Horde.Storage.Implementation.TransactionLog;
+using Jupiter.Common.Implementation;
 using Jupiter.Implementation;
 using Microsoft.AspNetCore.Mvc;
 using Microsoft.Extensions.Options;
@@ -31,18 +32,18 @@ namespace Horde.Storage.Implementation
         private readonly AsyncManualResetEvent _replicationFinishedEvent = new AsyncManualResetEvent(true);
         private readonly CancellationTokenSource _replicationTokenSource = new CancellationTokenSource();
         private readonly ReplicatorSettings _replicatorSettings;
-        private readonly IBlobStore _blobStore;
+        private readonly IBlobService _blobService;
         private readonly IReplicationLog _replicationLog;
         private readonly IServiceCredentials _serviceCredentials;
         private readonly HttpClient _httpClient;
         private RefsState _refsState;
         private bool _replicationRunning;
 
-        public RefsReplicator(ReplicatorSettings replicatorSettings, IOptionsMonitor<ReplicationSettings> replicationSettings, IBlobStore blobStore, IHttpClientFactory httpClientFactory, IReplicationLog replicationLog, IServiceCredentials serviceCredentials)
+        public RefsReplicator(ReplicatorSettings replicatorSettings, IOptionsMonitor<ReplicationSettings> replicationSettings, IBlobService blobService, IHttpClientFactory httpClientFactory, IReplicationLog replicationLog, IServiceCredentials serviceCredentials)
         {
             _name = replicatorSettings.ReplicatorName;
             _replicatorSettings = replicatorSettings;
-            _blobStore = blobStore;
+            _blobService = blobService;
             _replicationLog = replicationLog;
             _serviceCredentials = serviceCredentials;
 
@@ -372,7 +373,7 @@ namespace Horde.Storage.Implementation
         {
             // We could potentially do this, but that could be dangerous if missing child references
             // check if this blob exists locally before replicating, if it does we assume we have all of its references already
-            //if (await _blobStore.Exists(ns, blob))
+            //if (await _blobService.Exists(ns, blob))
             //    return currentOffset;
 
             HttpRequestMessage referencesRequest = BuildHttpRequest(HttpMethod.Get, $"api/v1/objects/{ns}/{blob}/references");
@@ -388,27 +389,18 @@ namespace Horde.Storage.Implementation
             Array.Copy(refs.References, potentialBlobs, refs.References.Length);
             potentialBlobs[^1] = blob;
 
-            BlobIdentifier[] missingBlobs = await _blobStore.FilterOutKnownBlobs(ns, potentialBlobs);
+            BlobIdentifier[] missingBlobs = await _blobService.FilterOutKnownBlobs(ns, potentialBlobs);
             Task[] blobReplicationTasks = new Task[missingBlobs.Length];
             for (int i = 0; i < missingBlobs.Length; i++)
             {
                 BlobIdentifier blobToReplicate = missingBlobs[i];
                 blobReplicationTasks[i] = Task.Run(async () =>
                 {
-                    
                     HttpRequestMessage blobRequest = BuildHttpRequest(HttpMethod.Get, $"api/v1/blobs/{ns}/{blobToReplicate}");
                     HttpResponseMessage blobResponse = await _httpClient.SendAsync(blobRequest, cancellationToken);
-                    byte[] blobContents = await blobResponse.Content.ReadAsByteArrayAsync(cancellationToken);
-
-                    BlobIdentifier calculatedBlob = BlobIdentifier.FromBlob(blobContents);
-                    if (!blobToReplicate.Equals(calculatedBlob))
-                    {
-                        _logger.Warning("Mismatching blob when replicating {Blob}. Determined Hash was {Hash} size was {Size}", blobToReplicate, calculatedBlob, blobContents.LongLength);
-                        // TODO: attempt to replicate again
-                        return;
-                    }
-
-                    await _blobStore.PutObject(ns, blobContents, calculatedBlob);
+                    
+                    using FilesystemBufferedPayload payload = await FilesystemBufferedPayload.Create(await blobResponse.Content.ReadAsStreamAsync(cancellationToken));
+                    await _blobService.PutObject(ns, payload, blobToReplicate);
                 });
             }
 
diff --git a/Engine/Source/Programs/Horde/HordeStorage/Horde.Storage/Implementation/Replication/ReplicationSettings.cs b/Engine/Source/Programs/Horde/HordeStorage/Horde.Storage/Implementation/Replication/ReplicationSettings.cs
--- a/Engine/Source/Programs/Horde/HordeStorage/Horde.Storage/Implementation/Replication/ReplicationSettings.cs
+++ b/Engine/Source/Programs/Horde/HordeStorage/Horde.Storage/Implementation/Replication/ReplicationSettings.cs
@@ -25,13 +25,6 @@ namespace Horde.Storage
         /// </summary>
         public string OldStateRoot { get; set; } = "";
 
-        /// <summary>
-        ///  Name of the current site, has to be globally unique across all deployments, used to avoid replication loops
-        /// </summary>
-        [Required]
-        [Key]
-        public string CurrentSite { get; set; } = "";
-
         /// <summary>
         /// The frequency at which to poll for new replication events
         /// </summary>
diff --git a/Engine/Source/Programs/Horde/HordeStorage/Horde.Storage/Implementation/Replication/Replicator.cs b/Engine/Source/Programs/Horde/HordeStorage/Horde.Storage/Implementation/Replication/Replicator.cs
--- a/Engine/Source/Programs/Horde/HordeStorage/Horde.Storage/Implementation/Replication/Replicator.cs
+++ b/Engine/Source/Programs/Horde/HordeStorage/Horde.Storage/Implementation/Replication/Replicator.cs
@@ -8,6 +8,7 @@ using System.Net;
 using System.Threading;
 using System.Threading.Tasks;
 using Datadog.Trace;
+using Jupiter;
 using Jupiter.Implementation;
 using Microsoft.Extensions.Options;
 using Newtonsoft.Json;
@@ -23,12 +24,12 @@ namespace Horde.Storage.Implementation
     {
         private readonly ITransactionLogWriter _transactionLogWriter;
 
-        public ReplicatorV1(ReplicatorSettings replicatorSettings, IOptionsMonitor<ReplicationSettings> replicationSettings, IBlobStore blobStore, ITransactionLogWriter transactionLogWriter, IServiceCredentials serviceCredentials) : base(replicatorSettings, replicationSettings, blobStore, CreateRemoteClient(replicatorSettings, serviceCredentials))
+        public ReplicatorV1(ReplicatorSettings replicatorSettings, IOptionsMonitor<ReplicationSettings> replicationSettings, IOptionsMonitor<JupiterSettings> jupiterSettings, IBlobService blobService, ITransactionLogWriter transactionLogWriter, IServiceCredentials serviceCredentials) : base(replicatorSettings, replicationSettings, jupiterSettings, blobService, CreateRemoteClient(replicatorSettings, serviceCredentials))
         {
             _transactionLogWriter = transactionLogWriter;
         }
 
-        internal ReplicatorV1(ReplicatorSettings replicatorSettings, IOptionsMonitor<ReplicationSettings> replicationSettings, IBlobStore blobStore, ITransactionLogWriter transactionLogWriter, IRestClient remoteClient) : base(replicatorSettings, replicationSettings, blobStore, remoteClient)
+        internal ReplicatorV1(ReplicatorSettings replicatorSettings, IOptionsMonitor<ReplicationSettings> replicationSettings, IOptionsMonitor<JupiterSettings> jupiterSettings, IBlobService blobService, ITransactionLogWriter transactionLogWriter, IRestClient remoteClient) : base(replicatorSettings, replicationSettings, jupiterSettings, blobService, remoteClient)
         {
             _transactionLogWriter = transactionLogWriter;
         }
@@ -86,7 +87,7 @@ namespace Horde.Storage.Implementation
     public abstract class Replicator<T> : IReplicator
         where T: IReplicationEvent
     {
-        private readonly IBlobStore _blobStore;
+        private readonly IBlobService _blobService;
         protected readonly string _name;
         private readonly string _currentSite;
         protected readonly ReplicatorSettings _replicatorSettings;
@@ -114,11 +115,11 @@ namespace Horde.Storage.Implementation
             return remoteClient;
         }
 
-        protected Replicator(ReplicatorSettings replicatorSettings, IOptionsMonitor<ReplicationSettings> replicationSettings, IBlobStore blobStore, IRestClient remoteClient)
+        protected Replicator(ReplicatorSettings replicatorSettings, IOptionsMonitor<ReplicationSettings> replicationSettings, IOptionsMonitor<JupiterSettings> jupiterSettings, IBlobService blobService, IRestClient remoteClient)
         {
             _name = replicatorSettings.ReplicatorName;
-            _blobStore = blobStore;
-            _currentSite = replicationSettings.CurrentValue.CurrentSite;
+            _blobService = blobService;
+            _currentSite = jupiterSettings.CurrentValue.CurrentSite;
             _replicatorSettings = replicatorSettings;
 
             string stateFileName = $"{_name}.json";
@@ -332,7 +333,7 @@ namespace Horde.Storage.Implementation
                 blobReplicateTasks[index] = Task.Run(async () =>
                 {
                     // check if this blob exists locally before replicating
-                    if (await _blobStore.Exists(ns, blob))
+                    if (await _blobService.Exists(ns, blob))
                         return;
 
                     // attempt to replicate for a few tries with some delay in between
@@ -363,6 +364,12 @@ namespace Horde.Storage.Implementation
                             continue;
                         }
 
+                        if (response.StatusCode == HttpStatusCode.BadGateway)
+                        {
+                            _logger.Information("BadGateway while replicating {Blob}, retrying. Retry attempt {Attempts}", blob, attempts);
+                            continue;
+                        }
+
                         if (response.StatusCode == HttpStatusCode.BadRequest || response.StatusCode == HttpStatusCode.NotFound)
                         {
                             _logger.Warning("Remote blob {Blob} missing, unable to replicate.", blob);
@@ -408,8 +415,7 @@ namespace Horde.Storage.Implementation
                         _logger.Warning("Mismatching blob when replicating {Blob}. Determined Hash was {Hash} size was {Size}. Multiple attempts failed, giving up.", blob, calculatedBlob, rawContent.LongLength);
                         return; 
                     }
-
-                    await _blobStore.PutObject(ns, rawContent, calculatedBlob);
+                    await _blobService.PutObject(ns, rawContent, calculatedBlob);
                 }, replicationToken);
             }
 
diff --git a/Engine/Source/Programs/Horde/HordeStorage/Horde.Storage/Implementation/TransactionLog/ReplicationLogSnapshotBuilder.cs b/Engine/Source/Programs/Horde/HordeStorage/Horde.Storage/Implementation/TransactionLog/ReplicationLogSnapshotBuilder.cs
--- a/Engine/Source/Programs/Horde/HordeStorage/Horde.Storage/Implementation/TransactionLog/ReplicationLogSnapshotBuilder.cs
+++ b/Engine/Source/Programs/Horde/HordeStorage/Horde.Storage/Implementation/TransactionLog/ReplicationLogSnapshotBuilder.cs
@@ -11,13 +11,13 @@ namespace Horde.Storage.Implementation.TransactionLog
     public class ReplicationLogSnapshotBuilder
     {
         private readonly IReplicationLog _replicationLog;
-        private readonly IBlobStore _blobStore;
+        private readonly IBlobService _blobService;
         private readonly IReferencesStore _referencesStore;
 
-        public ReplicationLogSnapshotBuilder(IReplicationLog replicationLog, IBlobStore blobStore, IReferencesStore referencesStore)
+        public ReplicationLogSnapshotBuilder(IReplicationLog replicationLog, IBlobService blobService, IReferencesStore referencesStore)
         {
             _replicationLog = replicationLog;
-            _blobStore = blobStore;
+            _blobService = blobService;
             _referencesStore = referencesStore;
         }
 
@@ -36,7 +36,7 @@ namespace Horde.Storage.Implementation.TransactionLog
             if (snapshotInfo != null)
             {
                 // append to the previous snapshot if one is available
-                await using BlobContents blobContents = await _blobStore.GetObject(snapshotInfo.BlobNamespace, snapshotInfo.SnapshotBlob);
+                await using BlobContents blobContents = await _blobService.GetObject(snapshotInfo.BlobNamespace, snapshotInfo.SnapshotBlob);
                 if (cancellationToken.IsCancellationRequested)
                     throw new TaskCanceledException();
                 snapshot = await ReplicationLogSnapshot.DeserializeSnapshot(blobContents.Stream);
@@ -78,7 +78,7 @@ namespace Horde.Storage.Implementation.TransactionLog
                 byte[] cbObjectBytes = writer.Save();
                 BlobIdentifier cbBlobId = BlobIdentifier.FromBlob(cbObjectBytes);
                 await _referencesStore.Put(storeInNamespace, new BucketId("snapshot"), new KeyId(blobIdentifier.ToString()), cbBlobId, cbObjectBytes, true);
-                await _blobStore.PutObject(storeInNamespace, buf, blobIdentifier);
+                await _blobService.PutObject(storeInNamespace, buf, blobIdentifier);
 
                 if (cancellationToken.IsCancellationRequested)
                     throw new TaskCanceledException();
diff --git a/Engine/Source/Programs/Horde/HordeStorage/Horde.Storage/Settings.cs b/Engine/Source/Programs/Horde/HordeStorage/Horde.Storage/Settings.cs
--- a/Engine/Source/Programs/Horde/HordeStorage/Horde.Storage/Settings.cs
+++ b/Engine/Source/Programs/Horde/HordeStorage/Horde.Storage/Settings.cs
@@ -4,6 +4,7 @@ using System;
 using System.Collections.Generic;
 using System.ComponentModel.DataAnnotations;
 using System.Linq;
+using Horde.Storage.Implementation;
 using Microsoft.Extensions.Caching.Memory;
 
 namespace Horde.Storage
@@ -47,7 +48,8 @@ namespace Horde.Storage
             S3,
             Azure,
             FileSystem,
-            Memory
+            Memory,
+            MemoryBlobStore
         }
 
         public enum ReferencesDbImplementations
@@ -62,6 +64,12 @@ namespace Horde.Storage
             Scylla
         }
 
+        public enum BlobIndexImplementations
+        {
+            Memory,
+            Scylla
+        }
+
         public enum LeaderElectionImplementations
         {
             Static, 
@@ -86,6 +94,16 @@ namespace Horde.Storage
         [ValidStorageBackend]
         public string[]? StorageImplementations { get; set; }
 
+        public IEnumerable<StorageBackendImplementations> GetStorageImplementations()
+        {
+            foreach (string s in StorageImplementations ?? new [] {HordeStorageSettings.StorageBackendImplementations.Memory.ToString()})
+            {
+                HordeStorageSettings.StorageBackendImplementations impl = (HordeStorageSettings.StorageBackendImplementations)Enum.Parse(typeof(HordeStorageSettings.StorageBackendImplementations), s, ignoreCase: true);
+
+                yield return impl;
+            }
+        }
+
         [Required]
         public TransactionLogWriterImplementations TransactionLogWriterImplementation { get; set; } = TransactionLogWriterImplementations.Memory;
 
@@ -104,6 +122,7 @@ namespace Horde.Storage
 
         public LeaderElectionImplementations LeaderElectionImplementation { get; set; } = LeaderElectionImplementations.Static;
         public ContentIdStoreImplementations ContentIdStoreImplementation { get; set; } = ContentIdStoreImplementations.Memory;
+        public BlobIndexImplementations BlobIndexImplementation { get; set; } = BlobIndexImplementations.Memory;
 
         public int? MaxSingleBlobSize { get; set; } = null; // disable blob partitioning
 
@@ -211,6 +230,9 @@ namespace Horde.Storage
 
         // Options to disable setting of bucket access policies, useful for local testing as minio does not support them.
         public bool SetBucketPolicies { get; set; } = true;
+
+        public bool UseBlobIndexForExistsCheck { get; set; } = false;
+
     }
 
     public class GCSettings
diff --git a/Engine/Source/Programs/Horde/HordeStorage/Horde.Storage/Tests/Functional/Blob/BlobIndexTests.cs b/Engine/Source/Programs/Horde/HordeStorage/Horde.Storage/Tests/Functional/Blob/BlobIndexTests.cs
new file mode 100644
--- /dev/null
+++ b/Engine/Source/Programs/Horde/HordeStorage/Horde.Storage/Tests/Functional/Blob/BlobIndexTests.cs
@@ -0,0 +1,175 @@
+using System;
+using System.Collections.Generic;
+using System.IO;
+using System.Linq;
+using System.Net.Http;
+using System.Net.Http.Headers;
+using System.Net.Mime;
+using System.Text;
+using System.Threading.Tasks;
+using Horde.Storage.Implementation;
+using Horde.Storage.Implementation.Blob;
+using Jupiter;
+using Jupiter.Implementation;
+using Microsoft.AspNetCore.Hosting;
+using Microsoft.AspNetCore.TestHost;
+using Microsoft.Extensions.Configuration;
+using Microsoft.Extensions.DependencyInjection;
+using Microsoft.VisualStudio.TestTools.UnitTesting;
+using Newtonsoft.Json.Linq;
+using Serilog;
+using Serilog.Core;
+
+namespace Horde.Storage.FunctionalTests.Storage
+{
+    public abstract class BlobIndexTests
+    {
+        protected TestServer? _server;
+        protected HttpClient? _httpClient;
+
+        protected readonly NamespaceId TestNamespaceName = new NamespaceId("testbucket");
+        protected const string SmallFileContents = "Small file contents";
+        protected const string AnotherFileContents = "Another file with contents";
+        protected const string DeletableFileContents = "Delete Me";
+        protected const string OldFileContents = "a old blob used for testing cutoff filtering";
+
+        protected readonly BlobIdentifier _smallFileHash =
+            BlobIdentifier.FromBlob(Encoding.ASCII.GetBytes(SmallFileContents));
+
+        protected readonly BlobIdentifier _anotherFileHash =
+            BlobIdentifier.FromBlob(Encoding.ASCII.GetBytes(AnotherFileContents));
+
+        protected readonly BlobIdentifier _deleteFileHash =
+            BlobIdentifier.FromBlob(Encoding.ASCII.GetBytes(DeletableFileContents));
+
+        protected readonly BlobIdentifier _oldBlobFileHash =
+            BlobIdentifier.FromBlob(Encoding.ASCII.GetBytes(OldFileContents));
+
+        [TestInitialize]
+        public async Task Setup()
+        {
+            IConfigurationRoot configuration = new ConfigurationBuilder()
+                // we are not reading the base appSettings here as we want exact control over what runs in the tests
+                .AddJsonFile("appsettings.Testing.json", true)
+                .AddInMemoryCollection(GetSettings())
+                .AddEnvironmentVariables()
+                .Build();
+
+            Logger logger = new LoggerConfiguration()
+                .ReadFrom.Configuration(configuration)
+                .CreateLogger();
+
+            TestServer server = new TestServer(new WebHostBuilder()
+                .UseConfiguration(configuration)
+                .UseEnvironment("Testing")
+                .UseSerilog(logger)
+                .UseStartup<HordeStorageStartup>()
+            );
+            _httpClient = server.CreateClient();
+            _server = server;
+
+            // Seed storage
+            await Seed(_server.Services);
+        }
+
+        protected abstract IEnumerable<KeyValuePair<string, string>> GetSettings();
+
+        protected abstract Task Seed(IServiceProvider serverServices);
+        protected abstract Task Teardown();
+
+
+        [TestCleanup]
+        public async Task MyTeardown()
+        {
+            await Teardown();
+        }
+
+        
+        [TestMethod]
+        public async Task PutBlobToIndex()
+        {
+            byte[] payload = Encoding.ASCII.GetBytes("I am a blob with contents");
+            using ByteArrayContent requestContent = new ByteArrayContent(payload);
+            requestContent.Headers.ContentType = new MediaTypeHeaderValue(MediaTypeNames.Application.Octet);
+            BlobIdentifier contentHash = BlobIdentifier.FromBlob(payload);
+            HttpResponseMessage response = await _httpClient!.PutAsync(requestUri: $"api/v1/s/{TestNamespaceName}/{contentHash}", requestContent);
+            response.EnsureSuccessStatusCode();
+
+            IBlobIndex? index = _server!.Services.GetService<IBlobIndex>();
+            Assert.IsNotNull(index);
+            IBlobIndex.BlobInfo? blobInfo = await index.GetBlobInfo(TestNamespaceName, contentHash);
+
+            Assert.IsNotNull(blobInfo);
+            Assert.IsTrue(blobInfo.Regions.Contains("test"));
+        }
+
+        [TestMethod]
+        public async Task DeleteBlob()
+        {
+            // upload a blob
+            byte[] payload = Encoding.ASCII.GetBytes("I am a blob with contents");
+            BlobIdentifier contentHash = BlobIdentifier.FromBlob(payload);
+
+            {
+                using ByteArrayContent requestContent = new ByteArrayContent(payload);
+                requestContent.Headers.ContentType = new MediaTypeHeaderValue(MediaTypeNames.Application.Octet);
+                HttpResponseMessage response = await _httpClient!.PutAsync(requestUri: $"api/v1/s/{TestNamespaceName}/{contentHash}", requestContent);
+                response.EnsureSuccessStatusCode();
+            }
+            // verify its present in the blob index
+            IBlobIndex? index = _server!.Services.GetService<IBlobIndex>();
+            Assert.IsNotNull(index);
+            IBlobIndex.BlobInfo? blobInfo = await index.GetBlobInfo(TestNamespaceName, contentHash);
+            Assert.IsNotNull(blobInfo);
+
+            // delete the blob
+            {
+                HttpResponseMessage response = await _httpClient!.DeleteAsync(requestUri: $"api/v1/s/{TestNamespaceName}/{contentHash}");
+                response.EnsureSuccessStatusCode();
+            }
+
+            IBlobIndex.BlobInfo? deletedBlobInfo = await index.GetBlobInfo(TestNamespaceName, contentHash);
+            Assert.IsNull(deletedBlobInfo);
+        }
+
+     
+    }
+
+    [TestClass]
+    public class MemoryBlobIndexTests : BlobIndexTests
+    {
+        protected override IEnumerable<KeyValuePair<string, string>> GetSettings()
+        {
+            return new[] { new KeyValuePair<string, string>("Horde.Storage:BlobIndexImplementation", HordeStorageSettings.BlobIndexImplementations.Memory.ToString()) };
+        }
+
+        protected override Task Seed(IServiceProvider serverServices)
+        {
+            return Task.CompletedTask;
+        }
+
+        protected override Task Teardown()
+        {
+            return Task.CompletedTask;
+        }
+    }
+
+    [TestClass]
+    public class ScyllaBlobIndexTests : BlobIndexTests
+    {
+        protected override IEnumerable<KeyValuePair<string, string>> GetSettings()
+        {
+            return new[] { new KeyValuePair<string, string>("Horde.Storage:BlobIndexImplementation", HordeStorageSettings.BlobIndexImplementations.Scylla.ToString()) };
+        }
+
+        protected override Task Seed(IServiceProvider serverServices)
+        {
+            return Task.CompletedTask;
+        }
+
+        protected override Task Teardown()
+        {
+            return Task.CompletedTask;
+        }
+    }
+}
diff --git a/Engine/Source/Programs/Horde/HordeStorage/Horde.Storage/Tests/Functional/ContentId/CompressedBlobTests.cs b/Engine/Source/Programs/Horde/HordeStorage/Horde.Storage/Tests/Functional/ContentId/CompressedBlobTests.cs
--- a/Engine/Source/Programs/Horde/HordeStorage/Horde.Storage/Tests/Functional/ContentId/CompressedBlobTests.cs
+++ b/Engine/Source/Programs/Horde/HordeStorage/Horde.Storage/Tests/Functional/ContentId/CompressedBlobTests.cs
@@ -163,5 +163,22 @@ namespace Horde.Storage.FunctionalTests.CompressedBlobs
                 Assert.AreEqual(HttpStatusCode.NotFound, result.StatusCode);
             }
         }
+
+        [TestMethod]
+        public async Task PutWrongIdentifier()
+        {
+            byte[] texturePayload = await File.ReadAllBytesAsync("ContentId/Payloads/UncompressedTexture_CAS_dea81b6c3b565bb5089695377c98ce0f1c13b0c3.udd");
+            BlobIdentifier compressedPayloadIdentifier = BlobIdentifier.FromBlob(texturePayload);
+            BlobIdentifier uncompressedPayloadIdentifier = new BlobIdentifier("DEA81B6C3B565BB5089695377C98CE0F1C13B0C3");
+
+            {
+                ByteArrayContent content = new ByteArrayContent(texturePayload);
+                content.Headers.ContentType = new MediaTypeHeaderValue(CustomMediaTypeNames.UnrealCompressedBuffer);
+                // we purposefully use the compressed identifier here which is not what is expected
+                HttpResponseMessage result = await _httpClient!.PutAsync($"api/v1/compressed-blobs/{TestNamespace}/{compressedPayloadIdentifier}", content);
+
+                Assert.AreEqual(HttpStatusCode.BadRequest, result.StatusCode);
+            }
+        }
     }
 }
diff --git a/Engine/Source/Programs/Horde/HordeStorage/Horde.Storage/Tests/Functional/GC/GCBlobTests.cs b/Engine/Source/Programs/Horde/HordeStorage/Horde.Storage/Tests/Functional/GC/GCBlobTests.cs
--- a/Engine/Source/Programs/Horde/HordeStorage/Horde.Storage/Tests/Functional/GC/GCBlobTests.cs
+++ b/Engine/Source/Programs/Horde/HordeStorage/Horde.Storage/Tests/Functional/GC/GCBlobTests.cs
@@ -36,7 +36,7 @@ namespace Horde.Storage.FunctionalTests.GC
         private readonly BlobIdentifier object5id = new BlobIdentifier("5555555555555555555555555555555555555555");
         private readonly BlobIdentifier object6id = new BlobIdentifier("6666666666666666666666666666666666666666");
         private Mock<IRestClient> _callistoBlobMock = null!;
-        private IBlobStore? _blobStore;
+        private IBlobService? _blobService;
 
         private readonly NamespaceId TestNamespace = new NamespaceId("test-namespace");
 
@@ -47,6 +47,7 @@ namespace Horde.Storage.FunctionalTests.GC
                 // we are not reading the base appSettings here as we want exact control over what runs in the tests
                 .AddJsonFile("appsettings.Testing.json", false)
                 .AddEnvironmentVariables()
+                .AddInMemoryCollection(new List<KeyValuePair<string, string>>() { new KeyValuePair<string, string>("Horde.Storage:StorageImplementations:0", "MemoryBlobStore")})
                 .Build();
 
             Logger logger = new LoggerConfiguration()
@@ -72,21 +73,19 @@ namespace Horde.Storage.FunctionalTests.GC
                 {
                     collection.AddSingleton<IBlobCleanup>(provider =>
                     {
-                        IBlobStore blobStore = provider.GetService<IBlobStore>()!;
-                        return new OrphanBlobCleanup(blobStore, new LeaderElectionStub(true), _callistoBlobMock.Object);
+                        IBlobService blobService = provider.GetService<IBlobService>()!;
+                        return new OrphanBlobCleanup(blobService!, new LeaderElectionStub(true), _callistoBlobMock.Object);
                     });
 
-                    // make sure we have a memory blob store as we do some funky calling to it below
-                    collection.AddSingleton<IBlobStore>(provider => ActivatorUtilities.CreateInstance<MemoryBlobStore>(provider));
                 })
             );
 
             _httpClient = server.CreateClient();
             _server = server;
 
-            _blobStore = server.Services.GetService<IBlobStore>()!;
+            _blobService = server.Services.GetService<IBlobService>()!;
 
-            MemoryBlobStore memoryBlobStore = (MemoryBlobStore) _blobStore;
+            MemoryBlobStore memoryBlobStore = (MemoryBlobStore) ((BlobService)_blobService).BlobStore.First();
             byte[] emptyContents = new byte[0];
             await memoryBlobStore.PutObject(TestNamespace, emptyContents, object0id);
             await memoryBlobStore.PutObject(TestNamespace, emptyContents, object1id);// this is not in callisto
@@ -157,7 +156,7 @@ namespace Horde.Storage.FunctionalTests.GC
         [TestMethod]
         public async Task RunBlobCleanup()
         {
-            OrphanBlobCleanup cleanup = new OrphanBlobCleanup(_blobStore!, new LeaderElectionStub(true), _callistoBlobMock.Object);
+            OrphanBlobCleanup cleanup = new OrphanBlobCleanup(_blobService!, new LeaderElectionStub(true), _callistoBlobMock.Object);
             List<NamespaceId> namespaces = await cleanup.ListNamespaces().ToListAsync();
             Assert.AreEqual(1, namespaces.Count);
 
diff --git a/Engine/Source/Programs/Horde/HordeStorage/Horde.Storage/Tests/Functional/Objects/RefReplicatorTests.cs b/Engine/Source/Programs/Horde/HordeStorage/Horde.Storage/Tests/Functional/Objects/RefReplicatorTests.cs
--- a/Engine/Source/Programs/Horde/HordeStorage/Horde.Storage/Tests/Functional/Objects/RefReplicatorTests.cs
+++ b/Engine/Source/Programs/Horde/HordeStorage/Horde.Storage/Tests/Functional/Objects/RefReplicatorTests.cs
@@ -34,7 +34,7 @@ namespace Horde.Storage.FunctionalTests.Replication
     {
         private static TestServer? _server;
         private static HttpClient? _httpClient;
-        protected IBlobStore _blobStore = null!;
+        protected IBlobService _blobStore = null!;
         protected IReferencesStore _referencesStore = null!;
         private IReplicationLog _replicationLog = null!;
 
@@ -65,7 +65,7 @@ namespace Horde.Storage.FunctionalTests.Replication
             _httpClient = server.CreateClient();
             _server = server;
 
-            _blobStore = _server.Services.GetService<IBlobStore>()!;
+            _blobStore = _server.Services.GetService<IBlobService>()!;
             _referencesStore = _server.Services.GetService<IReferencesStore>()!;
             _replicationLog = _server.Services.GetService<IReplicationLog>()!;
 
diff --git a/Engine/Source/Programs/Horde/HordeStorage/Horde.Storage/Tests/Functional/Objects/ReferencesTests.cs b/Engine/Source/Programs/Horde/HordeStorage/Horde.Storage/Tests/Functional/Objects/ReferencesTests.cs
--- a/Engine/Source/Programs/Horde/HordeStorage/Horde.Storage/Tests/Functional/Objects/ReferencesTests.cs
+++ b/Engine/Source/Programs/Horde/HordeStorage/Horde.Storage/Tests/Functional/Objects/ReferencesTests.cs
@@ -126,7 +126,7 @@ namespace Horde.Storage.FunctionalTests.References
     {
         private static TestServer? _server;
         private static HttpClient? _httpClient;
-        protected IBlobStore _blobStore = null!;
+        protected IBlobService _blobService = null!;
         protected IReferencesStore ReferencesStore = null!;
 
         protected readonly NamespaceId TestNamespace = new NamespaceId("test-namespace");
@@ -155,7 +155,7 @@ namespace Horde.Storage.FunctionalTests.References
             _httpClient = server.CreateClient();
             _server = server;
 
-            _blobStore = _server.Services.GetService<IBlobStore>()!;
+            _blobService = _server.Services.GetService<IBlobService>()!;
             ReferencesStore = _server.Services.GetService<IReferencesStore>()!;
             await SeedDb(server.Services);
         }
@@ -568,11 +568,11 @@ namespace Horde.Storage.FunctionalTests.References
             BlobIdentifier objectHash = BlobIdentifier.FromBlob(data);
 
             HttpContent requestContent = new ByteArrayContent(data);
-            requestContent.Headers.ContentType = new MediaTypeHeaderValue(CustomMediaTypeNames.UnrealCompactBinary);
+            requestContent.Headers.ContentType = new MediaTypeHeaderValue(MediaTypeNames.Application.Octet);
             requestContent.Headers.Add(CommonHeaders.HashHeaderName, objectHash.ToString());
 
             {
-                HttpResponseMessage result = await _httpClient!.PutAsync(requestUri: $"api/v1/refs/{TestNamespace}/bucket/newObject.raw", requestContent);
+                HttpResponseMessage result = await _httpClient!.PutAsync(requestUri: $"api/v1/refs/{TestNamespace}/bucket/newObject", requestContent);
                 result.EnsureSuccessStatusCode();
             }
 
@@ -594,12 +594,12 @@ namespace Horde.Storage.FunctionalTests.References
             string blobContents = "This is a string that is referenced as a blob";
             byte[] blobData = Encoding.ASCII.GetBytes(blobContents);
             BlobIdentifier blobHash = BlobIdentifier.FromBlob(blobData);
-            await _blobStore.PutObject(TestNamespace, blobData, blobHash);
+            await _blobService.PutObject(TestNamespace, blobData, blobHash);
 
             string blobContentsChild = "This string is also referenced as a blob but from a child object";
             byte[] dataChild = Encoding.ASCII.GetBytes(blobContentsChild);
             BlobIdentifier blobHashChild = BlobIdentifier.FromBlob(dataChild);
-            await _blobStore.PutObject(TestNamespace, dataChild, blobHashChild);
+            await _blobService.PutObject(TestNamespace, dataChild, blobHashChild);
 
             CompactBinaryWriter writerChild = new CompactBinaryWriter();
             writerChild.BeginObject();
@@ -608,7 +608,7 @@ namespace Horde.Storage.FunctionalTests.References
 
             byte[] childDataObject = writerChild.Save();
             BlobIdentifier childDataObjectHash = BlobIdentifier.FromBlob(childDataObject);
-            await _blobStore.PutObject(TestNamespace, childDataObject, childDataObjectHash);
+            await _blobService.PutObject(TestNamespace, childDataObject, childDataObjectHash);
 
             CompactBinaryWriter writerParent = new CompactBinaryWriter();
             writerParent.BeginObject();
@@ -776,7 +776,7 @@ namespace Horde.Storage.FunctionalTests.References
 
             byte[] childDataObject = writerChild.Save();
             BlobIdentifier childDataObjectHash = BlobIdentifier.FromBlob(childDataObject);
-            await _blobStore.PutObject(TestNamespace, childDataObject, childDataObjectHash);
+            await _blobService.PutObject(TestNamespace, childDataObject, childDataObjectHash);
 
             CompactBinaryWriter writerParent = new CompactBinaryWriter();
             writerParent.BeginObject();
@@ -860,7 +860,7 @@ namespace Horde.Storage.FunctionalTests.References
 
             byte[] childDataObject = writerChild.Save();
             BlobIdentifier childDataObjectHash = BlobIdentifier.FromBlob(childDataObject);
-            await _blobStore.PutObject(TestNamespace, childDataObject, childDataObjectHash);
+            await _blobService.PutObject(TestNamespace, childDataObject, childDataObjectHash);
 
             CompactBinaryWriter writerParent = new CompactBinaryWriter();
             writerParent.BeginObject();
@@ -908,8 +908,8 @@ namespace Horde.Storage.FunctionalTests.References
 
             // upload missing pieces
             {
-                await _blobStore.PutObject(TestNamespace, blobData, blobHash);
-                await _blobStore.PutObject(TestNamespace, dataChild, blobHashChild);
+                await _blobService.PutObject(TestNamespace, blobData, blobHash);
+                await _blobService.PutObject(TestNamespace, dataChild, blobHashChild);
             }
 
             // finalize the object as no pieces is now missing
@@ -986,7 +986,7 @@ namespace Horde.Storage.FunctionalTests.References
 
             {
                 // delete the blob 
-                await _blobStore.DeleteObject(TestNamespace, blobHash);
+                await _blobService.DeleteObject(TestNamespace, blobHash);
             }
 
             {
diff --git a/Engine/Source/Programs/Horde/HordeStorage/Horde.Storage/Tests/Functional/Objects/ReplicationTests.cs b/Engine/Source/Programs/Horde/HordeStorage/Horde.Storage/Tests/Functional/Objects/ReplicationTests.cs
--- a/Engine/Source/Programs/Horde/HordeStorage/Horde.Storage/Tests/Functional/Objects/ReplicationTests.cs
+++ b/Engine/Source/Programs/Horde/HordeStorage/Horde.Storage/Tests/Functional/Objects/ReplicationTests.cs
@@ -106,7 +106,7 @@ namespace Horde.Storage.FunctionalTests.References
     {
         private static TestServer? _server;
         private static HttpClient? _httpClient;
-        protected IBlobStore _blobStore = null!;
+        protected IBlobService _blobStore = null!;
         protected IReferencesStore _referencesStore = null!;
         private IReplicationLog _replicationLog = null!;
 
@@ -138,7 +138,7 @@ namespace Horde.Storage.FunctionalTests.References
             _httpClient = server.CreateClient();
             _server = server;
 
-            _blobStore = _server.Services.GetService<IBlobStore>()!;
+            _blobStore = _server.Services.GetService<IBlobService>()!;
             _referencesStore = _server.Services.GetService<IReferencesStore>()!;
             _replicationLog = _server.Services.GetService<IReplicationLog>()!;
 
diff --git a/Engine/Source/Programs/Horde/HordeStorage/Horde.Storage/Tests/Functional/RefTests.cs b/Engine/Source/Programs/Horde/HordeStorage/Horde.Storage/Tests/Functional/RefTests.cs
--- a/Engine/Source/Programs/Horde/HordeStorage/Horde.Storage/Tests/Functional/RefTests.cs
+++ b/Engine/Source/Programs/Horde/HordeStorage/Horde.Storage/Tests/Functional/RefTests.cs
@@ -60,7 +60,7 @@ namespace Horde.Storage.FunctionalTests.Ref
             Assert.IsTrue(RefStoreIs(refStore, typeof(MongoRefsStore)));
 
             _refsStore = new MongoRefsStore(provider.GetService<IOptionsMonitor<MongoSettings>>()!);
-            IBlobStore blobStore = (IBlobStore)provider.GetService(typeof(IBlobStore))!;
+            IBlobService blobStore = (IBlobService)provider.GetService(typeof(IBlobService))!;
 
             await SeedRefTestData(_refsStore, blobStore);
 
@@ -107,9 +107,9 @@ namespace Horde.Storage.FunctionalTests.Ref
             Assert.IsTrue(RefStoreIs(cacheStore, typeof(DynamoDbRefsStore)));
 
             _refStore = new DynamoDbRefsStore(dynamoSettings, provider.GetService<IAmazonDynamoDB>()!, provider.GetService<DynamoNamespaceStore>()!);
-            IBlobStore blobStore = (IBlobStore)provider.GetService(typeof(IBlobStore))!;
+            IBlobService blobService = provider.GetService<IBlobService>()!;
 
-            await SeedRefTestData(_refStore, blobStore);
+            await SeedRefTestData(_refStore, blobService);
         }
 
         protected override async Task TeardownDb(IServiceProvider provider)
@@ -186,7 +186,7 @@ namespace Horde.Storage.FunctionalTests.Ref
             await SeedDb(server.Services);
         }
 
-        protected async Task SeedRefTestData(IRefsStore refsStore, IBlobStore blobStore)
+        protected async Task SeedRefTestData(IRefsStore refsStore, IBlobService blobStore)
         {
             // add the objects as if they were last accessed a day ago
             DateTime lastAccessTime = DateTime.Now.AddDays(-1);
@@ -609,7 +609,7 @@ namespace Horde.Storage.FunctionalTests.Ref
             };
             foreach(OldRecord record in oldRecords)
             {
-                Assert.IsTrue(validRefs.Contains(record.RefName));
+                Assert.IsTrue(validRefs.Contains(record.RefName), $"ValidRefs did not contain ref {record.RefName}");
             }
 
             // we should have our 4 pre-seeded objects
diff --git a/Engine/Source/Programs/Horde/HordeStorage/Horde.Storage/Tests/Functional/StorageTests.cs b/Engine/Source/Programs/Horde/HordeStorage/Horde.Storage/Tests/Functional/StorageTests.cs
--- a/Engine/Source/Programs/Horde/HordeStorage/Horde.Storage/Tests/Functional/StorageTests.cs
+++ b/Engine/Source/Programs/Horde/HordeStorage/Horde.Storage/Tests/Functional/StorageTests.cs
@@ -200,13 +200,14 @@ namespace Horde.Storage.FunctionalTests.Storage
         [TestMethod]
         public async Task ListOldBlobs()
         {
-            IBlobStore blobStore = _server!.Services.GetService<IBlobStore>()!;
+            FileSystemStore? fsStore = _server!.Services.GetService<FileSystemStore>();
+            Assert.IsNotNull(fsStore);
 
             // we fetch all objects that are more then a day old
             // as the old blob is set to be a week old this should be the only object returned
             DateTime cutoff = DateTime.Now.AddDays(-1);
             {
-                BlobIdentifier[] blobs = await blobStore.ListOldObjects(TestNamespaceName, cutoff).ToArrayAsync();
+                BlobIdentifier[] blobs = await fsStore.ListObjects(TestNamespaceName).Where(tuple => tuple.Item2 < cutoff).Select(tuple => tuple.Item1).ToArrayAsync();
                 Assert.AreEqual(_oldBlobFileHash, blobs[0]);
             }
         }
@@ -214,8 +215,8 @@ namespace Horde.Storage.FunctionalTests.Storage
         [TestMethod]
         public async Task ListNamespaces()
         {
-            IBlobStore blobStore = _server!.Services.GetService<IBlobStore>()!;
-            FileSystemStore fsStore = (blobStore as FileSystemStore)!;
+            FileSystemStore? fsStore = _server!.Services.GetService<FileSystemStore>();
+            Assert.IsNotNull(fsStore);
 
             List<NamespaceId> namespaces = await fsStore.ListNamespaces().ToListAsync();
             Assert.AreEqual(1, namespaces.Count);
@@ -235,8 +236,8 @@ namespace Horde.Storage.FunctionalTests.Storage
             
             FilesystemSettings fsSettings = _server!.Services.GetService<IOptionsMonitor<FilesystemSettings>>()!.CurrentValue;
             fsSettings.MaxSizeBytes = 600;
-            IBlobStore blobStore = _server!.Services.GetService<IBlobStore>()!;
-            FileSystemStore fsStore = (blobStore as FileSystemStore)!;
+            FileSystemStore? fsStore = _server!.Services.GetService<FileSystemStore>();
+            Assert.IsNotNull(fsStore);
 
             CancellationTokenSource cts = new CancellationTokenSource();
             Assert.IsTrue((await fsStore.CleanupInternal(cts.Token, batchSize: 2)).Count == 0); // No garbage to collect, should return false
@@ -263,17 +264,17 @@ namespace Horde.Storage.FunctionalTests.Storage
         }
         
         [TestMethod]
-        public async Task GetLeastRecentlyAccessedObjects()
+        public void GetLeastRecentlyAccessedObjects()
         {
             FileInfo[] fooFiles = CreateFilesInNamespace(FooNamespace, 10);
             CreateFilesInNamespace(new NamespaceId("bar"), 10);
 
-            IBlobStore blobStore = _server!.Services.GetService<IBlobStore>()!;
-            FileSystemStore fsStore = (blobStore as FileSystemStore)!;
+            FileSystemStore? fsStore = _server!.Services.GetService<FileSystemStore>();
+            Assert.IsNotNull(fsStore);
 
-            Assert.AreEqual(10, (await fsStore.GetLeastRecentlyAccessedObjects(FooNamespace)).Length);
+            Assert.AreEqual(10, (fsStore.GetLeastRecentlyAccessedObjects(FooNamespace)).ToArray().Length);
 
-            FileInfo[] results = await fsStore.GetLeastRecentlyAccessedObjects(FooNamespace, 3);
+            FileInfo[] results = fsStore.GetLeastRecentlyAccessedObjects(FooNamespace, 3).ToArray();
             Assert.AreEqual(3, results.Length);
             Assert.AreEqual(fooFiles[7].LastAccessTime, results[2].LastAccessTime);
             Assert.AreEqual(fooFiles[8].LastAccessTime, results[1].LastAccessTime);
@@ -283,8 +284,8 @@ namespace Horde.Storage.FunctionalTests.Storage
         [TestMethod]
         public async Task CalculateUsedDiskSpace()
         {
-            IBlobStore blobStore = _server!.Services.GetService<IBlobStore>()!;
-            FileSystemStore fsStore = (blobStore as FileSystemStore)!;
+            FileSystemStore? fsStore = _server!.Services.GetService<FileSystemStore>();
+            Assert.IsNotNull(fsStore);
             await fsStore.PutObject(FooNamespace, Encoding.ASCII.GetBytes(SmallFileContents), _smallFileHash);
             await fsStore.PutObject(FooNamespace, Encoding.ASCII.GetBytes(AnotherFileContents), _anotherFileHash);
             
@@ -435,11 +436,9 @@ namespace Horde.Storage.FunctionalTests.Storage
                 _oldBlobFileHash
             };
 
-            IBlobStore blobStore = _server!.Services.GetService<IBlobStore>()!;
-            // fetch all objects so set cutoff to a time in the future
-            DateTime cutoff = DateTime.Now.AddMinutes(5);
+            IBlobService blobService = _server!.Services.GetService<IBlobService>()!;
 
-            BlobIdentifier[] oldObjects = await blobStore.ListOldObjects(TestNamespaceName, cutoff).ToArrayAsync();
+            BlobIdentifier[] oldObjects = await blobService.ListObjects(TestNamespaceName).Select(tuple => tuple.Item1).ToArrayAsync();
 
             Assert.AreEqual(4, oldObjects.Length);
             Assert.IsTrue(validBlobHashes.Contains(oldObjects[0]));
@@ -792,6 +791,58 @@ namespace Horde.Storage.FunctionalTests.Storage
                     tempOutputFile.Delete();
             }
         }
+
+        /// <summary>
+        /// Uploads a compact binary with a lot of attachments that are missing
+        /// </summary>
+        /// <returns></returns>
+        /*[TestMethod]
+        public async Task LargeCBUpload()
+        {
+            // we submit a blob so large that it can not fit using the memory blob store
+            IBlobStore? blobStore = _server?.Services.GetService<IBlobStore>();
+            Assert.IsFalse(blobStore is MemoryBlobStore);
+
+            if (blobStore is AzureBlobStore)
+                Assert.Inconclusive("Azure blob store gets internal server errors when receiving large blobs");
+            CompactBinaryWriter writer = new CompactBinaryWriter();
+            writer.BeginObject();
+            const int countOfFields = 100_000;
+            for (int i = 0; i < countOfFields; i++)
+            {
+                string s = i.ToString("x40");
+                writer.AddBinaryAttachment(new BlobIdentifier(s));
+            }
+            writer.EndObject();
+
+            byte[] objectData = writer.Save();
+            string path = Path.GetTempFileName();
+            File.WriteAllBytes(path, objectData);
+
+            Console.WriteLine("CB Object output to: " + path);
+            BlobIdentifier objectHash = BlobIdentifier.FromBlob(objectData);
+
+            HttpContent requestContent = new ByteArrayContent(objectData);
+            requestContent.Headers.ContentType = new MediaTypeHeaderValue(CustomMediaTypeNames.UnrealCompactBinary);
+            requestContent.Headers.Add(CommonHeaders.HashHeaderName, objectHash.ToString());
+
+            HttpResponseMessage result = await _httpClient!.PutAsync(requestUri: $"api/v1/refs/{TestNamespaceName}/bucket/newLargeRef.uecb", requestContent);
+            result.EnsureSuccessStatusCode();
+
+            {
+                Assert.AreEqual(CustomMediaTypeNames.UnrealCompactBinary, result!.Content.Headers.ContentType!.MediaType);
+                // check that no blobs are missing
+                await using MemoryStream ms = new MemoryStream();
+                await result.Content.CopyToAsync(ms);
+                byte[] roundTrippedBuffer = ms.ToArray();
+                ReadOnlyMemory<byte> localMemory = new ReadOnlyMemory<byte>(roundTrippedBuffer);
+                CompactBinaryObject cb = CompactBinaryObject.Load(ref localMemory);
+                CompactBinaryField? needsField = cb["needs"];
+                Assert.IsNotNull(needsField);
+                List<BlobIdentifier?> missingBlobs = needsField!.AsArray().Select(field => field.AsHash()).ToList();
+                Assert.AreEqual(countOfFields, missingBlobs.Count);
+            }
+        }*/
     }
 
 
diff --git a/Engine/Source/Programs/Horde/HordeStorage/Horde.Storage/Tests/UnitTests/BlobCleanupServiceTest.cs b/Engine/Source/Programs/Horde/HordeStorage/Horde.Storage/Tests/UnitTests/BlobCleanupServiceTest.cs
--- a/Engine/Source/Programs/Horde/HordeStorage/Horde.Storage/Tests/UnitTests/BlobCleanupServiceTest.cs
+++ b/Engine/Source/Programs/Horde/HordeStorage/Horde.Storage/Tests/UnitTests/BlobCleanupServiceTest.cs
@@ -1,5 +1,6 @@
 ﻿// Copyright Epic Games, Inc. All Rights Reserved.
 
+using System;
 using System.Collections.Generic;
 using System.Threading;
 using System.Threading.Tasks;
@@ -16,9 +17,12 @@ namespace Horde.Storage.UnitTests
         [TestMethod]
         public async Task CleanupOnPoll()
         {
-            GCSettings gcSettings = new GCSettings();
+            GCSettings gcSettings = new GCSettings()
+            {
+                CleanOldBlobs = false
+            };
             TestOptionsMonitor<GCSettings> gcSettingsMon = new TestOptionsMonitor<GCSettings>(gcSettings);
-            BlobCleanupService blobCleanupService = new BlobCleanupService(gcSettingsMon);
+            BlobCleanupService blobCleanupService = new BlobCleanupService(Mock.Of<IServiceProvider>(), gcSettingsMon);
 
             Mock<IBlobCleanup> store1 = new Mock<IBlobCleanup>();
             Mock<IBlobCleanup> store2 = new Mock<IBlobCleanup>();
diff --git a/Engine/Source/Programs/Horde/HordeStorage/Horde.Storage/Tests/UnitTests/HierarchicalBlobStorageTest.cs b/Engine/Source/Programs/Horde/HordeStorage/Horde.Storage/Tests/UnitTests/HierarchicalBlobStorageTest.cs
--- a/Engine/Source/Programs/Horde/HordeStorage/Horde.Storage/Tests/UnitTests/HierarchicalBlobStorageTest.cs
+++ b/Engine/Source/Programs/Horde/HordeStorage/Horde.Storage/Tests/UnitTests/HierarchicalBlobStorageTest.cs
@@ -8,8 +8,11 @@ using System.Text;
 using System.Threading.Tasks;
 using Dasync.Collections;
 using Horde.Storage.Implementation;
+using Horde.Storage.Implementation.Blob;
 using Jupiter.Implementation;
+using Microsoft.Extensions.Options;
 using Microsoft.VisualStudio.TestTools.UnitTesting;
+using Moq;
 
 namespace Horde.Storage.UnitTests
 {
@@ -20,10 +23,10 @@ namespace Horde.Storage.UnitTests
         private readonly NamespaceId NsnonExistingNs = new NamespaceId("non-existing-ns");
         private readonly NamespaceId NsOnlyFirst = new NamespaceId("ns-only-in-first");
         private readonly NamespaceId NsOnlySecond = new NamespaceId("ns-only-in-second");
-        private readonly MemoryBlobStore _first = new MemoryBlobStore(throwOnOverwrite: true);
-        private readonly MemoryBlobStore _second = new MemoryBlobStore(throwOnOverwrite: true);
-        private readonly MemoryBlobStore _third = new MemoryBlobStore(throwOnOverwrite: true);
-        private readonly IBlobStore _chained;
+        private readonly MemoryBlobStore _first = new MemoryBlobStore(throwOnOverwrite: false);
+        private readonly MemoryBlobStore _second = new MemoryBlobStore(throwOnOverwrite: false);
+        private readonly MemoryBlobStore _third = new MemoryBlobStore(throwOnOverwrite: false);
+        private BlobService _chained = null!;
         
         private readonly BlobIdentifier _onlyFirstId = new BlobIdentifier(new string('1', 40));
         private readonly BlobIdentifier _onlySecondId = new BlobIdentifier(new string('2', 40));
@@ -35,12 +38,18 @@ namespace Horde.Storage.UnitTests
 
         public HierarchicalBlobStorageTest()
         {
-            _chained = new HierarchicalBlobStore(new []{_first, _second, _third});
+
         }
 
         [TestInitialize]
         public async Task Setup()
         {
+            Mock<IServiceProvider> serviceProviderMock = new Mock<IServiceProvider>();
+            serviceProviderMock.Setup(x => x.GetService(typeof(MemoryCacheBlobStore))).Returns(new MemoryCacheBlobStore(Mock.Of<IOptionsMonitor<MemoryCacheBlobSettings>>(_ => _.CurrentValue == new MemoryCacheBlobSettings())));
+            IOptionsMonitor<HordeStorageSettings> settingsMonitor = Mock.Of<IOptionsMonitor<HordeStorageSettings>>(_ => _.CurrentValue == new HordeStorageSettings());
+            _chained = new BlobService(serviceProviderMock.Object, settingsMonitor, Mock.Of<IBlobIndex>());
+            _chained.BlobStore = new List<IBlobStore> { _first, _second, _third };
+
             await _first.PutObject(Ns, Encoding.ASCII.GetBytes("onlyFirstContent"), _onlyFirstId);
             await _second.PutObject(Ns, Encoding.ASCII.GetBytes("onlySecondContent"), _onlySecondId);
             await _third.PutObject(Ns, Encoding.ASCII.GetBytes("onlyThirdContent"), _onlyThirdId);
@@ -48,14 +57,14 @@ namespace Horde.Storage.UnitTests
             await _first.PutObject(Ns, Encoding.ASCII.GetBytes("allContent"), _allId);
             await _second.PutObject(Ns, Encoding.ASCII.GetBytes("allContent"), _allId);
             await _third.PutObject(Ns, Encoding.ASCII.GetBytes("allContent"), _allId);
-            await _second.PutObject(NsOnlyFirst, Encoding.ASCII.GetBytes("onlyFirstUniqueNs"), _onlyFirstUniqueNsId);
+            await _first.PutObject(NsOnlyFirst, Encoding.ASCII.GetBytes("onlyFirstUniqueNs"), _onlyFirstUniqueNsId);
             await _second.PutObject(NsOnlySecond, Encoding.ASCII.GetBytes("onlySecondUniqueNs"), _onlySecondUniqueNsId);
         }
         
         [TestMethod]
         public async Task PutObject()
         {
-            BlobIdentifier new1 = new BlobIdentifier("1000000000000000000000000000000000000000");
+            BlobIdentifier new1 = new BlobIdentifier("A418A2821A76B110092C9745151E112253F7999B");
             Assert.IsFalse(await _chained.Exists(Ns, new1));
             await _chained.PutObject(Ns, Encoding.ASCII.GetBytes("new1"), new1);
             Assert.IsTrue(await _chained.Exists(Ns, new1));
@@ -137,27 +146,6 @@ namespace Horde.Storage.UnitTests
             await Assert.ThrowsExceptionAsync<NamespaceNotFoundException>(() => _first.DeleteNamespace(Ns));
             await Assert.ThrowsExceptionAsync<NamespaceNotFoundException>(() => _second.DeleteNamespace(Ns));
         }
-        
-        [TestMethod]
-        public async Task ListOldObjects()
-        {
-            List<BlobIdentifier> blobIds = await _chained.ListOldObjects(Ns, DateTime.Now.AddYears(10)).GetAsyncEnumerator().ToListAsync();
-            Assert.IsTrue(blobIds.Exists(x => x.Equals(_onlyFirstId)));
-            Assert.IsTrue(blobIds.Exists(x => x.Equals(_onlySecondId)));
-            Assert.IsTrue(blobIds.Exists(x => x.Equals(_onlyThirdId)));
-            Assert.IsTrue(blobIds.Exists(x => x.Equals(_allId)));
-            Assert.AreEqual(4, blobIds.Count);
-            
-            blobIds = await _chained.ListOldObjects(NsOnlyFirst, DateTime.Now.AddYears(10)).GetAsyncEnumerator().ToListAsync();
-            Assert.IsTrue(blobIds.Exists(x => x.Equals(_onlyFirstUniqueNsId)));
-            Assert.AreEqual(1, blobIds.Count);
-            
-            blobIds = await _chained.ListOldObjects(NsOnlySecond, DateTime.Now.AddYears(10)).GetAsyncEnumerator().ToListAsync();
-            Assert.IsTrue(blobIds.Exists(x => x.Equals(_onlySecondUniqueNsId)));
-            Assert.AreEqual(1, blobIds.Count);
-            
-            await Assert.ThrowsExceptionAsync<NamespaceNotFoundException>(() => _chained.ListOldObjects(NsnonExistingNs, DateTime.Now.AddYears(10)).GetAsyncEnumerator().ToListAsync());
-        }
 
         private string BlobToString(BlobContents contents)
         {
diff --git a/Engine/Source/Programs/Horde/HordeStorage/Horde.Storage/Tests/UnitTests/HordeStorageStartupTest.cs b/Engine/Source/Programs/Horde/HordeStorage/Horde.Storage/Tests/UnitTests/HordeStorageStartupTest.cs
--- a/Engine/Source/Programs/Horde/HordeStorage/Horde.Storage/Tests/UnitTests/HordeStorageStartupTest.cs
+++ b/Engine/Source/Programs/Horde/HordeStorage/Horde.Storage/Tests/UnitTests/HordeStorageStartupTest.cs
@@ -27,29 +27,31 @@ namespace Horde.Storage.UnitTests
         public void StorageBackendImplConfiguration()
         {
             // No configuration set
-            IBlobStore blobStore = GetBlobStoreForConfig(new Dictionary<string, string>());
-            Assert.IsTrue(blobStore is MemoryCacheBlobStore);
+            BlobService blobService = GetBlobServiceForConfig(new Dictionary<string, string>());
+            Assert.IsTrue(blobService.BlobStore.Single() is MemoryCacheBlobStore);
             
             // A single blob store configuration should yield the store itself without a hierarchical wrapper store
-            blobStore = GetBlobStoreForConfig(new Dictionary<string, string> {{"Horde.Storage:StorageImplementations:0", "FileSystem"}});
-            Assert.IsTrue(blobStore is FileSystemStore);
+            blobService = GetBlobServiceForConfig(new Dictionary<string, string> {{"Horde.Storage:StorageImplementations:0", "FileSystem"}});
+            Assert.IsTrue(blobService.BlobStore.Single() is FileSystemStore);
             
             // Should not be case-sensitive
-            blobStore = GetBlobStoreForConfig(new Dictionary<string, string> {{"Horde.Storage:StorageImplementations:0", "FiLeSYsTEm"}});
-            Assert.IsTrue(blobStore is FileSystemStore);
+            blobService = GetBlobServiceForConfig(new Dictionary<string, string> {{"Horde.Storage:StorageImplementations:0", "FiLeSYsTEm"}});
+            Assert.IsTrue(blobService.BlobStore.Single() is FileSystemStore);
             
             // Two or more blob stores returns a hierarchical store
-            blobStore = GetBlobStoreForConfig(new Dictionary<string, string>
+            blobService = GetBlobServiceForConfig(new Dictionary<string, string>
             {
                 {"Horde.Storage:StorageImplementations:0", "FileSystem"},
                 {"Horde.Storage:StorageImplementations:1", "Memory"},
             });
-            Assert.IsTrue(blobStore is HierarchicalBlobStore);
-            Assert.IsTrue((blobStore as HierarchicalBlobStore)!.BlobStores.ToList()[0] is FileSystemStore);
-            Assert.IsTrue((blobStore as HierarchicalBlobStore)!.BlobStores.ToList()[1] is MemoryCacheBlobStore);
+
+
+            List<IBlobStore> blobStores = blobService.BlobStore.ToList();
+            Assert.IsTrue(blobStores[0] is FileSystemStore);
+            Assert.IsTrue(blobStores[1] is MemoryCacheBlobStore);
         }
 
-        private IBlobStore GetBlobStoreForConfig(Dictionary<string, string> configDict)
+        private BlobService GetBlobServiceForConfig(Dictionary<string, string> configDict)
         {
             IConfigurationRoot configuration = new ConfigurationBuilder()
                 .AddJsonFile("appsettings.Testing.json", true)
@@ -63,7 +65,7 @@ namespace Horde.Storage.UnitTests
                 .UseStartup<HordeStorageStartup>()
             );
 
-            return server.Services.GetService<IBlobStore>()!;
+            return (BlobService)server.Services.GetService<IBlobService>()!;
         }
     }
 }
diff --git a/Engine/Source/Programs/Horde/HordeStorage/Horde.Storage/Tests/UnitTests/ReplicatorTests.cs b/Engine/Source/Programs/Horde/HordeStorage/Horde.Storage/Tests/UnitTests/ReplicatorTests.cs
--- a/Engine/Source/Programs/Horde/HordeStorage/Horde.Storage/Tests/UnitTests/ReplicatorTests.cs
+++ b/Engine/Source/Programs/Horde/HordeStorage/Horde.Storage/Tests/UnitTests/ReplicatorTests.cs
@@ -10,6 +10,8 @@ using System.Threading;
 using Microsoft.VisualStudio.TestTools.UnitTesting;
 using System.Threading.Tasks;
 using Horde.Storage.Implementation;
+using Jupiter;
+using Jupiter.Common.Implementation;
 using Jupiter.Implementation;
 using Microsoft.Extensions.Options;
 using Moq;
@@ -22,7 +24,6 @@ namespace Horde.Storage.UnitTests
     {
         private readonly DirectoryInfo _tempPath = new DirectoryInfo(Path.GetTempPath());
         private readonly Guid _currentLogGeneration = new Guid();
-        private const string CurrentSite = "Test";
         private const string ReplicatorNameV1 = "TestReplicatorV1";
 
         private readonly NamespaceId NamespaceV1 = new NamespaceId("test-namespace-v1");
@@ -50,7 +51,7 @@ namespace Horde.Storage.UnitTests
             byte[] contents = Encoding.ASCII.GetBytes("test string");
             BlobIdentifier blobToReplication = BlobIdentifier.FromBlob(contents);
 
-            Mock<IBlobStore> blobStoreMock = new Mock<IBlobStore>();
+            Mock<IBlobService> blobStoreMock = new Mock<IBlobService>();
             IServiceCredentials serviceCredentials = Mock.Of<IServiceCredentials>();
             Mock<ITransactionLogWriter> transactionLogWriter = new Mock<ITransactionLogWriter>();
             Mock<IRestClient> remoteClientMock = new Mock<IRestClient> { DefaultValue = DefaultValue.Empty };
@@ -88,10 +89,11 @@ namespace Horde.Storage.UnitTests
                 .Setup(x => x.ExecuteGetAsync(It.IsAny<IRestRequest>(), It.IsAny<CancellationToken>()))
                 .ReturnsAsync(mockedIoResponse.Object);
 
-            ReplicationSettings replicationSettings = new ReplicationSettings {CurrentSite = CurrentSite, StateRoot = _tempPath.FullName};
+            ReplicationSettings replicationSettings = new ReplicationSettings {StateRoot = _tempPath.FullName};
             IOptionsMonitor<ReplicationSettings> replicationSettingsMonitor = Mock.Of<IOptionsMonitor<ReplicationSettings>>(_ => _.CurrentValue == replicationSettings);
+            IOptionsMonitor<JupiterSettings> jupiterSettingsMonitor = Mock.Of<IOptionsMonitor<JupiterSettings>>(_ => _.CurrentValue == new JupiterSettings());
 
-            using IReplicator replicator = new ReplicatorV1(replicatorSettings, replicationSettingsMonitor, blobStoreMock.Object, transactionLogWriter.Object, remoteClientMock.Object);
+            using IReplicator replicator = new ReplicatorV1(replicatorSettings, replicationSettingsMonitor, jupiterSettingsMonitor, blobStoreMock.Object, transactionLogWriter.Object, remoteClientMock.Object);
 
             Assert.IsNull(replicator.State.ReplicatorOffset,"Expected state to have been reset during test initialize");
 
@@ -142,12 +144,13 @@ namespace Horde.Storage.UnitTests
                 .ReturnsAsync(mockedCallistoResponse[0].Object)
                 .ReturnsAsync(mockedCallistoResponse[0].Object);
 
-            Mock<IBlobStore> blobStoreMock = new Mock<IBlobStore>();
+            Mock<IBlobService> blobStoreMock = new Mock<IBlobService>();
             Mock<ITransactionLogWriter> transactionLogWriter = new Mock<ITransactionLogWriter>();
-            ReplicationSettings replicationSettings = new ReplicationSettings {CurrentSite = CurrentSite, StateRoot = _tempPath.FullName};
+            ReplicationSettings replicationSettings = new ReplicationSettings {StateRoot = _tempPath.FullName};
             IOptionsMonitor<ReplicationSettings> replicationSettingsMonitor = Mock.Of<IOptionsMonitor<ReplicationSettings>>(_ => _.CurrentValue == replicationSettings);
+            IOptionsMonitor<JupiterSettings> jupiterSettingsMonitor = Mock.Of<IOptionsMonitor<JupiterSettings>>(_ => _.CurrentValue == new JupiterSettings());
 
-            using IReplicator replicator = new ReplicatorV1(replicatorSettings, replicationSettingsMonitor, blobStoreMock.Object, transactionLogWriter.Object, remoteClientMock.Object);
+            using IReplicator replicator = new ReplicatorV1(replicatorSettings, replicationSettingsMonitor, jupiterSettingsMonitor, blobStoreMock.Object, transactionLogWriter.Object, remoteClientMock.Object);
 
             Assert.IsNull(replicator.State.ReplicatorOffset,"Expected state to have been reset during test initialize");
 
@@ -177,7 +180,7 @@ namespace Horde.Storage.UnitTests
             byte[] contents = Encoding.ASCII.GetBytes("test string");
             BlobIdentifier blobToReplication = BlobIdentifier.FromBlob(contents);
 
-            Mock<IBlobStore> blobStoreMock = new Mock<IBlobStore>();
+            Mock<IBlobService> blobStoreMock = new Mock<IBlobService>();
             IServiceCredentials serviceCredentials = Mock.Of<IServiceCredentials>();
             Mock<ITransactionLogWriter> transactionLogWriter = new Mock<ITransactionLogWriter>();
             Mock<IRestClient> remoteClientMock = new Mock<IRestClient> { DefaultValue = DefaultValue.Empty };
@@ -218,10 +221,11 @@ namespace Horde.Storage.UnitTests
                 .ReturnsAsync(mockedIoResponse.Object)
                 .ReturnsAsync(mockedIoResponse2.Object);
 
-            ReplicationSettings replicationSettings = new ReplicationSettings {CurrentSite = CurrentSite, StateRoot = _tempPath.FullName};
+            ReplicationSettings replicationSettings = new ReplicationSettings {StateRoot = _tempPath.FullName};
             IOptionsMonitor<ReplicationSettings> replicationSettingsMonitor = Mock.Of<IOptionsMonitor<ReplicationSettings>>(_ => _.CurrentValue == replicationSettings);
+            IOptionsMonitor<JupiterSettings> jupiterSettingsMonitor = Mock.Of<IOptionsMonitor<JupiterSettings>>(_ => _.CurrentValue == new JupiterSettings());
 
-            using IReplicator replicator = new ReplicatorV1(replicatorSettings, replicationSettingsMonitor, blobStoreMock.Object, transactionLogWriter.Object, remoteClientMock.Object);
+            using IReplicator replicator = new ReplicatorV1(replicatorSettings, replicationSettingsMonitor, jupiterSettingsMonitor, blobStoreMock.Object, transactionLogWriter.Object, remoteClientMock.Object);
 
             Assert.IsNull(replicator.State.ReplicatorOffset,"Expected state to have been reset during test initialize");
 
@@ -257,7 +261,7 @@ namespace Horde.Storage.UnitTests
 
             NamespaceId ns = replicatorSettings.NamespaceToReplicate;
 
-            Mock<IBlobStore> blobStoreMock = new Mock<IBlobStore>();
+            Mock<IBlobService> blobStoreMock = new Mock<IBlobService>();
             IServiceCredentials serviceCredentials = Mock.Of<IServiceCredentials>();
             Mock<ITransactionLogWriter> transactionLogWriter = new Mock<ITransactionLogWriter>();
             Mock<IRestClient> remoteClientMock = new Mock<IRestClient> { DefaultValue = DefaultValue.Empty };
@@ -289,10 +293,11 @@ namespace Horde.Storage.UnitTests
                 .ReturnsAsync(mockedCallistoResponse[0].Object)
                 .ReturnsAsync(mockedCallistoResponse[1].Object);
 
-            ReplicationSettings replicationSettings = new ReplicationSettings {CurrentSite = CurrentSite, StateRoot = _tempPath.FullName};
+            ReplicationSettings replicationSettings = new ReplicationSettings {StateRoot = _tempPath.FullName};
             IOptionsMonitor<ReplicationSettings> replicationSettingsMonitor = Mock.Of<IOptionsMonitor<ReplicationSettings>>(_ => _.CurrentValue == replicationSettings);
+            IOptionsMonitor<JupiterSettings> jupiterSettingsMonitor = Mock.Of<IOptionsMonitor<JupiterSettings>>(_ => _.CurrentValue == new JupiterSettings());
 
-            using IReplicator replicator = new ReplicatorV1(replicatorSettings, replicationSettingsMonitor, blobStoreMock.Object, transactionLogWriter.Object, remoteClientMock.Object);
+            using IReplicator replicator = new ReplicatorV1(replicatorSettings, replicationSettingsMonitor, jupiterSettingsMonitor, blobStoreMock.Object, transactionLogWriter.Object, remoteClientMock.Object);
 
             Assert.IsNull(replicator.State.ReplicatorOffset, "Expected state to have been reset during test initialize");
 
diff --git a/Engine/Source/Programs/Horde/HordeStorage/Horde.Storage/Tests/UnitTests/S3StoreTests.cs b/Engine/Source/Programs/Horde/HordeStorage/Horde.Storage/Tests/UnitTests/S3StoreTests.cs
--- a/Engine/Source/Programs/Horde/HordeStorage/Horde.Storage/Tests/UnitTests/S3StoreTests.cs
+++ b/Engine/Source/Programs/Horde/HordeStorage/Horde.Storage/Tests/UnitTests/S3StoreTests.cs
@@ -7,6 +7,7 @@ using System.Threading.Tasks;
 using Amazon.S3;
 using Amazon.S3.Model;
 using Horde.Storage.Implementation;
+using Horde.Storage.Implementation.Blob;
 using Jupiter.Implementation;
 using Microsoft.Extensions.Options;
 using Microsoft.VisualStudio.TestTools.UnitTesting;
@@ -26,7 +27,7 @@ namespace Horde.Storage.UnitTests
             S3Settings settings = new S3Settings{BucketName = "tests-foo"};
             IOptionsMonitor<S3Settings> settingsMock = Mock.Of<IOptionsMonitor<S3Settings>>(_ => _.CurrentValue == settings);
 
-            AmazonS3Store store = new AmazonS3Store(s3Mock.Object, settingsMock);
+            AmazonS3Store store = new AmazonS3Store(s3Mock.Object, settingsMock, Mock.Of<IBlobIndex>());
             byte[] content = Encoding.ASCII.GetBytes("test content");
             BlobIdentifier blobIdentifier = BlobIdentifier.FromBlob(content);
             Task task = store.PutObject(Namespace, content.AsMemory(), blobIdentifier);
@@ -53,7 +54,7 @@ namespace Horde.Storage.UnitTests
             IOptionsMonitor<S3Settings> settingsMock = Mock.Of<IOptionsMonitor<S3Settings>>(_ => _.CurrentValue == settings);
 
             s3Mock.Setup(s3 => s3.PutObjectAsync(It.IsAny<PutObjectRequest>(), default)).Throws<Exception>();
-            AmazonS3Store store = new AmazonS3Store(s3Mock.Object, settingsMock);
+            AmazonS3Store store = new AmazonS3Store(s3Mock.Object, settingsMock, Mock.Of<IBlobIndex>());
             Task task = store.PutObject(Namespace, content, blob);
             await task;
 
@@ -74,7 +75,7 @@ namespace Horde.Storage.UnitTests
             S3Settings settings = new S3Settings{BucketName = "tests-foo"};
             IOptionsMonitor<S3Settings> settingsMock = Mock.Of<IOptionsMonitor<S3Settings>>(_ => _.CurrentValue == settings);
 
-            AmazonS3Store store = new AmazonS3Store(s3Mock.Object, settingsMock);
+            AmazonS3Store store = new AmazonS3Store(s3Mock.Object, settingsMock, Mock.Of<IBlobIndex>());
             await using BlobContents blobContents = await store.GetObject(Namespace, blob);
 
             s3Mock.Verify(s3 => s3.GetObjectAsync("tests-foo", blob.AsS3Key(), default));
@@ -89,7 +90,7 @@ namespace Horde.Storage.UnitTests
             S3Settings settings = new S3Settings{BucketName = "tests-foo"};
             IOptionsMonitor<S3Settings> settingsMock = Mock.Of<IOptionsMonitor<S3Settings>>(_ => _.CurrentValue == settings);
 
-            AmazonS3Store store = new AmazonS3Store(s3Mock.Object, settingsMock);
+            AmazonS3Store store = new AmazonS3Store(s3Mock.Object, settingsMock, Mock.Of<IBlobIndex>());
             Task task = store.DeleteObject(Namespace, blob);
             await task;
 
diff --git a/Engine/Source/Programs/Horde/HordeStorage/Horde.Storage/appsettings.Development.json b/Engine/Source/Programs/Horde/HordeStorage/Horde.Storage/appsettings.Development.json
--- a/Engine/Source/Programs/Horde/HordeStorage/Horde.Storage/appsettings.Development.json
+++ b/Engine/Source/Programs/Horde/HordeStorage/Horde.Storage/appsettings.Development.json
@@ -74,6 +74,7 @@
   },
 
   "Jupiter": {
+    "CurrentSite": "local",
     "DisableAuthOnPorts": [
       57003
     ]
diff --git a/Engine/Source/Programs/Horde/HordeStorage/Horde.Storage/appsettings.Testing.json b/Engine/Source/Programs/Horde/HordeStorage/Horde.Storage/appsettings.Testing.json
--- a/Engine/Source/Programs/Horde/HordeStorage/Horde.Storage/appsettings.Testing.json
+++ b/Engine/Source/Programs/Horde/HordeStorage/Horde.Storage/appsettings.Testing.json
@@ -68,7 +68,7 @@
     "AwsSecretKey": "WaUf4Zq+J6YMlvbCVvTRkQJIPc7DvZwTm5UP0tPJ"
   },
   "Azure": {
-    "ConnectionString": "DefaultEndpointsProtocol=https;BlobEndpoint=http://localhost:11002/blobaccount;AccountName=blobaccount;AccountKey=iU6uTvlF1ysppmft+NO5lAD0E3hwrAORr5Rb5xcBWUgEz/OicrSkFxwZYMNK5XL29/wXZKGOoOVSW040nAOfPg==",
+    "ConnectionString": "DefaultEndpointsProtocol=https;BlobEndpoint=http://localhost:11002/blobaccount;AccountName=blobaccount;AccountKey=iU6uTvlF1ysppmft+NO5lAD0E3hwrAORr5Rb5xcBWUgEz/OicrSkFxwZYMNK5XL29/wXZKGOoOVSW040nAOfPg=="
   },
   "Auth": {
     "Method": "Disabled"
@@ -84,6 +84,7 @@
   "AllowedHosts": "*",
 
   "Jupiter": {
+    "CurrentSite": "test", 
     "DisableAuthOnPorts": [
       57003
     ]
diff --git a/Engine/Source/Programs/Horde/HordeStorage/Jupiter.Common/BaseStartup.cs b/Engine/Source/Programs/Horde/HordeStorage/Jupiter.Common/BaseStartup.cs
--- a/Engine/Source/Programs/Horde/HordeStorage/Jupiter.Common/BaseStartup.cs
+++ b/Engine/Source/Programs/Horde/HordeStorage/Jupiter.Common/BaseStartup.cs
@@ -437,6 +437,11 @@ namespace Jupiter
 
     public class JupiterSettings
     {
+        /// <summary>
+        /// If the request is smaller then MemoryBufferSize we buffer it in memory rather then as a file
+        /// </summary>
+        public long MemoryBufferSize { get; set; } = Int32.MaxValue;
+
         // enable to unhide potentially personal information, see https://aka.ms/IdentityModel/PII
         public bool ShowPII = false;
         public bool DisableHealthChecks { get; set; } = false;
@@ -447,6 +452,13 @@ namespace Jupiter
         // Enable to echo every request to the log file, usually this is more efficiently done on the load balancer
         public bool UseRequestLogging { get; set; } = false;
 
+        /// <summary>
+        ///  Name of the current site, has to be globally unique across all deployments
+        /// </summary>
+        [Required]
+        [Key]
+        public string CurrentSite { get; set; } = "";
+
         //public Dictionary<string, int> DisableAuthOnPorts { get; set; } = new ();
     }
 
diff --git a/Engine/Source/Programs/Horde/HordeStorage/Jupiter.Common/Implementation/StreamingBodyPayload.cs b/Engine/Source/Programs/Horde/HordeStorage/Jupiter.Common/Implementation/StreamingBodyPayload.cs
new file mode 100644
--- /dev/null
+++ b/Engine/Source/Programs/Horde/HordeStorage/Jupiter.Common/Implementation/StreamingBodyPayload.cs
@@ -0,0 +1,119 @@
+using System;
+using System.IO;
+using System.Threading.Tasks;
+using Jupiter.Utils;
+using Microsoft.AspNetCore.Http;
+using Microsoft.Extensions.Options;
+
+namespace Jupiter.Common.Implementation
+{
+    public interface IBufferedPayload : IDisposable
+    {
+        Stream GetStream();
+        long Length { get; }
+    }
+
+    /// <summary>
+    /// Streaming request that is streamed into memory
+    /// </summary>
+    public class MemoryBufferedPayload : IBufferedPayload
+    {
+        private byte[] _buffer;
+
+        public MemoryBufferedPayload(byte[] source)
+        {
+            _buffer = source;
+        }
+
+        public static async Task<MemoryBufferedPayload> Create(Stream s)
+        {
+            MemoryBufferedPayload payload = new MemoryBufferedPayload(await s.ToByteArray());
+
+            return payload;
+        }
+
+        public void Dispose()
+        {
+            
+        }
+
+        public Stream GetStream()
+        {
+            return new MemoryStream(_buffer);
+        }
+
+        public long Length { get { return _buffer.LongLength; } }
+    }
+
+    /// <summary>
+    /// A streaming request backed by a temporary file on disk
+    /// </summary>
+    public class FilesystemBufferedPayload : IBufferedPayload
+    {
+        private readonly FileInfo _tempFile;
+        private long _length;
+
+        private FilesystemBufferedPayload()
+        {
+            _tempFile = new FileInfo(Path.GetTempFileName());
+        }
+
+        public static async Task<FilesystemBufferedPayload> Create(Stream s)
+        {
+            FilesystemBufferedPayload payload = new FilesystemBufferedPayload();
+
+            {
+                await using FileStream fs = payload._tempFile.OpenWrite();
+                await s.CopyToAsync(fs);
+                payload._length = payload._tempFile.Length;
+            }
+
+            return payload;
+        }
+
+        public void Dispose()
+        {
+            if (_tempFile.Exists)
+                _tempFile.Delete();
+        }
+
+        public Stream GetStream()
+        {
+            return _tempFile.OpenRead();
+        }
+
+        public long Length
+        {
+            get { return _length; }
+        }
+    }
+
+    public class BufferedPayloadFactory
+    {
+        private readonly IOptionsMonitor<JupiterSettings> _jupiterSettings;
+
+        public BufferedPayloadFactory(IOptionsMonitor<JupiterSettings> jupiterSettings)
+        {
+            _jupiterSettings = jupiterSettings;
+        }
+
+        public async Task<IBufferedPayload> CreateFromRequest(HttpRequest request)
+        {
+            long? contentLength = request.ContentLength;
+
+            if (contentLength == null)
+            {
+                throw new Exception("Expected content-length on all requests");
+            }
+
+            
+            // blob is small enough to fit into memory we just read it as is
+            if (contentLength < _jupiterSettings.CurrentValue.MemoryBufferSize)
+            {
+                return await MemoryBufferedPayload.Create(request.Body);
+            }
+
+            return await FilesystemBufferedPayload.Create(request.Body);
+        }
+    }
+}
