Horde Storage - Fixed issues with last access test.

Disabled the batch test in a better way for dynamo which doesnt disable the other tests from running.

Made sure to empty the last access list before running the last access rollup, should fix issues with other tests occasionally adding state which broke the last access test.

diff --git a/Engine/Source/Programs/Horde/HordeStorage/Horde.Storage/Tests/Functional/RefTests.cs b/Engine/Source/Programs/Horde/HordeStorage/Horde.Storage/Tests/Functional/RefTests.cs
--- a/Engine/Source/Programs/Horde/HordeStorage/Horde.Storage/Tests/Functional/RefTests.cs
+++ b/Engine/Source/Programs/Horde/HordeStorage/Horde.Storage/Tests/Functional/RefTests.cs
@@ -135,13 +135,6 @@ namespace Horde.Storage.FunctionalTests.Ref
 
             return record;
         }
-
-        // removing batch test using dynamo as it hangs the test runner for unknown reasons when run with a lot of other tests
-        [TestMethod]
-        public new Task Batch()
-        {
-            return Task.CompletedTask;
-        }
     }
     
     public abstract class RefTests
@@ -593,6 +586,8 @@ namespace Horde.Storage.FunctionalTests.Ref
         {
             IRefsStore refsStore = _server!.Services.GetService<IRefsStore>()!;
 
+            // empty the last access cache so we are sure it does not contain state from previous runs
+            await _lastAccessCache!.GetLastAccessedRecords();
             OldRecord[] emptyRecords = await refsStore.GetOldRecords(TestNamespace, TimeSpan.FromDays(7)).ToArrayAsync();
             // no content older then 7 days as we just insert it
             Assert.AreEqual(0, emptyRecords.Length);
@@ -631,6 +626,9 @@ namespace Horde.Storage.FunctionalTests.Ref
         [TestMethod]
         public async Task Batch()
         {
+            if (this is DynamoRefTests)
+                Assert.Inconclusive();
+
             // we chunk at 1 kb to lets generate a 3kb string
             string content = string.Concat(Enumerable.Repeat("Duplicate string", 200));
             byte[] payload = Encoding.ASCII.GetBytes(content);
