Horde: Add support for ModifiedAfter in the specialized find jobs method

diff --git a/Engine/Source/Programs/Horde/Horde.Build/Controllers/JobsController.cs b/Engine/Source/Programs/Horde/Horde.Build/Controllers/JobsController.cs
--- a/Engine/Source/Programs/Horde/Horde.Build/Controllers/JobsController.cs
+++ b/Engine/Source/Programs/Horde/Horde.Build/Controllers/JobsController.cs
@@ -707,6 +707,7 @@ namespace HordeServer.Controllers
 		/// <param name="StreamId">The stream to search for</param>
 		/// <param name="Templates">List of templates to find</param>
 		/// <param name="PreflightStartedByUserId">User id for which to include preflight jobs</param>
+		/// <param name="ModifiedAfter">If specified, only jobs updated after the given time will be returned</param>
 		/// <param name="Filter">Filter for properties to return</param>
 		/// <param name="Index">Index of the first result to be returned</param>
 		/// <param name="Count">Number of results to return</param>
@@ -719,6 +720,7 @@ namespace HordeServer.Controllers
 			string StreamId,
 			[FromQuery(Name = "template")] string[] Templates,
 			[FromQuery] string? PreflightStartedByUserId = null,
+			[FromQuery] DateTimeOffset? ModifiedAfter = null,
 			[FromQuery] PropertyFilter? Filter = null,
 			[FromQuery] int Index = 0,
 			[FromQuery] int Count = 100,
@@ -729,7 +731,7 @@ namespace HordeServer.Controllers
 			UserId? PreflightStartedByUserIdValue = PreflightStartedByUserId != null ? new UserId(PreflightStartedByUserId) : null;
 			Count = Math.Min(1000, Count);
 
-			List<IJob> Jobs = await JobService.FindJobsByStreamWithTemplatesAsync(StreamIdValue, TemplateRefIds, PreflightStartedByUserIdValue, Index, Count, ConsistentRead);
+			List<IJob> Jobs = await JobService.FindJobsByStreamWithTemplatesAsync(StreamIdValue, TemplateRefIds, PreflightStartedByUserIdValue, ModifiedAfter, Index, Count, ConsistentRead);
 			return await CreateAuthorizedJobResponses(Jobs, Filter);
 		}
 
diff --git a/Engine/Source/Programs/Horde/Horde.Build/Jobs/IJobCollection.cs b/Engine/Source/Programs/Horde/Horde.Build/Jobs/IJobCollection.cs
--- a/Engine/Source/Programs/Horde/Horde.Build/Jobs/IJobCollection.cs
+++ b/Engine/Source/Programs/Horde/Horde.Build/Jobs/IJobCollection.cs
@@ -115,10 +115,11 @@ namespace HordeServer.Collections
 		/// <param name="StreamId">The stream containing the job</param>
 		/// <param name="Templates">Templates to look for</param>
 		/// <param name="PreflightStartedByUser">User for which to include preflight jobs</param>
+		/// <param name="ModifiedAfter">Filter the results by modified time</param>
 		/// <param name="Index">Index of the first result to return</param>
 		/// <param name="Count">Number of results to return</param>
 		/// <param name="ConsistentRead">If the database read should be made to the replica server</param>
-		Task<List<IJob>> FindLatestByStreamWithTemplatesAsync(StreamId StreamId, TemplateRefId[] Templates, UserId? PreflightStartedByUser = null, int? Index = null, int? Count = null, bool ConsistentRead = false);
+		Task<List<IJob>> FindLatestByStreamWithTemplatesAsync(StreamId StreamId, TemplateRefId[] Templates, UserId? PreflightStartedByUser = null, DateTimeOffset? ModifiedAfter = null, int? Index = null, int? Count = null, bool ConsistentRead = false);
 
 		/// <summary>
 		/// Updates a new job
diff --git a/Engine/Source/Programs/Horde/Horde.Build/Jobs/Impl/JobCollection.cs b/Engine/Source/Programs/Horde/Horde.Build/Jobs/Impl/JobCollection.cs
--- a/Engine/Source/Programs/Horde/Horde.Build/Jobs/Impl/JobCollection.cs
+++ b/Engine/Source/Programs/Horde/Horde.Build/Jobs/Impl/JobCollection.cs
@@ -623,12 +623,12 @@ namespace HordeServer.Collections.Impl
 		}
 
 		/// <inheritdoc/>
-		public async Task<List<IJob>> FindLatestByStreamWithTemplatesAsync(StreamId StreamId, TemplateRefId[] Templates, UserId? PreflightStartedByUser, int? Index, int? Count, bool ConsistentRead)
+		public async Task<List<IJob>> FindLatestByStreamWithTemplatesAsync(StreamId StreamId, TemplateRefId[] Templates, UserId? PreflightStartedByUser, DateTimeOffset? ModifiedAfter, int? Index, int? Count, bool ConsistentRead)
 		{
 			// This find call uses an index hint. Modifying the parameter passed to FindAsync can affect execution time a lot as the query planner is forced to use the specified index.
 			// Casting to interface to benefit from default parameter values
 			return await (this as IJobCollection).FindAsync(
-				StreamId: StreamId, Templates: Templates, PreflightStartedByUser: PreflightStartedByUser,
+				StreamId: StreamId, Templates: Templates, PreflightStartedByUser: PreflightStartedByUser, ModifiedAfter: ModifiedAfter,
 				Index: Index, Count: Count, IndexHint: Indexes.CreateTimeUtc, ConsistentRead: ConsistentRead);
 		}
 
diff --git a/Engine/Source/Programs/Horde/Horde.Build/Jobs/Impl/JobService.cs b/Engine/Source/Programs/Horde/Horde.Build/Jobs/Impl/JobService.cs
--- a/Engine/Source/Programs/Horde/Horde.Build/Jobs/Impl/JobService.cs
+++ b/Engine/Source/Programs/Horde/Horde.Build/Jobs/Impl/JobService.cs
@@ -611,21 +611,23 @@ namespace HordeServer.Services
 		/// </summary>
 		/// <param name="StreamId">The stream containing the job</param>
 		/// <param name="Templates">Templates to look for</param>
-		/// <param name="PreflightStartedByUser">User for which to include preflight jobs</param>		
+		/// <param name="PreflightStartedByUser">User for which to include preflight jobs</param>
+		/// <param name="ModifiedAfter">Filter the results by last modified time</param>	
 		/// <param name="Index">Index of the first result to return</param>
 		/// <param name="Count">Number of results to return</param>
 		/// <param name="ConsistentRead">If the database read should be made to the replica server</param>
 		/// <returns>List of jobs matching the given criteria</returns>
-		public async Task<List<IJob>> FindJobsByStreamWithTemplatesAsync(StreamId StreamId, TemplateRefId[] Templates, UserId? PreflightStartedByUser = null, int? Index = null, int? Count = null, bool ConsistentRead = true)
+		public async Task<List<IJob>> FindJobsByStreamWithTemplatesAsync(StreamId StreamId, TemplateRefId[] Templates, UserId? PreflightStartedByUser = null, DateTimeOffset? ModifiedAfter = null, int? Index = null, int? Count = null, bool ConsistentRead = true)
 		{
 			using IScope Scope = GlobalTracer.Instance.BuildSpan("JobService.FindJobsByStreamWithTemplatesAsync").StartActive();
 			Scope.Span.SetTag("StreamId", StreamId);
 			Scope.Span.SetTag("Templates", Templates);
 			Scope.Span.SetTag("PreflightStartedByUser", PreflightStartedByUser);
+			Scope.Span.SetTag("ModifiedAfter", ModifiedAfter);
 			Scope.Span.SetTag("Index", Index);
 			Scope.Span.SetTag("Count", Count);
 			
-			return await Jobs.FindLatestByStreamWithTemplatesAsync(StreamId, Templates, PreflightStartedByUser, Index, Count, ConsistentRead);
+			return await Jobs.FindLatestByStreamWithTemplatesAsync(StreamId, Templates, PreflightStartedByUser, ModifiedAfter, Index, Count, ConsistentRead);
 		}
 
 		/// <summary>
