Horde: Default to non-consistent reads in specialized find jobs method

diff --git a/Engine/Source/Programs/Horde/Horde.Build/Controllers/JobsController.cs b/Engine/Source/Programs/Horde/Horde.Build/Controllers/JobsController.cs
--- a/Engine/Source/Programs/Horde/Horde.Build/Controllers/JobsController.cs
+++ b/Engine/Source/Programs/Horde/Horde.Build/Controllers/JobsController.cs
@@ -724,7 +724,7 @@ namespace HordeServer.Controllers
 			[FromQuery] PropertyFilter? Filter = null,
 			[FromQuery] int Index = 0,
 			[FromQuery] int Count = 100,
-			[FromQuery] bool ConsistentRead = true)
+			[FromQuery] bool ConsistentRead = false)
 		{
 			StreamId StreamIdValue = new StreamId(StreamId);
 			TemplateRefId[] TemplateRefIds = Templates.Select(x => new TemplateRefId(x)).ToArray();
