Horde: Ensure the correct index is picked when querying for ModifiedAfter

diff --git a/Engine/Source/Programs/Horde/Horde.Build/Jobs/Impl/JobCollection.cs b/Engine/Source/Programs/Horde/Horde.Build/Jobs/Impl/JobCollection.cs
--- a/Engine/Source/Programs/Horde/Horde.Build/Jobs/Impl/JobCollection.cs
+++ b/Engine/Source/Programs/Horde/Horde.Build/Jobs/Impl/JobCollection.cs
@@ -629,7 +629,7 @@ namespace HordeServer.Collections.Impl
 			// Casting to interface to benefit from default parameter values
 			return await (this as IJobCollection).FindAsync(
 				StreamId: StreamId, Templates: Templates, PreflightStartedByUser: PreflightStartedByUser, ModifiedAfter: ModifiedAfter,
-				Index: Index, Count: Count, IndexHint: Indexes.CreateTimeUtc, ConsistentRead: ConsistentRead);
+				Index: Index, Count: Count, IndexHint: ModifiedAfter == null ? Indexes.CreateTimeUtc : Indexes.UpdateTimeUtc, ConsistentRead: ConsistentRead);
 		}
 
 		/// <inheritdoc/>
