Output current available disk space at the start of steps, this will help us nail down out of disk space issues when an agent at idle has plenty of space available by allowing us to see which steps are eating up all the space.

#rnx

diff --git a/Engine/Source/Programs/Horde/Horde.Agent/Services/WorkerService.cs b/Engine/Source/Programs/Horde/Horde.Agent/Services/WorkerService.cs
--- a/Engine/Source/Programs/Horde/Horde.Agent/Services/WorkerService.cs
+++ b/Engine/Source/Programs/Horde/Horde.Agent/Services/WorkerService.cs
@@ -1008,9 +1008,35 @@ namespace HordeAgent.Services
 					break;
 				}
 
+				// Get current disk space available. This will allow us to more easily spot steps that eat up a lot of disk space.
+				string? DriveName;
+				if (RuntimeInformation.IsOSPlatform(OSPlatform.Windows))
+				{
+					DriveName = Path.GetPathRoot(WorkingDir.FullName);
+				}
+				else
+				{
+					DriveName = WorkingDir.FullName;
+				}
+
+				float AvailableFreeSpace = 0;
+				if (DriveName != null)
+				{
+					try
+					{
+						DriveInfo Info = new DriveInfo(DriveName);
+						AvailableFreeSpace = (1.0f * Info.AvailableFreeSpace) / 1024 / 1024 / 1024;
+					}
+					catch (Exception Ex)
+					{
+						BatchLogger.LogWarning(Ex, "Unable to query disk info for path '{DriveName}'", DriveName);
+					}
+				}
+
 				// Print the new state
 				Stopwatch StepTimer = Stopwatch.StartNew();
-				BatchLogger.LogInformation("Starting job {JobId}, batch {BatchId}, step {StepId}", ExecuteTask.JobId, ExecuteTask.BatchId, Step.StepId);
+				
+				BatchLogger.LogInformation("Starting job {JobId}, batch {BatchId}, step {StepId} (Drive Space Left: {DriveSpaceRemaining} GB)", ExecuteTask.JobId, ExecuteTask.BatchId, Step.StepId, AvailableFreeSpace.ToString("F1"));
 
 				// Create a trace span
 				using IScope Scope = GlobalTracer.Instance.BuildSpan("Execute").WithResourceName(Step.Name).StartActive();
