Horde Storage

Seperated LastAccess tests into a seperate namespace to avoid issues when tests run at the same time.

diff --git a/Engine/Source/Programs/Horde/HordeStorage/Horde.Storage/Tests/Functional/RefTests.cs b/Engine/Source/Programs/Horde/HordeStorage/Horde.Storage/Tests/Functional/RefTests.cs
--- a/Engine/Source/Programs/Horde/HordeStorage/Horde.Storage/Tests/Functional/RefTests.cs
+++ b/Engine/Source/Programs/Horde/HordeStorage/Horde.Storage/Tests/Functional/RefTests.cs
@@ -35,7 +35,6 @@ namespace Horde.Storage.FunctionalTests.Ref
     [TestClass]
     public class MongoRefTests : RefTests
     {
-        private static readonly string MongoTestDatabase = "Europa_TestNamespace";
         private static MongoSettings? _mongoSettings;
         private MongoRefsStore? _refsStore;
 
@@ -49,12 +48,16 @@ namespace Horde.Storage.FunctionalTests.Ref
             _mongoSettings = provider.GetService<IOptionsMonitor<MongoSettings>>()!.CurrentValue;
 
             MongoClient client = GetMongoClient();
-            if (client.GetDatabase(MongoTestDatabase) != null)
+            foreach (NamespaceId ns in new[] {TestNamespace, LastAccessTestNamespace})
             {
-                await client.DropDatabaseAsync(MongoTestDatabase);
+                string dbName = $"Europa_{ns}";
+                if (client.GetDatabase(dbName) != null)
+                {
+                    await client.DropDatabaseAsync(dbName);
+                }
             }
 
-            
+
             //verify we are using the expected refs store
             IRefsStore? refStore = provider.GetService<IRefsStore>();
             Assert.IsTrue(RefStoreIs(refStore, typeof(MongoRefsStore)));
@@ -149,6 +152,7 @@ namespace Horde.Storage.FunctionalTests.Ref
         protected static readonly BlobIdentifier TestObjectBlobHash = BlobIdentifier.FromContentHash(TestObjectHash);
         
         protected readonly NamespaceId TestNamespace = new NamespaceId("test-namespace");
+        protected readonly NamespaceId LastAccessTestNamespace = new NamespaceId("test-namespace-last-access");
 
         [TestInitialize]
         public async Task Setup()
@@ -186,26 +190,28 @@ namespace Horde.Storage.FunctionalTests.Ref
 
             BucketId bucket = new BucketId("bucket");
 
-            await refsStore.Add(new RefRecord(TestNamespace, bucket, new KeyId("testObject"), blobs: new[] { TestObjectBlobHash },
-                lastAccessTime: lastAccessTime, contentHash: TestObjectHash, metadata: null));
-            await blobStore.PutObject(TestNamespace, TestObjectData, TestObjectBlobHash);
-
-            // a object with metadata
-            await refsStore.Add(new RefRecord(TestNamespace, bucket, new KeyId("testObjectWithMetadata"), blobs: new[] { TestObjectBlobHash },
-                lastAccessTime: lastAccessTime, contentHash: TestObjectHash, metadata: new Dictionary<string, object>() {{"Foo", "Bar"}}));
-            await blobStore.PutObject(TestNamespace, TestObjectData, TestObjectBlobHash);
-
-            // the deletableObject should be slightly newer then testObject to control sorting for last access
-            await refsStore.Add(new RefRecord(TestNamespace, bucket, new KeyId("deletableObject"),
-                blobs: new[] { TestObjectBlobHash }, contentHash: TestObjectHash,lastAccessTime: lastAccessTime.AddMinutes(2), metadata: null));
-            await blobStore.PutObject(TestNamespace, TestObjectData, TestObjectBlobHash);
-
-            byte[] contents = Encoding.ASCII.GetBytes("This content will only be hashed");
-            BlobIdentifier newContentIdentifier = BlobIdentifier.FromBlob(contents);
-            await refsStore.Add(new RefRecord(TestNamespace, bucket, new KeyId("notUploadedToIo"), blobs: new[] { newContentIdentifier },
-                lastAccessTime: lastAccessTime, contentHash: newContentIdentifier, metadata: null));
-            // this object should not be uploaded to IO, and is considered recently uploaded for last access filtering
-
+            foreach (NamespaceId ns in new [] {TestNamespace, LastAccessTestNamespace})
+            {
+                await refsStore.Add(new RefRecord(ns, bucket, new KeyId("testObject"), blobs: new[] { TestObjectBlobHash },
+                    lastAccessTime: lastAccessTime, contentHash: TestObjectHash, metadata: null));
+                await blobStore.PutObject(ns, TestObjectData, TestObjectBlobHash);
+
+                // a object with metadata
+                await refsStore.Add(new RefRecord(ns, bucket, new KeyId("testObjectWithMetadata"), blobs: new[] { TestObjectBlobHash },
+                    lastAccessTime: lastAccessTime, contentHash: TestObjectHash, metadata: new Dictionary<string, object>() {{"Foo", "Bar"}}));
+                await blobStore.PutObject(ns, TestObjectData, TestObjectBlobHash);
+
+                // the deletableObject should be slightly newer then testObject to control sorting for last access
+                await refsStore.Add(new RefRecord(ns, bucket, new KeyId("deletableObject"),
+                    blobs: new[] { TestObjectBlobHash }, contentHash: TestObjectHash,lastAccessTime: lastAccessTime.AddMinutes(2), metadata: null));
+                await blobStore.PutObject(ns, TestObjectData, TestObjectBlobHash);
+
+                byte[] contents = Encoding.ASCII.GetBytes("This content will only be hashed");
+                BlobIdentifier newContentIdentifier = BlobIdentifier.FromBlob(contents);
+                await refsStore.Add(new RefRecord(ns, bucket, new KeyId("notUploadedToIo"), blobs: new[] { newContentIdentifier },
+                    lastAccessTime: lastAccessTime, contentHash: newContentIdentifier, metadata: null));
+                // this object should not be uploaded to IO, and is considered recently uploaded for last access filtering 
+            }
         }
 
 
@@ -586,14 +592,12 @@ namespace Horde.Storage.FunctionalTests.Ref
         {
             IRefsStore refsStore = _server!.Services.GetService<IRefsStore>()!;
 
-            // empty the last access cache so we are sure it does not contain state from previous runs
-            await _lastAccessCache!.GetLastAccessedRecords();
-            OldRecord[] emptyRecords = await refsStore.GetOldRecords(TestNamespace, TimeSpan.FromDays(7)).ToArrayAsync();
+            OldRecord[] emptyRecords = await refsStore.GetOldRecords(LastAccessTestNamespace, TimeSpan.FromDays(7)).ToArrayAsync();
             // no content older then 7 days as we just insert it
             Assert.AreEqual(0, emptyRecords.Length);
 
             // check content inserted in the last hour, which should return the 3 objects set as inserted a day ago, but not the one that just got inserted
-            OldRecord[] oldRecords = await refsStore.GetOldRecords(TestNamespace, TimeSpan.FromHours(1)).ToArrayAsync();
+            OldRecord[] oldRecords = await refsStore.GetOldRecords(LastAccessTestNamespace, TimeSpan.FromHours(1)).ToArrayAsync();
 
             List<KeyId> validRefs = new List<KeyId>()
             {
