Horde Storage

Fixed a issue when fetching old ref records from dynamo, were objects of other namespaces were returned.

diff --git a/Engine/Source/Programs/Horde/HordeStorage/Horde.Storage/Implementation/Refs/DynamoDbRefsStore.cs b/Engine/Source/Programs/Horde/HordeStorage/Horde.Storage/Implementation/Refs/DynamoDbRefsStore.cs
--- a/Engine/Source/Programs/Horde/HordeStorage/Horde.Storage/Implementation/Refs/DynamoDbRefsStore.cs
+++ b/Engine/Source/Programs/Horde/HordeStorage/Horde.Storage/Implementation/Refs/DynamoDbRefsStore.cs
@@ -252,6 +252,8 @@ namespace Horde.Storage.Implementation
                     DynamoBaseRefRecord? baseRecord = await context.LoadAsync<DynamoBaseRefRecord>(parentKey);
                     if (baseRecord == null || baseRecord.Namespace == null || baseRecord.Bucket == null || baseRecord.Name == null)
                         continue;
+                    if (baseRecord.Namespace != ns.ToString())
+                        continue;
                     yield return new OldRecord(new NamespaceId(baseRecord.Namespace), new BucketId(baseRecord.Bucket), new KeyId(baseRecord.Name));
                 }
             } while (!search.IsDone);
