Horde: Add additional MaxCreateTime parameter to specialized find jobs method

diff --git a/Engine/Source/Programs/Horde/Horde.Build/Controllers/JobsController.cs b/Engine/Source/Programs/Horde/Horde.Build/Controllers/JobsController.cs
--- a/Engine/Source/Programs/Horde/Horde.Build/Controllers/JobsController.cs
+++ b/Engine/Source/Programs/Horde/Horde.Build/Controllers/JobsController.cs
@@ -707,6 +707,7 @@ namespace HordeServer.Controllers
 		/// <param name="StreamId">The stream to search for</param>
 		/// <param name="Templates">List of templates to find</param>
 		/// <param name="PreflightStartedByUserId">User id for which to include preflight jobs</param>
+		/// <param name="MaxCreateTime">Maximum creation time</param>
 		/// <param name="ModifiedAfter">If specified, only jobs updated after the given time will be returned</param>
 		/// <param name="Filter">Filter for properties to return</param>
 		/// <param name="Index">Index of the first result to be returned</param>
@@ -720,6 +721,7 @@ namespace HordeServer.Controllers
 			string StreamId,
 			[FromQuery(Name = "template")] string[] Templates,
 			[FromQuery] string? PreflightStartedByUserId = null,
+			[FromQuery] DateTimeOffset? MaxCreateTime = null,
 			[FromQuery] DateTimeOffset? ModifiedAfter = null,
 			[FromQuery] PropertyFilter? Filter = null,
 			[FromQuery] int Index = 0,
@@ -731,7 +733,7 @@ namespace HordeServer.Controllers
 			UserId? PreflightStartedByUserIdValue = PreflightStartedByUserId != null ? new UserId(PreflightStartedByUserId) : null;
 			Count = Math.Min(1000, Count);
 
-			List<IJob> Jobs = await JobService.FindJobsByStreamWithTemplatesAsync(StreamIdValue, TemplateRefIds, PreflightStartedByUserIdValue, ModifiedAfter, Index, Count, ConsistentRead);
+			List<IJob> Jobs = await JobService.FindJobsByStreamWithTemplatesAsync(StreamIdValue, TemplateRefIds, PreflightStartedByUserIdValue, MaxCreateTime, ModifiedAfter, Index, Count, ConsistentRead);
 			return await CreateAuthorizedJobResponses(Jobs, Filter);
 		}
 
diff --git a/Engine/Source/Programs/Horde/Horde.Build/Jobs/IJobCollection.cs b/Engine/Source/Programs/Horde/Horde.Build/Jobs/IJobCollection.cs
--- a/Engine/Source/Programs/Horde/Horde.Build/Jobs/IJobCollection.cs
+++ b/Engine/Source/Programs/Horde/Horde.Build/Jobs/IJobCollection.cs
@@ -115,11 +115,12 @@ namespace HordeServer.Collections
 		/// <param name="StreamId">The stream containing the job</param>
 		/// <param name="Templates">Templates to look for</param>
 		/// <param name="PreflightStartedByUser">User for which to include preflight jobs</param>
+		/// <param name="MaxCreateTime">The maximum creation time</param>
 		/// <param name="ModifiedAfter">Filter the results by modified time</param>
 		/// <param name="Index">Index of the first result to return</param>
 		/// <param name="Count">Number of results to return</param>
 		/// <param name="ConsistentRead">If the database read should be made to the replica server</param>
-		Task<List<IJob>> FindLatestByStreamWithTemplatesAsync(StreamId StreamId, TemplateRefId[] Templates, UserId? PreflightStartedByUser = null, DateTimeOffset? ModifiedAfter = null, int? Index = null, int? Count = null, bool ConsistentRead = false);
+		Task<List<IJob>> FindLatestByStreamWithTemplatesAsync(StreamId StreamId, TemplateRefId[] Templates, UserId? PreflightStartedByUser = null, DateTimeOffset? MaxCreateTime = null, DateTimeOffset? ModifiedAfter = null, int? Index = null, int? Count = null, bool ConsistentRead = false);
 
 		/// <summary>
 		/// Updates a new job
diff --git a/Engine/Source/Programs/Horde/Horde.Build/Jobs/Impl/JobCollection.cs b/Engine/Source/Programs/Horde/Horde.Build/Jobs/Impl/JobCollection.cs
--- a/Engine/Source/Programs/Horde/Horde.Build/Jobs/Impl/JobCollection.cs
+++ b/Engine/Source/Programs/Horde/Horde.Build/Jobs/Impl/JobCollection.cs
@@ -623,13 +623,16 @@ namespace HordeServer.Collections.Impl
 		}
 
 		/// <inheritdoc/>
-		public async Task<List<IJob>> FindLatestByStreamWithTemplatesAsync(StreamId StreamId, TemplateRefId[] Templates, UserId? PreflightStartedByUser, DateTimeOffset? ModifiedAfter, int? Index, int? Count, bool ConsistentRead)
+		public async Task<List<IJob>> FindLatestByStreamWithTemplatesAsync(StreamId StreamId, TemplateRefId[] Templates, UserId? PreflightStartedByUser, DateTimeOffset? MaxCreateTime, DateTimeOffset? ModifiedAfter, int? Index, int? Count, bool ConsistentRead)
 		{
+			string IndexHint = Indexes.CreateTimeUtc;
+			if (ModifiedAfter != null) IndexHint = Indexes.UpdateTimeUtc;
+			
 			// This find call uses an index hint. Modifying the parameter passed to FindAsync can affect execution time a lot as the query planner is forced to use the specified index.
 			// Casting to interface to benefit from default parameter values
 			return await (this as IJobCollection).FindAsync(
-				StreamId: StreamId, Templates: Templates, PreflightStartedByUser: PreflightStartedByUser, ModifiedAfter: ModifiedAfter,
-				Index: Index, Count: Count, IndexHint: ModifiedAfter == null ? Indexes.CreateTimeUtc : Indexes.UpdateTimeUtc, ConsistentRead: ConsistentRead);
+				StreamId: StreamId, Templates: Templates, PreflightStartedByUser: PreflightStartedByUser, ModifiedAfter: ModifiedAfter, MaxCreateTime: MaxCreateTime,
+				Index: Index, Count: Count, IndexHint: IndexHint, ConsistentRead: ConsistentRead);
 		}
 
 		/// <inheritdoc/>
diff --git a/Engine/Source/Programs/Horde/Horde.Build/Jobs/Impl/JobService.cs b/Engine/Source/Programs/Horde/Horde.Build/Jobs/Impl/JobService.cs
--- a/Engine/Source/Programs/Horde/Horde.Build/Jobs/Impl/JobService.cs
+++ b/Engine/Source/Programs/Horde/Horde.Build/Jobs/Impl/JobService.cs
@@ -612,22 +612,24 @@ namespace HordeServer.Services
 		/// <param name="StreamId">The stream containing the job</param>
 		/// <param name="Templates">Templates to look for</param>
 		/// <param name="PreflightStartedByUser">User for which to include preflight jobs</param>
+		/// <param name="MaxCreateTime">The maximum creation time</param>
 		/// <param name="ModifiedAfter">Filter the results by last modified time</param>	
 		/// <param name="Index">Index of the first result to return</param>
 		/// <param name="Count">Number of results to return</param>
 		/// <param name="ConsistentRead">If the database read should be made to the replica server</param>
 		/// <returns>List of jobs matching the given criteria</returns>
-		public async Task<List<IJob>> FindJobsByStreamWithTemplatesAsync(StreamId StreamId, TemplateRefId[] Templates, UserId? PreflightStartedByUser = null, DateTimeOffset? ModifiedAfter = null, int? Index = null, int? Count = null, bool ConsistentRead = true)
+		public async Task<List<IJob>> FindJobsByStreamWithTemplatesAsync(StreamId StreamId, TemplateRefId[] Templates, UserId? PreflightStartedByUser = null, DateTimeOffset? MaxCreateTime = null, DateTimeOffset? ModifiedAfter = null, int? Index = null, int? Count = null, bool ConsistentRead = true)
 		{
 			using IScope Scope = GlobalTracer.Instance.BuildSpan("JobService.FindJobsByStreamWithTemplatesAsync").StartActive();
 			Scope.Span.SetTag("StreamId", StreamId);
 			Scope.Span.SetTag("Templates", Templates);
 			Scope.Span.SetTag("PreflightStartedByUser", PreflightStartedByUser);
+			Scope.Span.SetTag("MaxCreateTime", MaxCreateTime);
 			Scope.Span.SetTag("ModifiedAfter", ModifiedAfter);
 			Scope.Span.SetTag("Index", Index);
 			Scope.Span.SetTag("Count", Count);
 			
-			return await Jobs.FindLatestByStreamWithTemplatesAsync(StreamId, Templates, PreflightStartedByUser, ModifiedAfter, Index, Count, ConsistentRead);
+			return await Jobs.FindLatestByStreamWithTemplatesAsync(StreamId, Templates, PreflightStartedByUser, MaxCreateTime, ModifiedAfter, Index, Count, ConsistentRead);
 		}
 
 		/// <summary>
