Horde: Add issue V2 endpoint and hook it up to dashboard with stream severity and resolution state
#rnx

diff --git a/Engine/Source/Programs/Horde/Horde.Build/Api/Issues.cs b/Engine/Source/Programs/Horde/Horde.Build/Api/Issues.cs
--- a/Engine/Source/Programs/Horde/Horde.Build/Api/Issues.cs
+++ b/Engine/Source/Programs/Horde/Horde.Build/Api/Issues.cs
@@ -570,6 +570,191 @@ namespace HordeServer.Api
 		}
 	}
 
+	/// <summary>
+	/// Information about a span within an issue
+	/// </summary>
+	public class FindIssueSpanResponse
+	{
+		/// <summary>
+		/// Unique id of this span
+		/// </summary>
+		public string Id { get; set; }
+
+		/// <summary>
+		/// The template containing this step
+		/// </summary>
+		public string TemplateId { get; set; }
+
+		/// <summary>
+		/// Name of the step
+		/// </summary>
+		public string Name { get; set; }
+
+		/// <summary>
+		/// The previous build 
+		/// </summary>
+		public GetIssueStepResponse? LastSuccess { get; set; }
+
+		/// <summary>
+		/// The following successful build
+		/// </summary>
+		public GetIssueStepResponse? NextSuccess { get; set; }
+
+		/// <summary>
+		/// Constructor
+		/// </summary>
+		/// <param name="Span"></param>
+		public FindIssueSpanResponse(IIssueSpan Span)
+		{
+			this.Id = Span.Id.ToString();
+			this.TemplateId = Span.TemplateRefId.ToString();
+			this.Name = Span.NodeName;
+			if (Span.LastSuccess != null)
+			{
+				this.LastSuccess = new GetIssueStepResponse(Span.LastSuccess);
+			}
+			if (Span.NextSuccess != null)
+			{
+				this.NextSuccess = new GetIssueStepResponse(Span.NextSuccess);
+			}
+		}
+	}
+
+	/// <summary>
+	/// Stores information about a build health issue
+	/// </summary>
+	public class FindIssueResponse
+	{
+		/// <summary>
+		/// The unique object id
+		/// </summary>
+		public int Id { get; set; }
+
+		/// <summary>
+		/// Time at which the issue was created
+		/// </summary>
+		public DateTime CreatedAt { get; set; }
+
+		/// <summary>
+		/// Time at which the issue was retrieved
+		/// </summary>
+		public DateTime RetrievedAt { get; set; }
+
+		/// <summary>
+		/// The associated project for the issue
+		/// </summary>
+		public string? Project { get; set; }
+
+		/// <summary>
+		/// The summary text for this issue
+		/// </summary>
+		public string Summary { get; set; }
+
+		/// <summary>
+		/// Detailed description text
+		/// </summary>
+		public string? Description { get; set; }
+
+		/// <summary>
+		/// Severity of this issue
+		/// </summary>
+		public IssueSeverity Severity { get; set; }
+
+		/// <summary>
+		/// Severity of this issue in the stream
+		/// </summary>
+		public IssueSeverity? StreamSeverity { get; set; }
+
+		/// <summary>
+		/// Whether the issue is promoted
+		/// </summary>
+		public bool Promoted { get; set; }
+
+		/// <summary>
+		/// Owner of the issue
+		/// </summary>
+		public GetThinUserInfoResponse? Owner { get; set; }
+
+		/// <summary>
+		/// Owner of the issue
+		/// </summary>
+		public GetThinUserInfoResponse? NominatedBy { get; set; }
+
+		/// <summary>
+		/// Time that the issue was acknowledged
+		/// </summary>
+		public DateTime? AcknowledgedAt { get; set; }
+
+		/// <summary>
+		/// Changelist that fixed this issue
+		/// </summary>
+		public int? FixChange { get; set; }
+
+		/// <summary>
+		/// Time at which the issue was resolved
+		/// </summary>
+		public DateTime? ResolvedAt { get; set; }
+
+		/// <summary>
+		/// User that resolved the issue
+		/// </summary>
+		public GetThinUserInfoResponse? ResolvedBy { get; set; }
+
+		/// <summary>
+		/// Time at which the issue was verified
+		/// </summary>
+		public DateTime? VerifiedAt { get; set; }
+
+		/// <summary>
+		/// Time that the issue was last seen
+		/// </summary>
+		public DateTime LastSeenAt { get; set; }
+
+		/// <summary>
+		/// Spans for this issue
+		/// </summary>
+		public List<FindIssueSpanResponse> Spans { get; set; }
+
+		/// <summary>
+		/// Constructs a new issue
+		/// </summary>
+		/// <param name="Issue">The isseu information</param>
+		/// <param name="Owner">Owner of the issue</param>
+		/// <param name="NominatedBy">User that nominated the current fixer</param>
+		/// <param name="ResolvedBy">User that resolved the issue</param>
+		/// <param name="StreamSeverity">The current severity in the stream</param>
+		/// <param name="Spans">Spans for this issue</param>
+		public FindIssueResponse(IIssue Issue, IUser? Owner, IUser? NominatedBy, IUser? ResolvedBy, IssueSeverity? StreamSeverity, List<FindIssueSpanResponse> Spans)
+		{
+			this.Id = Issue.Id;
+			this.CreatedAt = Issue.CreatedAt;
+			this.RetrievedAt = DateTime.UtcNow;
+			this.Summary = String.IsNullOrEmpty(Issue.UserSummary) ? Issue.Summary : Issue.UserSummary;
+			this.Description = Issue.Description;
+			this.Severity = Issue.Severity;
+			this.StreamSeverity = StreamSeverity;
+			this.Promoted = Issue.Promoted;
+			if (Owner != null)
+			{
+				this.Owner = new GetThinUserInfoResponse(Owner);
+			}
+			if (NominatedBy != null)
+			{
+				this.NominatedBy = new GetThinUserInfoResponse(NominatedBy);
+			}
+			this.AcknowledgedAt = Issue.AcknowledgedAt;
+			this.FixChange = Issue.FixChange;
+			this.ResolvedAt = Issue.ResolvedAt;
+			if (ResolvedBy != null)
+			{
+				this.ResolvedBy = new GetThinUserInfoResponse(ResolvedBy);
+			}
+			this.VerifiedAt = Issue.VerifiedAt;
+			this.LastSeenAt = Issue.LastSeenAt;
+			this.Spans = Spans;
+		}
+	}
+
 	/// <summary>
 	/// Request an issue to be updated
 	/// </summary>
diff --git a/Engine/Source/Programs/Horde/Horde.Build/Controllers/IssuesController.cs b/Engine/Source/Programs/Horde/Horde.Build/Controllers/IssuesController.cs
--- a/Engine/Source/Programs/Horde/Horde.Build/Controllers/IssuesController.cs
+++ b/Engine/Source/Programs/Horde/Horde.Build/Controllers/IssuesController.cs
@@ -56,6 +56,115 @@ namespace HordeServer.Controllers
 			this.LogFileService = LogFileService;
 		}
 
+		/// <summary>
+		/// Retrieve information about a specific issue
+		/// </summary>
+		/// <param name="Ids">Set of issue ids to find</param>
+		/// <param name="StreamId">The stream to query for</param>
+		/// <param name="MinChange">The minimum changelist range to query, inclusive</param>
+		/// <param name="MaxChange">The minimum changelist range to query, inclusive</param>
+		/// <param name="Resolved">Whether to include resolved issues</param>
+		/// <param name="Index">Starting offset of the window of results to return</param>
+		/// <param name="Count">Number of results to return</param>
+		/// <param name="Filter">Filter for the properties to return</param>
+		/// <returns>List of matching agents</returns>
+		[HttpGet]
+		[Route("/api/v2/issues")]
+		[ProducesResponseType(typeof(List<GetIssueResponse>), 200)]
+		public async Task<ActionResult<object>> FindIssuesV2Async([FromQuery(Name = "Id")] int[]? Ids = null, [FromQuery] StreamId? StreamId = null, [FromQuery] int? MinChange = null, [FromQuery] int? MaxChange = null, [FromQuery] bool? Resolved = null, [FromQuery] int Index = 0, [FromQuery] int Count = 10, [FromQuery] PropertyFilter? Filter = null)
+		{
+			if (Ids != null && Ids.Length == 0)
+			{
+				Ids = null;
+			}
+
+			List<object> Responses = new List<object>();
+			if (StreamId != null)
+			{
+				if (!await StreamService.AuthorizeAsync(StreamId.Value, AclAction.ViewStream, User, new StreamPermissionsCache()))
+				{
+					return Forbid();
+				}
+
+				List<IIssueSpan> Spans = await IssueCollection.FindSpansAsync(null, Ids, StreamId.Value, MinChange, MaxChange, Resolved);
+				if(Spans.Count > 0)
+				{
+					// Group all the spans by their issue id
+					Dictionary<int, List<IIssueSpan>> IssueIdToSpans = new Dictionary<int, List<IIssueSpan>>();
+					foreach (IIssueSpan Span in Spans)
+					{
+						List<IIssueSpan>? SpansForIssue;
+						if (!IssueIdToSpans.TryGetValue(Span.IssueId, out SpansForIssue))
+						{
+							SpansForIssue = new List<IIssueSpan>();
+							IssueIdToSpans.Add(Span.IssueId, SpansForIssue);
+						}
+						SpansForIssue.Add(Span);
+					}
+
+					// Find the matching issues
+					List<IIssue> Issues = await IssueCollection.FindIssuesAsync(IssueIdToSpans.Keys, Index: Index, Count: Count);
+
+					// Create the corresponding responses
+					foreach (IIssue Issue in Issues.OrderByDescending(x => x.Id))
+					{
+						IssueSeverity StreamSeverity = IssueSeverity.Unspecified;
+
+						List<FindIssueSpanResponse> SpanResponses = new List<FindIssueSpanResponse>();
+						if (IssueIdToSpans.TryGetValue(Issue.Id, out List<IIssueSpan>? SpansForIssue))
+						{
+							// Filter issues on resolved state
+							if (Resolved != null && (Resolved.Value != SpansForIssue.All(x => x.NextSuccess != null)))
+							{
+								continue;
+							}
+
+							// Find the current severity in the stream
+							DateTime LastStepTime = DateTime.MinValue;
+							foreach (IIssueSpan Span in SpansForIssue)
+							{
+								if (Span.LastFailure != null && Span.LastFailure.StepTime > LastStepTime)
+								{
+									LastStepTime = Span.LastFailure.StepTime;
+									StreamSeverity = Span.LastFailure.Severity;
+								}
+							}
+
+							SpanResponses.AddRange(SpansForIssue.Select(x => new FindIssueSpanResponse(x)));
+						}
+					
+						IUser? Owner = null;
+						IUser? NominatedBy = null;
+						IUser? ResolvedBy = null;
+
+						if (Issue.OwnerId != null)
+						{
+							Owner = await UserCollection.GetCachedUserAsync(Issue.OwnerId.Value);
+						}
+						
+						if (Issue.NominatedById != null)
+						{
+							NominatedBy = await UserCollection.GetCachedUserAsync(Issue.NominatedById.Value);
+						}
+						
+						if (Issue.ResolvedById != null)
+						{
+							ResolvedBy = await UserCollection.GetCachedUserAsync(Issue.ResolvedById.Value);
+						}						
+
+						FindIssueResponse Response = new FindIssueResponse(Issue, Owner, NominatedBy, ResolvedBy, StreamSeverity, SpanResponses);
+						Responses.Add(PropertyFilter.Apply(Response, Filter));
+					}
+				}
+			}
+			else
+			{
+				return BadRequest("Missing StreamId on request");
+			}
+
+			return Responses;
+		}
+
 		/// <summary>
 		/// Retrieve information about a specific issue
 		/// </summary>
diff --git a/Engine/Source/Programs/Horde/Horde.Build/Issues/IIssueCollection.cs b/Engine/Source/Programs/Horde/Horde.Build/Issues/IIssueCollection.cs
--- a/Engine/Source/Programs/Horde/Horde.Build/Issues/IIssueCollection.cs
+++ b/Engine/Source/Programs/Horde/Horde.Build/Issues/IIssueCollection.cs
@@ -512,6 +512,20 @@ namespace HordeServer.Collections
 		/// <returns>List of open issues</returns>
 		Task<List<IIssueSpan>> FindOpenSpansAsync(StreamId StreamId, TemplateRefId TemplateId, string Name, int Change);
 
+		/// <summary>
+		/// Searches for open issues
+		/// </summary>
+		/// <param name="Ids">Set of issue ids to find</param>
+		/// <param name="IssueIds">The issue ids to retrieve spans for</param>
+		/// <param name="StreamId">The stream affected by the issue</param>
+		/// <param name="MinChange">Minimum changelist affected by the issue</param>
+		/// <param name="MaxChange">Maximum changelist affected by the issue</param>
+		/// <param name="Resolved">Include issues that are now resolved</param>
+		/// <param name="Index">Index within the results to return</param>
+		/// <param name="Count">Number of results</param>
+		/// <returns>List of streams open in the given stream at the given changelist</returns>
+		Task<List<IIssueSpan>> FindSpansAsync(IEnumerable<ObjectId>? Ids = null, IEnumerable<int>? IssueIds = null, StreamId? StreamId = null, int? MinChange = null, int? MaxChange = null, bool? Resolved = null, int? Index = null, int? Count = null);
+
 		#endregion
 
 		#region Steps
diff --git a/Engine/Source/Programs/Horde/Horde.Build/Issues/Impl/IssueCollection.cs b/Engine/Source/Programs/Horde/Horde.Build/Issues/Impl/IssueCollection.cs
--- a/Engine/Source/Programs/Horde/Horde.Build/Issues/Impl/IssueCollection.cs
+++ b/Engine/Source/Programs/Horde/Horde.Build/Issues/Impl/IssueCollection.cs
@@ -1143,6 +1143,48 @@ namespace HordeServer.Collections.Impl
 			return Spans.ConvertAll<IIssueSpan>(x => x);
 		}
 
+		/// <inheritdoc/>
+		public Task<List<IIssueSpan>> FindSpansAsync(IEnumerable<ObjectId>? SpanIds, IEnumerable<int>? IssueIds, StreamId? StreamId, int? MinChange, int? MaxChange, bool? Resolved, int? Index, int? Count)
+		{
+			FilterDefinition<IssueSpan> Filter = FilterDefinition<IssueSpan>.Empty;
+
+			if(SpanIds != null)
+			{
+				Filter &= Builders<IssueSpan>.Filter.In(x => x.Id, SpanIds);
+			}
+
+			if (StreamId != null)
+			{
+				Filter &= Builders<IssueSpan>.Filter.Eq(x => x.StreamId, StreamId);
+			}
+
+			if (IssueIds != null)
+			{
+				Filter &= Builders<IssueSpan>.Filter.In(x => x.IssueId, IssueIds.Select<int, int?>(x => x));
+			}
+			if (MinChange != null)
+			{
+				Filter &= Builders<IssueSpan>.Filter.Not(Builders<IssueSpan>.Filter.Lt(x => x.MaxChange, MinChange.Value));
+			}
+			if (MaxChange != null)
+			{
+				Filter &= Builders<IssueSpan>.Filter.Not(Builders<IssueSpan>.Filter.Gt(x => x.MinChange, MaxChange.Value));
+			}
+			if (Resolved != null)
+			{
+				if (Resolved.Value)
+				{
+					Filter &= Builders<IssueSpan>.Filter.Ne(x => x.ResolvedAt, null);
+				}
+				else
+				{
+					Filter &= Builders<IssueSpan>.Filter.Eq(x => x.ResolvedAt, null);
+				}
+			}
+
+			return IssueSpans.Find(Filter).Range(Index, Count).ToListAsync<IssueSpan, IIssueSpan>();
+		}
+
 		#endregion
 
 		#region Steps
diff --git a/Engine/Source/Programs/Horde/HordeDashboard/src/backend/Api.ts b/Engine/Source/Programs/Horde/HordeDashboard/src/backend/Api.ts
--- a/Engine/Source/Programs/Horde/HordeDashboard/src/backend/Api.ts
+++ b/Engine/Source/Programs/Horde/HordeDashboard/src/backend/Api.ts
@@ -251,6 +251,19 @@ export type LogLineData = {
 	lines?: LogLine[];
 };
 
+export type JobStreamQuery = {
+	filter?: string;
+	template?: string[];
+	preflightStartedByUserId?: string,
+	includePreflight?: boolean,
+	maxCreateTime?: string;
+	modifiedAfter?: string;
+	index?: number;
+	count?: number;
+	consistentRead?: boolean;
+
+};
+
 export type JobQuery = {
 	filter?: string;
 	streamId?: string;
@@ -271,6 +284,7 @@ export type JobQuery = {
 	count?: number;
 };
 
+
 export type IssueQuery = {
 	jobId?: string;
 	batchId?: string;
@@ -287,6 +301,18 @@ export type IssueQuery = {
 	promoted?: boolean;
 }
 
+export type IssueQueryV2 = {
+	id?: string[]; 
+	streamId?: string;
+	minChange?: number;
+	maxChange?: number;
+	resolved?: boolean;
+	index?: number;
+	count?: number;
+	filter?: string;
+}
+
+
 export type UsersQuery = {
 	ids?: string[];
 	nameRegex?: string;
@@ -3413,7 +3439,7 @@ export type NamespaceConfig = {
 	buckets: BucketConfig[];
 
 	/// Access control for this namespace
-	//UpdateAclRequest? Acl { get; set; }
+	//UpdateAclRequest? Acl;
 }
 
 /// Configuration for a bucket
@@ -3444,7 +3470,7 @@ export type GlobalConfig = {
 	storage?: StorageConfig;
 
 	/// Access control list
-	// public UpdateAclRequest ? Acl { get; set; }
+	// public UpdateAclRequest ? Acl;
 }
 
 
@@ -3495,7 +3521,7 @@ export type ProjectConfig = {
 	streams: StreamConfigRef[];
 
 	/// Acl entries
-	// public UpdateAclRequest? Acl { get; set; }
+	// public UpdateAclRequest? Acl;
 }
 
 
@@ -3645,17 +3671,17 @@ export type StreamConfig = {
 	/// <summary>
 	/// Custom permissions for this object
 	/// </summary>
-	// public UpdateAclRequest ? Acl { get; set; }
+	// public UpdateAclRequest ? Acl;
 
 	/// <summary>
 	/// Pause stream builds until specified date
 	/// </summary>
-	/// public DateTime ? PausedUntil { get; set; }
+	/// public DateTime ? PausedUntil;
 
 	/// <summary>
 	/// Reason for pausing builds of the stream
 	/// </summary>
-	// public string ? PauseComment { get; set; }
+	// public string ? PauseComment;
 }
 
 /// Parameters to update server settings
@@ -3693,3 +3719,79 @@ export type GetServerInfoResponse = {
 	singleInstance: boolean;
 }
 
+/// Information about a span within an issue
+export type FindIssueSpanResponse = {
+
+	/// Unique id of this span
+	id: string;
+
+	/// The template containing this step
+	templateId: string;
+
+	/// Name of the step
+	name: string;
+
+	/// The previous build 
+	lastSuccess?: GetIssueStepResponse;
+
+	/// The following successful build
+	nextSuccess?: GetIssueStepResponse;
+
+}
+
+/// Stores information about a build health issue
+export type FindIssueResponse = {
+	/// The unique object id
+	id: number;
+
+	/// Time at which the issue was created
+	createdAt: Date | string;
+
+	/// Time at which the issue was retrieved
+	RetrievedAt: Date | string;
+
+	/// The associated project for the issue
+	project?: string;
+
+	/// The summary text for this issue
+	summary: string;
+
+	/// Detailed description text
+	description?: string;
+
+	/// Severity of this issue
+	severity: IssueSeverity;
+
+	/// Current severity in the queried stream
+	streamSeverity?: IssueSeverity;
+
+	/// Whether the issue is promoted
+	promoted: boolean;
+
+	/// Owner of the issue
+	owner?: GetThinUserInfoResponse;
+
+	/// Owner of the issue
+	mominatedBy?: GetThinUserInfoResponse;
+
+	/// Time that the issue was acknowledged
+	acknowledgedAt?: Date | string;
+
+	/// Changelist that fixed this issue
+	fixChange?: number;
+
+	/// Time at which the issue was resolved
+	resolvedAt?: Date | string;
+
+	/// User that resolved the issue
+	resolvedBy?: GetThinUserInfoResponse;
+
+	/// Time at which the issue was verified
+	verifiedAt?: Date | string;
+
+	/// Time that the issue was last seen
+	lastSeenAt: Date | string;
+
+	/// Spans for this issue
+	spans: FindIssueSpanResponse[];
+}
diff --git a/Engine/Source/Programs/Horde/HordeDashboard/src/backend/Backend.ts b/Engine/Source/Programs/Horde/HordeDashboard/src/backend/Backend.ts
--- a/Engine/Source/Programs/Horde/HordeDashboard/src/backend/Backend.ts
+++ b/Engine/Source/Programs/Horde/HordeDashboard/src/backend/Backend.ts
@@ -3,7 +3,7 @@
 import { configure } from "mobx";
 import { isNumber } from 'util';
 import templateCache from '../backend/TemplateCache';
-import { AgentData, ArtifactData, AuditLogEntry, AuditLogQuery, BatchUpdatePoolRequest, ChangeSummaryData, CreateDeviceRequest, CreateDeviceResponse, CreateJobRequest, CreateJobResponse, CreatePoolRequest, CreateSoftwareResponse, CreateSubscriptionRequest, CreateSubscriptionResponse, DashboardPreference, EventData, GetArtifactZipRequest, GetDevicePlatformResponse, GetDevicePoolResponse, GetDeviceReservationResponse, GetDeviceResponse, GetGraphResponse, GetIssueStreamResponse, GetJobsTabResponse, GetJobStepRefResponse, GetJobStepTraceResponse, GetJobTimingResponse, GetLogEventResponse, GetNotificationResponse, GetPerforceServerStatusResponse, GetServerInfoResponse, GetServerSettingsResponse, GetSoftwareResponse, GetSubscriptionResponse, GetUserResponse, GetUtilizationTelemetryResponse, GlobalConfig, IssueData, IssueQuery, JobData, JobQuery, JobsTabColumnType, LeaseData, LogData, LogLineData, PoolData, ProjectData, ScheduleData, ScheduleQuery, SearchLogFileResponse, ServerUpdateResponse, SessionData, StreamData, TabType, TemplateData, TestData, UpdateAgentRequest, UpdateDeviceRequest, UpdateGlobalConfigRequest, UpdateIssueRequest, UpdateJobRequest, UpdateLeaseRequest, UpdateNotificationsRequest, UpdatePoolRequest, UpdateServerSettingsRequest, UpdateStepRequest, UpdateStepResponse, UpdateUserRequest, UsersQuery } from './Api';
+import { AgentData, ArtifactData, AuditLogEntry, AuditLogQuery, BatchUpdatePoolRequest, ChangeSummaryData, CreateDeviceRequest, CreateDeviceResponse, CreateJobRequest, CreateJobResponse, CreatePoolRequest, CreateSoftwareResponse, CreateSubscriptionRequest, CreateSubscriptionResponse, DashboardPreference, EventData, FindIssueResponse, GetArtifactZipRequest, GetDevicePlatformResponse, GetDevicePoolResponse, GetDeviceReservationResponse, GetDeviceResponse, GetGraphResponse, GetIssueStreamResponse, GetJobsTabResponse, GetJobStepRefResponse, GetJobStepTraceResponse, GetJobTimingResponse, GetLogEventResponse, GetNotificationResponse, GetPerforceServerStatusResponse, GetServerInfoResponse, GetServerSettingsResponse, GetSoftwareResponse, GetSubscriptionResponse, GetUserResponse, GetUtilizationTelemetryResponse, GlobalConfig, IssueData, IssueQuery, IssueQueryV2, JobData, JobQuery, JobsTabColumnType, JobStreamQuery, LeaseData, LogData, LogLineData, PoolData, ProjectData, ScheduleData, ScheduleQuery, SearchLogFileResponse, ServerUpdateResponse, SessionData, StreamData, TabType, TemplateData, TestData, UpdateAgentRequest, UpdateDeviceRequest, UpdateGlobalConfigRequest, UpdateIssueRequest, UpdateJobRequest, UpdateLeaseRequest, UpdateNotificationsRequest, UpdatePoolRequest, UpdateServerSettingsRequest, UpdateStepRequest, UpdateStepResponse, UpdateUserRequest, UsersQuery } from './Api';
 import dashboard from './Dashboard';
 import { ChallengeStatus, Fetch } from './Fetch';
 import graphCache, { GraphQuery } from './GraphCache';
@@ -325,6 +325,55 @@ export class Backend {
 
     }
 
+    getStreamJobs(streamId: string, query: JobStreamQuery, queryGraph: boolean = false): Promise<JobData[]> {
+
+        if (typeof query.index === 'number') {
+            query.index = 0;
+        }
+
+        if (typeof query.count !== 'number') {
+            query.count = 100;
+        }
+
+        return new Promise<JobData[]>((resolve, reject) => {
+
+            this.backend.get(`/api/v1/jobs/streams/${streamId}`, {
+                params: query
+            }).then((value) => {
+                const jobs = value.data as JobData[];
+
+                if (!queryGraph) {
+                    resolve(jobs);
+                    return;
+                }
+
+                const query: GraphQuery[] = [];
+                jobs.forEach(j => {
+                    if (!j.graphHash) {
+                        console.error(`Job ${j.id} has no graph hash`);
+                        return;
+                    }
+
+                    if (!query.find(q => q.graphHash === j.graphHash)) {
+                        query.push({ graphHash: j.graphHash!, jobId: j.id });
+                    }
+                });
+
+                graphCache.getGraphs(query).then(graphs => {
+                    jobs.forEach(j => {
+                        j.graphRef = graphs.find(g => g.hash === j.graphHash);
+                    });
+                    resolve(jobs);
+                }).catch(reason => reject(reason));
+
+            }).catch((reason) => {
+                reject(reason);
+            });
+        });
+
+    }
+
+
     getJobs(query: JobQuery, queryGraph: boolean = false): Promise<JobData[]> {
 
         if (typeof query.index === 'number') {
@@ -867,6 +916,23 @@ export class Backend {
     }
 
 
+    getIssuesV2(queryIn?: IssueQueryV2): Promise<FindIssueResponse[]> {
+        
+        const query = queryIn ?? {};
+
+        return new Promise<FindIssueResponse[]>((resolve, reject) => {
+
+            this.backend.get("/api/v2/issues", {
+                params: query
+            }).then((value) => {
+                resolve(value.data as FindIssueResponse[]);
+            }).catch((reason) => {
+                reject(reason);
+            });
+        });
+
+    }
+
     getIssues(queryIn?: IssueQuery): Promise<IssueData[]> {
 
         const query = queryIn ?? {};
diff --git a/Engine/Source/Programs/Horde/HordeDashboard/src/backend/JobHandler.ts b/Engine/Source/Programs/Horde/HordeDashboard/src/backend/JobHandler.ts
--- a/Engine/Source/Programs/Horde/HordeDashboard/src/backend/JobHandler.ts
+++ b/Engine/Source/Programs/Horde/HordeDashboard/src/backend/JobHandler.ts
@@ -3,14 +3,14 @@
 import { action, observable } from 'mobx';
 import moment from 'moment-timezone';
 import backend from '../backend';
-import { JobData, JobQuery, JobState, JobStepOutcome, StreamData } from '../backend/Api';
+import { JobData, JobState, JobStepOutcome, JobStreamQuery, StreamData } from '../backend/Api';
 import graphCache, { GraphQuery } from '../backend/GraphCache';
 
 
 export type FilterStatus = "Running" | "Complete" | "Succeeded" | "Failed" | "Waiting";
 
 const jobPageSize = 50;
-const jobsRefreshTime = 3000;
+const jobsRefreshTime = 10000;
 
 export class JobHandler {
 
@@ -85,10 +85,9 @@ export class JobHandler {
             }
 
             // needs modified time
-            const query: JobQuery = {
+            const query: JobStreamQuery = {
                 filter: filter,
                 count: this.count,
-                streamId: this.stream?.id,
                 template: this.templateNames,
                 includePreflight: this.includePreflights,
                 preflightStartedByUserId: this.preflightStartedByUserId
@@ -119,21 +118,13 @@ export class JobHandler {
 
                 query.modifiedAfter = this.modifiedAfter;
 
-            }
-
-            // Set a min time as a TEMPORARY hint which greatly speeds up *all* job query requests
-            // @todo https://jira.it.epicgames.com/browse/UE-135429
-
-            if (!query.minCreateTime) {
-                query.minCreateTime = new Date('2010-1-1').toISOString();
-            }
-            
+            }            
 
             const cancelId = this.cancelId++;
 
             const queryTime = moment.utc().toISOString();
 
-            const mjobs = await backend.getJobs(query, false);
+            const mjobs = await backend.getStreamJobs(this.stream!.id, query, false);
 
             // check for canceled after modified test
             if (this.canceled.has(cancelId)) {
diff --git a/Engine/Source/Programs/Horde/HordeDashboard/src/base/utilities/issueUtils.ts b/Engine/Source/Programs/Horde/HordeDashboard/src/base/utilities/issueUtils.ts
--- a/Engine/Source/Programs/Horde/HordeDashboard/src/base/utilities/issueUtils.ts
+++ b/Engine/Source/Programs/Horde/HordeDashboard/src/base/utilities/issueUtils.ts
@@ -13,15 +13,15 @@ export type IssueJira = {
 // eslint-disable-next-line
 const urlRegex = /([\w+]+\:\/\/)?([\w\d-]+\.)*[\w-]+[\.\:]\w+([\/\?\=\&\#\.]?[\w-]+)*\/?/gm
 
-export function getIssueJiras(issue: GetIssueResponse): IssueJira[] {
+export function getIssueJiras(description?: string): IssueJira[] {
 
-	if (!issue.description) {
+	if (!description) {
 		return [];
 	}
 
 	const jiras: IssueJira[] = [];
 
-	const matches = Array.from(issue.description.matchAll(urlRegex));
+	const matches = Array.from(description.matchAll(urlRegex));
 
 	matches.forEach(m => {
 		m.forEach(match => {
