Add custom TagPicker that renders the input element in the same span as the tag items, rather than below it.
Point AgentView at new component.

diff --git a/Engine/Source/Programs/Horde/HordeDashboard/src/base/components/HordeFluentComponents/HBasePicker/HBasePicker.styles.ts b/Engine/Source/Programs/Horde/HordeDashboard/src/base/components/HordeFluentComponents/HBasePicker/HBasePicker.styles.ts
new file mode 100644
--- /dev/null
+++ b/Engine/Source/Programs/Horde/HordeDashboard/src/base/components/HordeFluentComponents/HBasePicker/HBasePicker.styles.ts
@@ -0,0 +1,106 @@
+import { getGlobalClassNames, getInputFocusStyle, hiddenContentStyle, HighContrastSelector, IBasePickerStyleProps, IBasePickerStyles } from '@fluentui/react';
+
+const GlobalClassNames = {
+  root: 'ms-BasePicker',
+  text: 'ms-BasePicker-text',
+  itemsWrapper: 'ms-BasePicker-itemsWrapper',
+  input: 'ms-BasePicker-input',
+};
+
+export function getStyles(props: IBasePickerStyleProps): IBasePickerStyles {
+  const { className, theme, isFocused, inputClassName, disabled } = props;
+
+  if (!theme) {
+    throw new Error('theme is undefined or null in base BasePicker getStyles function.');
+  }
+  const { semanticColors, effects, fonts } = theme;
+  const { inputBorder, inputBorderHovered, inputFocusBorderAlt } = semanticColors;
+
+  const classNames = getGlobalClassNames(GlobalClassNames, theme);
+
+  // The following lines are to create a semi-transparent color overlay for the disabled state with designer's approval.
+  // @todo: investigate the performance cost of the calculation below and apply if negligible.
+  //   Replacing with a static color for now.
+  // const rgbColor: IRGB | undefined = cssColor(palette.neutralQuaternaryAlt);
+  // const disabledOverlayColor = rgbColor ? `rgba(${rgbColor.r}, ${rgbColor.g}, ${rgbColor.b}, 0.29)` : 'transparent';
+  const disabledOverlayColor = 'rgba(218, 218, 218, 0.29)';
+
+  return {
+    root: [classNames.root, className],
+    text: [
+      classNames.text,
+      {
+        display: 'flex',
+        position: 'relative',
+        flexWrap: 'wrap',
+        alignItems: 'center',
+        boxSizing: 'border-box',
+        minWidth: 180,
+        minHeight: 30,
+        border: `1px solid ${inputBorder}`,
+        borderRadius: effects.roundedCorner2,
+      },
+      !isFocused &&
+        !disabled && {
+          selectors: {
+            ':hover': {
+              borderColor: inputBorderHovered,
+            },
+          },
+        },
+      isFocused && !disabled && getInputFocusStyle(inputFocusBorderAlt, effects.roundedCorner2),
+      disabled && {
+        borderColor: disabledOverlayColor,
+        selectors: {
+          ':after': {
+            content: '""',
+            position: 'absolute',
+            top: 0,
+            right: 0,
+            bottom: 0,
+            left: 0,
+            background: disabledOverlayColor,
+          },
+          [HighContrastSelector]: {
+            borderColor: 'GrayText',
+            selectors: {
+              ':after': {
+                background: 'none',
+              },
+            },
+          },
+        },
+      },
+    ],
+    itemsWrapper: [
+      classNames.itemsWrapper,
+      {
+        display: 'flex',
+        flexWrap: 'wrap',
+        maxWidth: '100%',
+      },
+    ],
+    input: [
+      classNames.input,
+      fonts.medium,
+      {
+        height: 30,
+        border: 'none',
+        flexGrow: 1,
+        outline: 'none',
+        padding: '0 6px 0',
+        alignSelf: 'flex-end',
+        borderRadius: effects.roundedCorner2,
+        backgroundColor: 'transparent',
+        color: semanticColors.inputText,
+        selectors: {
+          '::-ms-clear': {
+            display: 'none',
+          },
+        },
+      },
+      inputClassName,
+    ],
+    screenReaderText: hiddenContentStyle,
+  };
+}
