Horde Storage
Updated refs api to expect all keys to be IoHashes and not arbitrary strings (if you want to map from a arbitrary string hash it and map from the hash)

Most use cases should already be using this as hashes, as this has been the agreed upon api for a while, I just hadnt updated the code to reflect it.

diff --git a/Engine/Source/Programs/Horde/HordeStorage/Horde.Storage/Controllers/ReferencesController.cs b/Engine/Source/Programs/Horde/HordeStorage/Horde.Storage/Controllers/ReferencesController.cs
--- a/Engine/Source/Programs/Horde/HordeStorage/Horde.Storage/Controllers/ReferencesController.cs
+++ b/Engine/Source/Programs/Horde/HordeStorage/Horde.Storage/Controllers/ReferencesController.cs
@@ -90,7 +90,7 @@ namespace Horde.Storage.Controllers
         public async Task<IActionResult> Get(
             [FromRoute] [Required] NamespaceId ns,
             [FromRoute] [Required] BucketId bucket,
-            [FromRoute] [Required] KeyId key,
+            [FromRoute] [Required] IoHashKey key,
             [FromQuery] string[] fields,
             [FromRoute] string? format = null)
         {
@@ -219,7 +219,7 @@ namespace Horde.Storage.Controllers
         public async Task<IActionResult> Head(
             [FromRoute] [Required] NamespaceId ns,
             [FromRoute] [Required] BucketId bucket,
-            [FromRoute] [Required] KeyId key)
+            [FromRoute] [Required] IoHashKey key)
         {
             if (ShouldDoAuth())
             {
@@ -261,7 +261,7 @@ namespace Horde.Storage.Controllers
         public async Task<IActionResult> ExistsMultiple(
             [FromRoute] [Required] NamespaceId ns,
             [FromRoute] [Required] BucketId bucket,
-            [FromQuery] [Required] List<KeyId> names)
+            [FromQuery] [Required] List<IoHashKey> names)
         {
             AuthorizationResult authorizationResult = await _authorizationService.AuthorizeAsync(User, ns, NamespaceAccessRequirement.Name);
 
@@ -270,7 +270,7 @@ namespace Horde.Storage.Controllers
                 return Forbid();
             }
 
-            ConcurrentBag<KeyId> missingObject = new ConcurrentBag<KeyId>();
+            ConcurrentBag<IoHashKey> missingObject = new ConcurrentBag<IoHashKey>();
 
             IEnumerable<Task> tasks = names.Select(async name =>
             {
@@ -295,7 +295,7 @@ namespace Horde.Storage.Controllers
         public async Task<IActionResult> PutObject(
             [FromRoute] [Required] NamespaceId ns,
             [FromRoute] [Required] BucketId bucket,
-            [FromRoute] [Required] KeyId key)
+            [FromRoute] [Required] IoHashKey key)
         {
             if (ShouldDoAuth())
             {
@@ -394,7 +394,7 @@ namespace Horde.Storage.Controllers
         public async Task<IActionResult> FinalizeObject(
             [FromRoute] [Required] NamespaceId ns,
             [FromRoute] [Required] BucketId bucket,
-            [FromRoute] [Required] KeyId key,
+            [FromRoute] [Required] IoHashKey key,
             [FromRoute] [Required] BlobIdentifier hash)
         {
             if (ShouldDoAuth())
@@ -493,7 +493,7 @@ namespace Horde.Storage.Controllers
         public async Task<IActionResult> Delete(
             [FromRoute] [Required] NamespaceId ns,
             [FromRoute] [Required] BucketId bucket,
-            [FromRoute] [Required] KeyId key)
+            [FromRoute] [Required] IoHashKey key)
         {
             AuthorizationResult authorizationResult = await _authorizationService.AuthorizeAsync(User, ns, NamespaceAccessRequirement.Name);
 
diff --git a/Engine/Source/Programs/Horde/HordeStorage/Horde.Storage/Implementation/ObjectService.cs b/Engine/Source/Programs/Horde/HordeStorage/Horde.Storage/Implementation/ObjectService.cs
--- a/Engine/Source/Programs/Horde/HordeStorage/Horde.Storage/Implementation/ObjectService.cs
+++ b/Engine/Source/Programs/Horde/HordeStorage/Horde.Storage/Implementation/ObjectService.cs
@@ -14,13 +14,13 @@ namespace Horde.Storage.Implementation
 {
     public interface IObjectService
     {
-        Task<(ObjectRecord, BlobContents)> Get(NamespaceId ns, BucketId bucket, KeyId key, string[] fields);
-        Task<PutObjectResult> Put(NamespaceId ns, BucketId bucket, KeyId key, BlobIdentifier blobHash, CompactBinaryObject payload);
-        Task<BlobIdentifier[]> Finalize(NamespaceId ns, BucketId bucket, KeyId key, BlobIdentifier blobHash);
+        Task<(ObjectRecord, BlobContents)> Get(NamespaceId ns, BucketId bucket, IoHashKey key, string[] fields);
+        Task<PutObjectResult> Put(NamespaceId ns, BucketId bucket, IoHashKey key, BlobIdentifier blobHash, CompactBinaryObject payload);
+        Task<BlobIdentifier[]> Finalize(NamespaceId ns, BucketId bucket, IoHashKey key, BlobIdentifier blobHash);
 
         IAsyncEnumerator<NamespaceId> GetNamespaces();
 
-        Task<long> Delete(NamespaceId ns, BucketId bucket, KeyId key);
+        Task<long> Delete(NamespaceId ns, BucketId bucket, IoHashKey key);
         Task<long> DropNamespace(NamespaceId ns);
         Task<long> DeleteBucket(NamespaceId ns, BucketId bucket);
     }
@@ -40,7 +40,7 @@ namespace Horde.Storage.Implementation
             _replicationLog = replicationLog;
         }
 
-        public async Task<(ObjectRecord, BlobContents)> Get(NamespaceId ns, BucketId bucket, KeyId key, string[]? fields = null)
+        public async Task<(ObjectRecord, BlobContents)> Get(NamespaceId ns, BucketId bucket, IoHashKey key, string[]? fields = null)
         {
             ObjectRecord o = await _referencesStore.Get(ns, bucket, key);
 
@@ -56,7 +56,7 @@ namespace Horde.Storage.Implementation
             return (o, blobContents);
         }
 
-        public async Task<PutObjectResult> Put(NamespaceId ns, BucketId bucket, KeyId key, BlobIdentifier blobHash, CompactBinaryObject payload)
+        public async Task<PutObjectResult> Put(NamespaceId ns, BucketId bucket, IoHashKey key, BlobIdentifier blobHash, CompactBinaryObject payload)
         {
             bool hasReferences = payload.GetAllFields().Any(field => field.IsAttachment());
 
@@ -93,7 +93,7 @@ namespace Horde.Storage.Implementation
             return new PutObjectResult(missingReferences);
         }
 
-        public async Task<BlobIdentifier[]> Finalize(NamespaceId ns, BucketId bucket, KeyId key, BlobIdentifier blobHash)
+        public async Task<BlobIdentifier[]> Finalize(NamespaceId ns, BucketId bucket, IoHashKey key, BlobIdentifier blobHash)
         {
             (ObjectRecord o, BlobContents blob) = await Get(ns, bucket, key);
             byte[] blobContents = await blob.Stream.ToByteArray();
@@ -130,7 +130,7 @@ namespace Horde.Storage.Implementation
             return _referencesStore.GetNamespaces();
         }
 
-        public Task<long> Delete(NamespaceId ns, BucketId bucket, KeyId key)
+        public Task<long> Delete(NamespaceId ns, BucketId bucket, IoHashKey key)
         {
             return _referencesStore.Delete(ns, bucket, key);
         }
diff --git a/Engine/Source/Programs/Horde/HordeStorage/Horde.Storage/Implementation/Objects/IReferencesStore.cs b/Engine/Source/Programs/Horde/HordeStorage/Horde.Storage/Implementation/Objects/IReferencesStore.cs
--- a/Engine/Source/Programs/Horde/HordeStorage/Horde.Storage/Implementation/Objects/IReferencesStore.cs
+++ b/Engine/Source/Programs/Horde/HordeStorage/Horde.Storage/Implementation/Objects/IReferencesStore.cs
@@ -9,23 +9,23 @@ namespace Horde.Storage.Implementation
 {
     public interface IReferencesStore
     {
-        Task<ObjectRecord> Get(NamespaceId ns, BucketId bucket, KeyId key);
-        Task Put(NamespaceId ns, BucketId bucket, KeyId key, BlobIdentifier blobHash, byte[] blob, bool isFinalized);
-        Task Finalize(NamespaceId ns, BucketId bucket, KeyId key, BlobIdentifier blobHash);
+        Task<ObjectRecord> Get(NamespaceId ns, BucketId bucket, IoHashKey key);
+        Task Put(NamespaceId ns, BucketId bucket, IoHashKey key, BlobIdentifier blobHash, byte[] blob, bool isFinalized);
+        Task Finalize(NamespaceId ns, BucketId bucket, IoHashKey key, BlobIdentifier blobHash);
 
 
-        Task UpdateLastAccessTime(NamespaceId ns, BucketId bucket, KeyId key, DateTime newLastAccessTime);
+        Task UpdateLastAccessTime(NamespaceId ns, BucketId bucket, IoHashKey key, DateTime newLastAccessTime);
         IAsyncEnumerator<ObjectRecord> GetOldestRecords(NamespaceId ns);
 
         IAsyncEnumerator<NamespaceId> GetNamespaces();
-        Task<long> Delete(NamespaceId ns, BucketId bucket, KeyId key);
+        Task<long> Delete(NamespaceId ns, BucketId bucket, IoHashKey key);
         Task<long> DropNamespace(NamespaceId ns);
         Task<long> DeleteBucket(NamespaceId ns, BucketId bucket);
     }
 
     public class ObjectRecord
     {
-        public ObjectRecord(NamespaceId ns, BucketId bucket, KeyId name, DateTime lastAccess, byte[]? inlinePayload, BlobIdentifier blobIdentifier, bool isFinalized)
+        public ObjectRecord(NamespaceId ns, BucketId bucket, IoHashKey name, DateTime lastAccess, byte[]? inlinePayload, BlobIdentifier blobIdentifier, bool isFinalized)
         {
             Namespace = ns;
             Bucket = bucket;
@@ -38,7 +38,7 @@ namespace Horde.Storage.Implementation
 
         public NamespaceId Namespace { get; }
         public BucketId Bucket { get; }
-        public KeyId Name { get; }
+        public IoHashKey Name { get; }
         public DateTime LastAccess { get; }
         public byte[]? InlinePayload { get; set; }
         public BlobIdentifier BlobIdentifier { get; set; }
@@ -47,7 +47,7 @@ namespace Horde.Storage.Implementation
 
     public class ObjectNotFoundException : Exception
     {
-        public ObjectNotFoundException(NamespaceId ns, BucketId bucket, KeyId key) : base($"Object not found {key} in bucket {bucket} namespace {ns}")
+        public ObjectNotFoundException(NamespaceId ns, BucketId bucket, IoHashKey key) : base($"Object not found {key} in bucket {bucket} namespace {ns}")
         {
             Namespace = ns;
             Bucket = bucket;
@@ -56,6 +56,6 @@ namespace Horde.Storage.Implementation
 
         public NamespaceId Namespace { get; }
         public BucketId Bucket { get; }
-        public KeyId Key { get; }
+        public IoHashKey Key { get; }
     }
 }
diff --git a/Engine/Source/Programs/Horde/HordeStorage/Horde.Storage/Implementation/Objects/MemoryReferencesStore.cs b/Engine/Source/Programs/Horde/HordeStorage/Horde.Storage/Implementation/Objects/MemoryReferencesStore.cs
--- a/Engine/Source/Programs/Horde/HordeStorage/Horde.Storage/Implementation/Objects/MemoryReferencesStore.cs
+++ b/Engine/Source/Programs/Horde/HordeStorage/Horde.Storage/Implementation/Objects/MemoryReferencesStore.cs
@@ -20,7 +20,7 @@ namespace Horde.Storage.Implementation
 
         }
 
-        public Task<ObjectRecord> Get(NamespaceId ns, BucketId bucket, KeyId key)
+        public Task<ObjectRecord> Get(NamespaceId ns, BucketId bucket, IoHashKey key)
         {
             if (_objects.TryGetValue(BuildKey(ns, bucket, key), out MemoryStoreObject? o))
             {
@@ -30,7 +30,7 @@ namespace Horde.Storage.Implementation
             throw new ObjectNotFoundException(ns, bucket, key);
         }
 
-        public Task Put(NamespaceId ns, BucketId bucket, KeyId key, BlobIdentifier blobHash, byte[] blob, bool isFinalized)
+        public Task Put(NamespaceId ns, BucketId bucket, IoHashKey key, BlobIdentifier blobHash, byte[] blob, bool isFinalized)
         {
             lock (_namespaces)
             {
@@ -44,7 +44,7 @@ namespace Horde.Storage.Implementation
             return Task.FromResult(o);
         }
 
-        public Task Finalize(NamespaceId ns, BucketId bucket, KeyId key, BlobIdentifier blobHash)
+        public Task Finalize(NamespaceId ns, BucketId bucket, IoHashKey key, BlobIdentifier blobHash)
         {
             if (!_objects.TryGetValue(BuildKey(ns, bucket, key), out MemoryStoreObject? o))
             {
@@ -55,7 +55,7 @@ namespace Horde.Storage.Implementation
             return Task.CompletedTask;
         }
 
-        public Task UpdateLastAccessTime(NamespaceId ns, BucketId bucket, KeyId key, DateTime lastAccessTime)
+        public Task UpdateLastAccessTime(NamespaceId ns, BucketId bucket, IoHashKey key, DateTime lastAccessTime)
         {
             if (!_objects.TryGetValue(BuildKey(ns, bucket, key), out MemoryStoreObject? o))
             {
@@ -80,7 +80,7 @@ namespace Horde.Storage.Implementation
             return _namespaces.GetAsyncEnumerator();
         }
 
-        public Task<long> Delete(NamespaceId ns, BucketId bucket, KeyId key)
+        public Task<long> Delete(NamespaceId ns, BucketId bucket, IoHashKey key)
         {
             if (!_objects.TryRemove(BuildKey(ns, bucket, key), out MemoryStoreObject? o))
             {
@@ -143,7 +143,7 @@ namespace Horde.Storage.Implementation
             return Task.FromResult(removedCount);
         }
 
-        private static string BuildKey(NamespaceId ns, BucketId bucket, KeyId name)
+        private static string BuildKey(NamespaceId ns, BucketId bucket, IoHashKey name)
         {
             return $"{ns}.{bucket}.{name}";
         }
@@ -151,7 +151,7 @@ namespace Horde.Storage.Implementation
 
     public class MemoryStoreObject
     {
-        public MemoryStoreObject(NamespaceId ns, BucketId bucket, KeyId key, BlobIdentifier blobHash, byte[] blob, bool isFinalized)
+        public MemoryStoreObject(NamespaceId ns, BucketId bucket, IoHashKey key, BlobIdentifier blobHash, byte[] blob, bool isFinalized)
         {
             Namespace = ns;
             Bucket = bucket;
@@ -164,7 +164,7 @@ namespace Horde.Storage.Implementation
 
         public NamespaceId Namespace { get; }
         public BucketId Bucket { get; }
-        public KeyId Name { get; }
+        public IoHashKey Name { get; }
         public byte[] Blob { get; }
         public BlobIdentifier BlobHash { get; }
 
diff --git a/Engine/Source/Programs/Horde/HordeStorage/Horde.Storage/Implementation/Objects/ScyllaReferencesStore.cs b/Engine/Source/Programs/Horde/HordeStorage/Horde.Storage/Implementation/Objects/ScyllaReferencesStore.cs
--- a/Engine/Source/Programs/Horde/HordeStorage/Horde.Storage/Implementation/Objects/ScyllaReferencesStore.cs
+++ b/Engine/Source/Programs/Horde/HordeStorage/Horde.Storage/Implementation/Objects/ScyllaReferencesStore.cs
@@ -57,7 +57,7 @@ namespace Horde.Storage.Implementation
 
         }
 
-        public async Task<ObjectRecord> Get(NamespaceId ns, BucketId bucket, KeyId name)
+        public async Task<ObjectRecord> Get(NamespaceId ns, BucketId bucket, IoHashKey name)
         {
             using Scope _ = Tracer.Instance.StartActive("scylla.get");
 
@@ -67,10 +67,10 @@ namespace Horde.Storage.Implementation
                 throw new ObjectNotFoundException(ns, bucket, name);
 
             // TODO: Check returned values for null
-            return new ObjectRecord(new NamespaceId(o.Namespace!), new BucketId(o.Bucket!), new KeyId(o.Name!), o.LastAccessTime, o.InlinePayload, o.PayloadHash!.AsBlobIdentifier(), o.IsFinalized!.Value);
+            return new ObjectRecord(new NamespaceId(o.Namespace!), new BucketId(o.Bucket!), new IoHashKey(o.Name!), o.LastAccessTime, o.InlinePayload, o.PayloadHash!.AsBlobIdentifier(), o.IsFinalized!.Value);
         }
 
-        public async Task Put(NamespaceId ns, BucketId bucket, KeyId name, BlobIdentifier blobHash, byte[] blob, bool isFinalized)
+        public async Task Put(NamespaceId ns, BucketId bucket, IoHashKey name, BlobIdentifier blobHash, byte[] blob, bool isFinalized)
         {
             using Scope _ = Tracer.Instance.StartActive("scylla.put");
             if (blob.LongLength > _settings.CurrentValue.InlineBlobMaxSize)
@@ -89,7 +89,7 @@ namespace Horde.Storage.Implementation
         }
 
 
-        public async Task Finalize(NamespaceId ns, BucketId bucket, KeyId name, BlobIdentifier blobHash)
+        public async Task Finalize(NamespaceId ns, BucketId bucket, IoHashKey name, BlobIdentifier blobHash)
         {
             using Scope _ = Tracer.Instance.StartActive("scylla.finalize");
             ObjectRecord o = await Get(ns, bucket, name);
@@ -101,7 +101,7 @@ namespace Horde.Storage.Implementation
                 throw new InvalidOperationException("Failed to finalize object even though it existed and hashes matched");
         }
 
-        private async Task CreateTTLRecord(NamespaceId ns, BucketId bucket, KeyId name, sbyte partitionIndex)
+        private async Task CreateTTLRecord(NamespaceId ns, BucketId bucket, IoHashKey name, sbyte partitionIndex)
         {
             using Scope _ = Tracer.Instance.StartActive("scylla.create_ttl_record");
             DateTime now = DateTime.Now;
@@ -114,14 +114,14 @@ namespace Horde.Storage.Implementation
             return RemoveTTLRecord(record.Namespace, record.Bucket, record.Name, record.LastAccess, (sbyte)record.BlobIdentifier.HashData[0]);
         }
 
-        private async Task RemoveTTLRecord(NamespaceId ns, BucketId bucket, KeyId name, DateTime lastAccess, sbyte partitionIndex)
+        private async Task RemoveTTLRecord(NamespaceId ns, BucketId bucket, IoHashKey name, DateTime lastAccess, sbyte partitionIndex)
         {
             using Scope _ = Tracer.Instance.StartActive("scylla.remove_ttl_record");
             LocalDate lastAccessAt = new LocalDate(lastAccess.Year, lastAccess.Month, lastAccess.Day);
             await _mapper.UpdateAsync<ScyllaObjectLastAccessAt>("SET objects = objects - ? WHERE accessed_at = ? AND namespace = ? AND partition_index = ?", new string[] {$"{bucket}#{name}"},  lastAccessAt, ns.ToString(), partitionIndex);
         }
 
-        public async Task UpdateLastAccessTime(NamespaceId ns, BucketId bucket, KeyId name, DateTime lastAccessTime)
+        public async Task UpdateLastAccessTime(NamespaceId ns, BucketId bucket, IoHashKey name, DateTime lastAccessTime)
         {
             using Scope _ = Tracer.Instance.StartActive("scylla.update_last_access_time");
             // fetch the old record
@@ -170,7 +170,7 @@ namespace Horde.Storage.Implementation
                         {
                             int bucketSeparator = o.IndexOf("#", StringComparison.InvariantCultureIgnoreCase);
                             BucketId bucket = new BucketId(o.Substring(0, bucketSeparator));
-                            KeyId name = new KeyId(o.Substring(bucketSeparator + 1));
+                            IoHashKey name = new IoHashKey(o.Substring(bucketSeparator + 1));
 
                             ObjectRecord record;
                             try
@@ -204,7 +204,7 @@ namespace Horde.Storage.Implementation
             }
         }
 
-        public async Task<long> Delete(NamespaceId ns, BucketId bucket, KeyId key)
+        public async Task<long> Delete(NamespaceId ns, BucketId bucket, IoHashKey key)
         {
             using Scope _ = Tracer.Instance.StartActive("scylla.delete_record");
             ObjectRecord record;
@@ -238,7 +238,7 @@ namespace Horde.Storage.Implementation
                 string bucket = row.GetValue<string>("bucket");
                 string name = row.GetValue<string>("name");
 
-                await Delete(ns, new BucketId(bucket), new KeyId(name));
+                await Delete(ns, new BucketId(bucket), new IoHashKey(name));
 
                 deletedCount++;
             }
@@ -258,7 +258,7 @@ namespace Horde.Storage.Implementation
             {
                 string name = row.GetValue<string>("name");
 
-                await Delete(ns, bucket, new KeyId(name));
+                await Delete(ns, bucket, new IoHashKey(name));
                 deletedCount++;
             }
 
@@ -277,7 +277,7 @@ namespace Horde.Storage.Implementation
 
     public class ObjectHashMismatchException : Exception
     {
-        public ObjectHashMismatchException(NamespaceId ns, BucketId bucket, KeyId name, BlobIdentifier suppliedHash, BlobIdentifier actualHash) : base($"Object {name} in bucket {bucket} and namespace {ns} did not reference hash {suppliedHash} was referencing {actualHash}")
+        public ObjectHashMismatchException(NamespaceId ns, BucketId bucket, IoHashKey name, BlobIdentifier suppliedHash, BlobIdentifier actualHash) : base($"Object {name} in bucket {bucket} and namespace {ns} did not reference hash {suppliedHash} was referencing {actualHash}")
         {
         }
     }
@@ -351,7 +351,7 @@ namespace Horde.Storage.Implementation
 
         }
 
-        public ScyllaObject(NamespaceId ns, BucketId bucket, KeyId name, byte[] payload, BlobIdentifier payloadHash, bool isFinalized)
+        public ScyllaObject(NamespaceId ns, BucketId bucket, IoHashKey name, byte[] payload, BlobIdentifier payloadHash, bool isFinalized)
         {
             Namespace = ns.ToString();
             Bucket = bucket.ToString();
@@ -364,11 +364,6 @@ namespace Horde.Storage.Implementation
             LastAccessTime = DateTime.Now;
         }
 
-        public static string BuildKey(string ns, string bucket, string name)
-        {
-            return $"{ns}.{bucket}.{name}";
-        }
-
         [Cassandra.Mapping.Attributes.PartitionKey(0)]
         public string? Namespace { get; set; }
 
diff --git a/Engine/Source/Programs/Horde/HordeStorage/Horde.Storage/Implementation/Replication/RefsReplicator.cs b/Engine/Source/Programs/Horde/HordeStorage/Horde.Storage/Implementation/Replication/RefsReplicator.cs
--- a/Engine/Source/Programs/Horde/HordeStorage/Horde.Storage/Implementation/Replication/RefsReplicator.cs
+++ b/Engine/Source/Programs/Horde/HordeStorage/Horde.Storage/Implementation/Replication/RefsReplicator.cs
@@ -364,7 +364,7 @@ namespace Horde.Storage.Implementation
             return countOfReplicationsDone;
         }
 
-        private async Task AddToReplicationLog(NamespaceId ns, BucketId bucket, KeyId key, BlobIdentifier blob)
+        private async Task AddToReplicationLog(NamespaceId ns, BucketId bucket, IoHashKey key, BlobIdentifier blob)
         {
             await _replicationLog.InsertAddEvent(ns, bucket, key, blob);
         }
diff --git a/Engine/Source/Programs/Horde/HordeStorage/Horde.Storage/Implementation/TransactionLog/IReplicationLog.cs b/Engine/Source/Programs/Horde/HordeStorage/Horde.Storage/Implementation/TransactionLog/IReplicationLog.cs
--- a/Engine/Source/Programs/Horde/HordeStorage/Horde.Storage/Implementation/TransactionLog/IReplicationLog.cs
+++ b/Engine/Source/Programs/Horde/HordeStorage/Horde.Storage/Implementation/TransactionLog/IReplicationLog.cs
@@ -14,8 +14,8 @@ namespace Horde.Storage.Implementation
     {
         IAsyncEnumerable<NamespaceId> GetNamespaces();
 
-        Task<(string, Guid)> InsertAddEvent(NamespaceId ns, BucketId bucket, KeyId key, BlobIdentifier objectBlob, DateTime? timeBucket = null);
-        Task<(string, Guid)> InsertDeleteEvent(NamespaceId ns, BucketId bucket, KeyId key, BlobIdentifier objectBlob, DateTime? timeBucket = null);
+        Task<(string, Guid)> InsertAddEvent(NamespaceId ns, BucketId bucket, IoHashKey key, BlobIdentifier objectBlob, DateTime? timeBucket = null);
+        Task<(string, Guid)> InsertDeleteEvent(NamespaceId ns, BucketId bucket, IoHashKey key, BlobIdentifier objectBlob, DateTime? timeBucket = null);
         IAsyncEnumerable<ReplicationLogEvent> Get(NamespaceId ns, string? lastBucket, Guid? lastEvent);
 
         Task AddSnapshot(SnapshotInfo snapshotHeader);
@@ -40,7 +40,7 @@ namespace Horde.Storage.Implementation
 
     public class ReplicationLogEvent
     {
-        public ReplicationLogEvent(NamespaceId @namespace, BucketId bucket, KeyId key, BlobIdentifier blob, Guid eventId, string timeBucket, DateTime timestamp, OpType op)
+        public ReplicationLogEvent(NamespaceId @namespace, BucketId bucket, IoHashKey key, BlobIdentifier blob, Guid eventId, string timeBucket, DateTime timestamp, OpType op)
         {
             Namespace = @namespace;
             Bucket = bucket;
@@ -61,7 +61,7 @@ namespace Horde.Storage.Implementation
 
         public NamespaceId Namespace { get; }
         public BucketId Bucket { get; }
-        public KeyId Key { get; }
+        public IoHashKey Key { get; }
         public OpType Op { get; }
         public DateTime Timestamp { get; }
         public string TimeBucket { get; }
diff --git a/Engine/Source/Programs/Horde/HordeStorage/Horde.Storage/Implementation/TransactionLog/MemoryReplicationLog.cs b/Engine/Source/Programs/Horde/HordeStorage/Horde.Storage/Implementation/TransactionLog/MemoryReplicationLog.cs
--- a/Engine/Source/Programs/Horde/HordeStorage/Horde.Storage/Implementation/TransactionLog/MemoryReplicationLog.cs
+++ b/Engine/Source/Programs/Horde/HordeStorage/Horde.Storage/Implementation/TransactionLog/MemoryReplicationLog.cs
@@ -21,12 +21,12 @@ namespace Horde.Storage.Implementation
             return _replicationEvents.Keys.ToAsyncEnumerable();
         }
 
-        public Task<(string, Guid)> InsertAddEvent(NamespaceId ns, BucketId bucket, KeyId key, BlobIdentifier objectBlob, DateTime? timestamp)
+        public Task<(string, Guid)> InsertAddEvent(NamespaceId ns, BucketId bucket, IoHashKey key, BlobIdentifier objectBlob, DateTime? timestamp)
         {
             return DoInsert(ns, bucket, key, objectBlob, ReplicationLogEvent.OpType.Added, timestamp);
         }
 
-        private async Task<(string, Guid)> DoInsert(NamespaceId ns, BucketId bucket, KeyId key, BlobIdentifier hash, ReplicationLogEvent.OpType op, DateTime? lastTimestamp)
+        private async Task<(string, Guid)> DoInsert(NamespaceId ns, BucketId bucket, IoHashKey key, BlobIdentifier hash, ReplicationLogEvent.OpType op, DateTime? lastTimestamp)
         {
             DateTime timestamp = lastTimestamp.GetValueOrDefault(DateTime.Now);
 
@@ -65,7 +65,7 @@ namespace Horde.Storage.Implementation
             });
         }
 
-        public Task<(string, Guid)> InsertDeleteEvent(NamespaceId ns, BucketId bucket, KeyId key, BlobIdentifier objectBlob, DateTime? timestamp)
+        public Task<(string, Guid)> InsertDeleteEvent(NamespaceId ns, BucketId bucket, IoHashKey key, BlobIdentifier objectBlob, DateTime? timestamp)
         {
             return DoInsert(ns, bucket, key, objectBlob, ReplicationLogEvent.OpType.Deleted, timestamp); 
         }
diff --git a/Engine/Source/Programs/Horde/HordeStorage/Horde.Storage/Implementation/TransactionLog/ReplicationLogSnapshot.cs b/Engine/Source/Programs/Horde/HordeStorage/Horde.Storage/Implementation/TransactionLog/ReplicationLogSnapshot.cs
--- a/Engine/Source/Programs/Horde/HordeStorage/Horde.Storage/Implementation/TransactionLog/ReplicationLogSnapshot.cs
+++ b/Engine/Source/Programs/Horde/HordeStorage/Horde.Storage/Implementation/TransactionLog/ReplicationLogSnapshot.cs
@@ -37,7 +37,7 @@ namespace Horde.Storage.Implementation.TransactionLog
 
         public class SnapshotLiveObject
         {
-            public SnapshotLiveObject(BucketId bucket, KeyId key, BlobIdentifier blob)
+            public SnapshotLiveObject(BucketId bucket, IoHashKey key, BlobIdentifier blob)
             {
                 Bucket = bucket;
                 Key = key;
@@ -45,7 +45,7 @@ namespace Horde.Storage.Implementation.TransactionLog
             }
 
             public BucketId Bucket { get; set; }
-            public KeyId Key { get; set; }
+            public IoHashKey Key { get; set; }
             public BlobIdentifier Blob { get; set; }
         }
 
@@ -129,16 +129,16 @@ namespace Horde.Storage.Implementation.TransactionLog
         }
 
 
-        private string BuildObjectKey(BucketId bucket, KeyId key)
+        private string BuildObjectKey(BucketId bucket, IoHashKey key)
         {
             return $"{bucket}.{key}";
         }
 
-        private void ProcessAddEvent(BucketId bucket, KeyId key, BlobIdentifier blob)
+        private void ProcessAddEvent(BucketId bucket, IoHashKey key, BlobIdentifier blob)
         {
             _liveObjects.TryAdd(BuildObjectKey(bucket, key), new SnapshotLiveObject(bucket, key, blob));
         }
-        private void ProcessDeleteEvent(BucketId bucket, KeyId key)
+        private void ProcessDeleteEvent(BucketId bucket, IoHashKey key)
         {
             _liveObjects.Remove(BuildObjectKey(bucket, key));
         }
diff --git a/Engine/Source/Programs/Horde/HordeStorage/Horde.Storage/Implementation/TransactionLog/ReplicationLogSnapshotBuilder.cs b/Engine/Source/Programs/Horde/HordeStorage/Horde.Storage/Implementation/TransactionLog/ReplicationLogSnapshotBuilder.cs
--- a/Engine/Source/Programs/Horde/HordeStorage/Horde.Storage/Implementation/TransactionLog/ReplicationLogSnapshotBuilder.cs
+++ b/Engine/Source/Programs/Horde/HordeStorage/Horde.Storage/Implementation/TransactionLog/ReplicationLogSnapshotBuilder.cs
@@ -77,7 +77,7 @@ namespace Horde.Storage.Implementation.TransactionLog
 
                 byte[] cbObjectBytes = writer.Save();
                 BlobIdentifier cbBlobId = BlobIdentifier.FromBlob(cbObjectBytes);
-                await _referencesStore.Put(storeInNamespace, new BucketId("snapshot"), new KeyId(blobIdentifier.ToString()), cbBlobId, cbObjectBytes, true);
+                await _referencesStore.Put(storeInNamespace, new BucketId("snapshot"), new IoHashKey(blobIdentifier.ToString()), cbBlobId, cbObjectBytes, true);
                 await _blobService.PutObject(storeInNamespace, buf, blobIdentifier);
 
                 if (cancellationToken.IsCancellationRequested)
diff --git a/Engine/Source/Programs/Horde/HordeStorage/Horde.Storage/Implementation/TransactionLog/ScyllaReplicationLog.cs b/Engine/Source/Programs/Horde/HordeStorage/Horde.Storage/Implementation/TransactionLog/ScyllaReplicationLog.cs
--- a/Engine/Source/Programs/Horde/HordeStorage/Horde.Storage/Implementation/TransactionLog/ScyllaReplicationLog.cs
+++ b/Engine/Source/Programs/Horde/HordeStorage/Horde.Storage/Implementation/TransactionLog/ScyllaReplicationLog.cs
@@ -63,7 +63,7 @@ namespace Horde.Storage.Implementation
             }
         }
 
-        public async Task<(string, Guid)> InsertAddEvent(NamespaceId ns, BucketId bucket, KeyId key, BlobIdentifier objectBlob, DateTime? timestamp)
+        public async Task<(string, Guid)> InsertAddEvent(NamespaceId ns, BucketId bucket, IoHashKey key, BlobIdentifier objectBlob, DateTime? timestamp)
         {
             using Scope _ = Tracer.Instance.StartActive("scylla.insert_add_event");
             Task addNamespaceTask = PotentiallyAddNamespace(ns);
@@ -75,7 +75,7 @@ namespace Horde.Storage.Implementation
             return (log.GetReplicationBucketIdentifier(), log.ReplicationId);
         }
 
-        public async Task<(string, Guid)> InsertDeleteEvent(NamespaceId ns, BucketId bucket, KeyId key, BlobIdentifier objectBlob, DateTime? timestamp)
+        public async Task<(string, Guid)> InsertDeleteEvent(NamespaceId ns, BucketId bucket, IoHashKey key, BlobIdentifier objectBlob, DateTime? timestamp)
         {
             using Scope _ = Tracer.Instance.StartActive("scylla.insert_delete_event");
             Task addNamespaceTask = PotentiallyAddNamespace(ns);
@@ -132,7 +132,7 @@ namespace Horde.Storage.Implementation
                     yield return new ReplicationLogEvent(
                         new NamespaceId(scyllaReplicationLog.Namespace),
                         new BucketId(scyllaReplicationLog.Bucket!),
-                        new KeyId(scyllaReplicationLog.Key!),
+                        new IoHashKey(scyllaReplicationLog.Key!),
                         scyllaReplicationLog.ObjectIdentifier!.AsBlobIdentifier(),
                         scyllaReplicationLog.ReplicationId,
                         scyllaReplicationLog.GetReplicationBucketIdentifier(),
diff --git a/Engine/Source/Programs/Horde/HordeStorage/Horde.Storage/Tests/Functional/Objects/RefReplicatorTests.cs b/Engine/Source/Programs/Horde/HordeStorage/Horde.Storage/Tests/Functional/Objects/RefReplicatorTests.cs
--- a/Engine/Source/Programs/Horde/HordeStorage/Horde.Storage/Tests/Functional/Objects/RefReplicatorTests.cs
+++ b/Engine/Source/Programs/Horde/HordeStorage/Horde.Storage/Tests/Functional/Objects/RefReplicatorTests.cs
@@ -123,7 +123,7 @@ namespace Horde.Storage.FunctionalTests.Replication
                 byte[] blobContents = Encoding.UTF8.GetBytes($"random content {i}");
                 BlobIdentifier blob = BlobIdentifier.FromBlob(blobContents);
                 blobs.Add(blob, blobContents);
-                replicationEvents.Add(new ReplicationLogEvent(TestNamespace, TestBucket, new KeyId($"event-{i}"), blob, Guid.NewGuid(), "refs-000", DateTime.Now, ReplicationLogEvent.OpType.Added));
+                replicationEvents.Add(new ReplicationLogEvent(TestNamespace, TestBucket, IoHashKey.FromName($"event-{i}"), blob, Guid.NewGuid(), "refs-000", DateTime.Now, ReplicationLogEvent.OpType.Added));
             }
             string lastBucket = "refs-000";
             Guid lastEvent = Guid.NewGuid();
@@ -180,7 +180,7 @@ namespace Horde.Storage.FunctionalTests.Replication
                 byte[] blobContents = Encoding.UTF8.GetBytes($"random content {i}");
                 BlobIdentifier blob = BlobIdentifier.FromBlob(blobContents);
                 blobs.Add(blob, blobContents);
-                replicationEvents.Add(new ReplicationLogEvent(TestNamespace, TestBucket, new KeyId($"event-{i}"), blob, Guid.NewGuid(), "refs-000", DateTime.Now, ReplicationLogEvent.OpType.Added));
+                replicationEvents.Add(new ReplicationLogEvent(TestNamespace, TestBucket, IoHashKey.FromName($"event-{i}"), blob, Guid.NewGuid(), "refs-000", DateTime.Now, ReplicationLogEvent.OpType.Added));
             }
 
             // Build snapshot
@@ -257,7 +257,7 @@ namespace Horde.Storage.FunctionalTests.Replication
                 byte[] blobContents = Encoding.UTF8.GetBytes($"random content in snapshot {i}");
                 BlobIdentifier blob = BlobIdentifier.FromBlob(blobContents);
                 blobs.Add(blob, blobContents);
-                snapshotEvents.Add(new ReplicationLogEvent(TestNamespace, TestBucket, new KeyId($"event-{i}"), blob, Guid.NewGuid(), "refs-000", DateTime.Now, ReplicationLogEvent.OpType.Added));
+                snapshotEvents.Add(new ReplicationLogEvent(TestNamespace, TestBucket,IoHashKey.FromName($"event-{i}"), blob, Guid.NewGuid(), "refs-000", DateTime.Now, ReplicationLogEvent.OpType.Added));
             }
 
             for (int i = 0; i < countOfTestEvents; i++)
@@ -265,7 +265,7 @@ namespace Horde.Storage.FunctionalTests.Replication
                 byte[] blobContents = Encoding.UTF8.GetBytes($"random content {i}");
                 BlobIdentifier blob = BlobIdentifier.FromBlob(blobContents);
                 blobs.Add(blob, blobContents);
-                incrementalEvents.Add(new ReplicationLogEvent(TestNamespace, TestBucket, new KeyId($"incremental-event-{i}"), blob, Guid.NewGuid(), "refs-000", DateTime.Now, ReplicationLogEvent.OpType.Added));
+                incrementalEvents.Add(new ReplicationLogEvent(TestNamespace, TestBucket, IoHashKey.FromName($"incremental-event-{i}"), blob, Guid.NewGuid(), "refs-000", DateTime.Now, ReplicationLogEvent.OpType.Added));
             }
             
             // Build snapshot
@@ -345,7 +345,7 @@ namespace Horde.Storage.FunctionalTests.Replication
                 byte[] blobContents = Encoding.UTF8.GetBytes($"random content in snapshot {i}");
                 BlobIdentifier blob = BlobIdentifier.FromBlob(blobContents);
                 blobs.Add(blob, blobContents);
-                snapshotEvents.Add(new ReplicationLogEvent(TestNamespace, TestBucket, new KeyId($"event-{i}"), blob, Guid.NewGuid(), "refs-000", DateTime.Now, ReplicationLogEvent.OpType.Added));
+                snapshotEvents.Add(new ReplicationLogEvent(TestNamespace, TestBucket, IoHashKey.FromName($"event-{i}"), blob, Guid.NewGuid(), "refs-000", DateTime.Now, ReplicationLogEvent.OpType.Added));
             }
 
             // Build snapshot
diff --git a/Engine/Source/Programs/Horde/HordeStorage/Horde.Storage/Tests/Functional/Objects/ReferencesTests.cs b/Engine/Source/Programs/Horde/HordeStorage/Horde.Storage/Tests/Functional/Objects/ReferencesTests.cs
--- a/Engine/Source/Programs/Horde/HordeStorage/Horde.Storage/Tests/Functional/Objects/ReferencesTests.cs
+++ b/Engine/Source/Programs/Horde/HordeStorage/Horde.Storage/Tests/Functional/Objects/ReferencesTests.cs
@@ -64,11 +64,9 @@ namespace Horde.Storage.FunctionalTests.References
         {
             IScyllaSessionManager scyllaSessionManager = provider.GetService<IScyllaSessionManager>()!;
 
-            // Tests generally do not require the database to be empty and creating the tables all the times slows down the tests significantly
-            // Still this can be useful to have for debugging purposes.
-
-            //await session.ExecuteAsync(new SimpleStatement("DROP TABLE IF EXISTS objects"));
-            //await session.ExecuteAsync(new SimpleStatement("DROP TABLE IF EXISTS buckets"));
+            ISession replicatedKeyspace = scyllaSessionManager.GetSessionForReplicatedKeyspace();
+            await replicatedKeyspace.ExecuteAsync(new SimpleStatement("DROP TABLE IF EXISTS objects"));
+            await replicatedKeyspace.ExecuteAsync(new SimpleStatement("DROP TABLE IF EXISTS object_last_access"));
 
             // we need to clear out the state we modified in the replication log, otherwise the replication log tests will fail
             ISession localKeyspace = scyllaSessionManager.GetSessionForLocalKeyspace();
@@ -184,16 +182,16 @@ namespace Horde.Storage.FunctionalTests.References
             const string objectContents = "This is treated as a opaque blob";
             byte[] data = Encoding.ASCII.GetBytes(objectContents);
             BlobIdentifier objectHash = BlobIdentifier.FromBlob(data);
-
+            IoHashKey key = IoHashKey.FromName("newBlobObject");
             HttpContent requestContent = new ByteArrayContent(data);
             requestContent.Headers.ContentType = new MediaTypeHeaderValue(MediaTypeNames.Application.Octet);
             requestContent.Headers.Add(CommonHeaders.HashHeaderName, objectHash.ToString());
 
-            HttpResponseMessage result = await _httpClient!.PutAsync(requestUri: $"api/v1/refs/{TestNamespace}/bucket/newBlobObject", requestContent);
+            HttpResponseMessage result = await _httpClient!.PutAsync(requestUri: $"api/v1/refs/{TestNamespace}/bucket/{key}", requestContent);
             result.EnsureSuccessStatusCode();
 
             {
-                HttpResponseMessage getResponse = await _httpClient.GetAsync($"api/v1/refs/{TestNamespace}/bucket/newBlobObject.raw");
+                HttpResponseMessage getResponse = await _httpClient.GetAsync($"api/v1/refs/{TestNamespace}/bucket/{key}.raw");
                 getResponse.EnsureSuccessStatusCode();
                 await using MemoryStream ms = new MemoryStream();
                 await getResponse.Content.CopyToAsync(ms);
@@ -209,7 +207,7 @@ namespace Horde.Storage.FunctionalTests.References
             {
                 BlobIdentifier attachment;
                 {
-                    HttpResponseMessage getResponse = await _httpClient.GetAsync($"api/v1/refs/{TestNamespace}/bucket/newBlobObject.uecb");
+                    HttpResponseMessage getResponse = await _httpClient.GetAsync($"api/v1/refs/{TestNamespace}/bucket/{key}.uecb");
                     getResponse.EnsureSuccessStatusCode();
                     await using MemoryStream ms = new MemoryStream();
                     await getResponse.Content.CopyToAsync(ms);
@@ -240,7 +238,7 @@ namespace Horde.Storage.FunctionalTests.References
             }
 
             {
-                HttpResponseMessage getResponse = await _httpClient.GetAsync($"api/v1/refs/{TestNamespace}/bucket/newBlobObject.json");
+                HttpResponseMessage getResponse = await _httpClient.GetAsync($"api/v1/refs/{TestNamespace}/bucket/{key}.json");
                 getResponse.EnsureSuccessStatusCode();
                 await using MemoryStream ms = new MemoryStream();
                 await getResponse.Content.CopyToAsync(ms);
@@ -264,7 +262,7 @@ namespace Horde.Storage.FunctionalTests.References
 
             {
                 // request the object as a json response using accept instead of the format filter
-                var request = new HttpRequestMessage(HttpMethod.Get, $"api/v1/refs/{TestNamespace}/bucket/newBlobObject");
+                var request = new HttpRequestMessage(HttpMethod.Get, $"api/v1/refs/{TestNamespace}/bucket/{key}");
                 request.Headers.Accept.Add(new MediaTypeWithQualityHeaderValue(MediaTypeNames.Application.Json));
                 HttpResponseMessage getResponse = await _httpClient.SendAsync(request);
                 getResponse.EnsureSuccessStatusCode();
@@ -301,12 +299,13 @@ namespace Horde.Storage.FunctionalTests.References
 
             byte[] objectData = writer.Save();
             BlobIdentifier objectHash = BlobIdentifier.FromBlob(objectData);
+            IoHashKey key = IoHashKey.FromName("newReferenceObject");
 
             HttpContent requestContent = new ByteArrayContent(objectData);
             requestContent.Headers.ContentType = new MediaTypeHeaderValue(CustomMediaTypeNames.UnrealCompactBinary);
             requestContent.Headers.Add(CommonHeaders.HashHeaderName, objectHash.ToString());
 
-            HttpResponseMessage result = await _httpClient!.PutAsync(requestUri: $"api/v1/refs/{TestNamespace}/bucket/newReferenceObject.uecb", requestContent);
+            HttpResponseMessage result = await _httpClient!.PutAsync(requestUri: $"api/v1/refs/{TestNamespace}/bucket/{key}.uecb", requestContent);
             result.EnsureSuccessStatusCode();
 
             {
@@ -325,18 +324,17 @@ namespace Horde.Storage.FunctionalTests.References
 
             {
                 BucketId bucket = new BucketId("bucket");
-                KeyId key = new KeyId("newReferenceObject");
 
                 ObjectRecord objectRecord = await ReferencesStore.Get(TestNamespace, bucket, key);
 
                 Assert.IsTrue(objectRecord.IsFinalized);
-                Assert.AreEqual("newReferenceObject", objectRecord.Name.ToString());
+                Assert.AreEqual(key, objectRecord.Name);
                 Assert.AreEqual(objectHash, objectRecord.BlobIdentifier);
                 Assert.IsNotNull(objectRecord.InlinePayload);
             }
 
             {
-                HttpResponseMessage getResponse = await _httpClient.GetAsync($"api/v1/refs/{TestNamespace}/bucket/newReferenceObject.raw");
+                HttpResponseMessage getResponse = await _httpClient.GetAsync($"api/v1/refs/{TestNamespace}/bucket/{key}.raw");
                 getResponse.EnsureSuccessStatusCode();
                 await using MemoryStream ms = new MemoryStream();
                 await getResponse.Content.CopyToAsync(ms);
@@ -348,7 +346,7 @@ namespace Horde.Storage.FunctionalTests.References
             }
 
             {
-                HttpResponseMessage getResponse = await _httpClient.GetAsync($"api/v1/refs/{TestNamespace}/bucket/newReferenceObject.uecb");
+                HttpResponseMessage getResponse = await _httpClient.GetAsync($"api/v1/refs/{TestNamespace}/bucket/{key}.uecb");
                 getResponse.EnsureSuccessStatusCode();
                 await using MemoryStream ms = new MemoryStream();
                 await getResponse.Content.CopyToAsync(ms);
@@ -366,7 +364,7 @@ namespace Horde.Storage.FunctionalTests.References
 
 
             {
-                HttpResponseMessage getResponse = await _httpClient.GetAsync($"api/v1/refs/{TestNamespace}/bucket/newReferenceObject.json");
+                HttpResponseMessage getResponse = await _httpClient.GetAsync($"api/v1/refs/{TestNamespace}/bucket/{key}.json");
                 getResponse.EnsureSuccessStatusCode();
                 await using MemoryStream ms = new MemoryStream();
                 await getResponse.Content.CopyToAsync(ms);
@@ -409,13 +407,14 @@ namespace Horde.Storage.FunctionalTests.References
             byte[] parentObjectData = parentObjectWriter.Save();
             BlobIdentifier parentObjectHash = BlobIdentifier.FromBlob(parentObjectData);
 
+            IoHashKey key = IoHashKey.FromName("newHierarchyObject");
             // this first upload should fail with the child object missing
             {
                 HttpContent requestContent = new ByteArrayContent(parentObjectData);
                 requestContent.Headers.ContentType = new MediaTypeHeaderValue(CustomMediaTypeNames.UnrealCompactBinary);
                 requestContent.Headers.Add(CommonHeaders.HashHeaderName, parentObjectHash.ToString());
 
-                HttpResponseMessage result = await _httpClient!.PutAsync(requestUri: $"api/v1/refs/{TestNamespace}/bucket/newHierarchyObject.uecb", requestContent);
+                HttpResponseMessage result = await _httpClient!.PutAsync(requestUri: $"api/v1/refs/{TestNamespace}/bucket/{key}.uecb", requestContent);
                 result.EnsureSuccessStatusCode();
 
                 Assert.AreEqual(CustomMediaTypeNames.UnrealCompactBinary, result!.Content.Headers.ContentType!.MediaType);
@@ -463,7 +462,7 @@ namespace Horde.Storage.FunctionalTests.References
                 requestContent.Headers.ContentType = new MediaTypeHeaderValue(CustomMediaTypeNames.UnrealCompactBinary);
                 requestContent.Headers.Add(CommonHeaders.HashHeaderName, parentObjectHash.ToString());
 
-                HttpResponseMessage result = await _httpClient!.PutAsync(requestUri: $"api/v1/refs/{TestNamespace}/bucket/newHierarchyObject.uecb", requestContent);
+                HttpResponseMessage result = await _httpClient!.PutAsync(requestUri: $"api/v1/refs/{TestNamespace}/bucket/{key}.uecb", requestContent);
                 result.EnsureSuccessStatusCode();
 
                 Assert.AreEqual(result!.Content.Headers.ContentType!.MediaType, CustomMediaTypeNames.UnrealCompactBinary);
@@ -481,18 +480,17 @@ namespace Horde.Storage.FunctionalTests.References
 
             {
                 BucketId bucket = new BucketId("bucket");
-                KeyId key = new KeyId("newHierarchyObject");
 
                 ObjectRecord objectRecord = await ReferencesStore.Get(TestNamespace, bucket, key);
 
                 Assert.IsTrue(objectRecord.IsFinalized);
-                Assert.AreEqual("newHierarchyObject", objectRecord.Name.ToString());
+                Assert.AreEqual(key, objectRecord.Name);
                 Assert.AreEqual(parentObjectHash, objectRecord.BlobIdentifier);
                 Assert.IsNotNull(objectRecord.InlinePayload);
             }
 
             {
-                HttpResponseMessage getResponse = await _httpClient.GetAsync($"api/v1/refs/{TestNamespace}/bucket/newHierarchyObject.raw");
+                HttpResponseMessage getResponse = await _httpClient.GetAsync($"api/v1/refs/{TestNamespace}/bucket/{key}.raw");
                 getResponse.EnsureSuccessStatusCode();
                 await using MemoryStream ms = new MemoryStream();
                 await getResponse.Content.CopyToAsync(ms);
@@ -504,7 +502,7 @@ namespace Horde.Storage.FunctionalTests.References
             }
 
             {
-                HttpResponseMessage getResponse = await _httpClient.GetAsync($"api/v1/refs/{TestNamespace}/bucket/newHierarchyObject.uecb");
+                HttpResponseMessage getResponse = await _httpClient.GetAsync($"api/v1/refs/{TestNamespace}/bucket/{key}.uecb");
                 getResponse.EnsureSuccessStatusCode();
                 await using MemoryStream ms = new MemoryStream();
                 await getResponse.Content.CopyToAsync(ms);
@@ -522,7 +520,7 @@ namespace Horde.Storage.FunctionalTests.References
 
 
             {
-                HttpResponseMessage getResponse = await _httpClient.GetAsync($"api/v1/refs/{TestNamespace}/bucket/newHierarchyObject.json");
+                HttpResponseMessage getResponse = await _httpClient.GetAsync($"api/v1/refs/{TestNamespace}/bucket/{key}.json");
                 getResponse.EnsureSuccessStatusCode();
                 await using MemoryStream ms = new MemoryStream();
                 await getResponse.Content.CopyToAsync(ms);
@@ -566,23 +564,23 @@ namespace Horde.Storage.FunctionalTests.References
             const string objectContents = "This is treated as a opaque blob";
             byte[] data = Encoding.ASCII.GetBytes(objectContents);
             BlobIdentifier objectHash = BlobIdentifier.FromBlob(data);
-
+            IoHashKey key = IoHashKey.FromName("newObject");
             HttpContent requestContent = new ByteArrayContent(data);
             requestContent.Headers.ContentType = new MediaTypeHeaderValue(MediaTypeNames.Application.Octet);
             requestContent.Headers.Add(CommonHeaders.HashHeaderName, objectHash.ToString());
 
             {
-                HttpResponseMessage result = await _httpClient!.PutAsync(requestUri: $"api/v1/refs/{TestNamespace}/bucket/newObject", requestContent);
+                HttpResponseMessage result = await _httpClient!.PutAsync(requestUri: $"api/v1/refs/{TestNamespace}/bucket/{key}", requestContent);
                 result.EnsureSuccessStatusCode();
             }
 
             {
-                HttpResponseMessage result = await _httpClient!.SendAsync(new HttpRequestMessage(HttpMethod.Head, $"api/v1/refs/{TestNamespace}/bucket/newObject"));
+                HttpResponseMessage result = await _httpClient!.SendAsync(new HttpRequestMessage(HttpMethod.Head, $"api/v1/refs/{TestNamespace}/bucket/{key}"));
                 result.EnsureSuccessStatusCode();
             }
 
             {
-                HttpResponseMessage result = await _httpClient!.SendAsync(new HttpRequestMessage(HttpMethod.Head, $"api/v1/refs/{TestNamespace}/bucket/missingObject"));
+                HttpResponseMessage result = await _httpClient!.SendAsync(new HttpRequestMessage(HttpMethod.Head, $"api/v1/refs/{TestNamespace}/bucket/{IoHashKey.FromName("missingObject")}"));
                 Assert.AreEqual(HttpStatusCode.NotFound, result.StatusCode);
             }
         }
@@ -620,11 +618,13 @@ namespace Horde.Storage.FunctionalTests.References
             byte[] objectData = writerParent.Save();
             BlobIdentifier objectHash = BlobIdentifier.FromBlob(objectData);
 
+            IoHashKey key = IoHashKey.FromName("newHierarchyObject");
+
             HttpContent requestContent = new ByteArrayContent(objectData);
             requestContent.Headers.ContentType = new MediaTypeHeaderValue(CustomMediaTypeNames.UnrealCompactBinary);
             requestContent.Headers.Add(CommonHeaders.HashHeaderName, objectHash.ToString());
 
-            HttpResponseMessage result = await _httpClient!.PutAsync(requestUri: $"api/v1/refs/{TestNamespace}/bucket/newHierarchyObject.uecb", requestContent);
+            HttpResponseMessage result = await _httpClient!.PutAsync(requestUri: $"api/v1/refs/{TestNamespace}/bucket/{key}.uecb", requestContent);
             result.EnsureSuccessStatusCode();
 
             // check the response
@@ -645,19 +645,18 @@ namespace Horde.Storage.FunctionalTests.References
             // check that actual internal representation
             {
                 BucketId bucket = new BucketId("bucket");
-                KeyId key = new KeyId("newHierarchyObject");
 
                 ObjectRecord objectRecord = await ReferencesStore.Get(TestNamespace, bucket, key);
 
                 Assert.IsTrue(objectRecord.IsFinalized);
-                Assert.AreEqual("newHierarchyObject", objectRecord.Name.ToString());
+                Assert.AreEqual(key, objectRecord.Name);
                 Assert.AreEqual(objectHash, objectRecord.BlobIdentifier);
                 Assert.IsNotNull(objectRecord.InlinePayload);
             }
 
             // verify attachments
             {
-                HttpResponseMessage getResponse = await _httpClient.GetAsync($"api/v1/refs/{TestNamespace}/bucket/newHierarchyObject.raw");
+                HttpResponseMessage getResponse = await _httpClient.GetAsync($"api/v1/refs/{TestNamespace}/bucket/{key}.raw");
                 getResponse.EnsureSuccessStatusCode();
                 await using MemoryStream ms = new MemoryStream();
                 await getResponse.Content.CopyToAsync(ms);
@@ -672,7 +671,7 @@ namespace Horde.Storage.FunctionalTests.References
                 BlobIdentifier blobAttachment;
                 BlobIdentifier objectAttachment;
                 {
-                    HttpResponseMessage getResponse = await _httpClient.GetAsync($"api/v1/refs/{TestNamespace}/bucket/newHierarchyObject.uecb");
+                    HttpResponseMessage getResponse = await _httpClient.GetAsync($"api/v1/refs/{TestNamespace}/bucket/{key}.uecb");
                     getResponse.EnsureSuccessStatusCode();
                     await using MemoryStream ms = new MemoryStream();
                     await getResponse.Content.CopyToAsync(ms);
@@ -734,7 +733,7 @@ namespace Horde.Storage.FunctionalTests.References
 
             // check json representation
             {
-                HttpResponseMessage getResponse = await _httpClient.GetAsync($"api/v1/refs/{TestNamespace}/bucket/newHierarchyObject.json");
+                HttpResponseMessage getResponse = await _httpClient.GetAsync($"api/v1/refs/{TestNamespace}/bucket/{key}.json");
                 getResponse.EnsureSuccessStatusCode();
                 await using MemoryStream ms = new MemoryStream();
                 await getResponse.Content.CopyToAsync(ms);
@@ -787,13 +786,15 @@ namespace Horde.Storage.FunctionalTests.References
 
             byte[] objectData = writerParent.Save();
             BlobIdentifier objectHash = BlobIdentifier.FromBlob(objectData);
+            
+            IoHashKey key = IoHashKey.FromName("newHierarchyObject");
 
             {
                 HttpContent requestContent = new ByteArrayContent(objectData);
                 requestContent.Headers.ContentType = new MediaTypeHeaderValue(CustomMediaTypeNames.UnrealCompactBinary);
                 requestContent.Headers.Add(CommonHeaders.HashHeaderName, objectHash.ToString());
 
-                HttpResponseMessage result = await _httpClient!.PutAsync(requestUri: $"api/v1/refs/{TestNamespace}/bucket/newHierarchyObject.uecb", requestContent);
+                HttpResponseMessage result = await _httpClient!.PutAsync(requestUri: $"api/v1/refs/{TestNamespace}/bucket/{key}.uecb", requestContent);
                 result.EnsureSuccessStatusCode();
 
                 {
@@ -817,7 +818,7 @@ namespace Horde.Storage.FunctionalTests.References
                 requestContent.Headers.ContentType = new MediaTypeHeaderValue(CustomMediaTypeNames.UnrealCompactBinary);
                 requestContent.Headers.Add(CommonHeaders.HashHeaderName, objectHash.ToString());
 
-                HttpResponseMessage result = await _httpClient!.PutAsync(requestUri: $"api/v1/refs/{TestNamespace}/bucket/newHierarchyObject.json", requestContent);
+                HttpResponseMessage result = await _httpClient!.PutAsync(requestUri: $"api/v1/refs/{TestNamespace}/bucket/{key}.json", requestContent);
                 result.EnsureSuccessStatusCode();
 
                 {
@@ -842,7 +843,7 @@ namespace Horde.Storage.FunctionalTests.References
         public async Task PutAndFinalize()
         {
             BucketId bucket = new BucketId("bucket");
-            KeyId key = new KeyId("willFinalizeObject");
+            IoHashKey key = IoHashKey.FromName("willFinalizeObject");
 
             // do not submit the content of the blobs, which should be reported in the response of the put
             string blobContents = "This is a string that is referenced as a blob";
@@ -877,7 +878,7 @@ namespace Horde.Storage.FunctionalTests.References
                 requestContent.Headers.ContentType = new MediaTypeHeaderValue(CustomMediaTypeNames.UnrealCompactBinary);
                 requestContent.Headers.Add(CommonHeaders.HashHeaderName, objectHash.ToString());
 
-                HttpResponseMessage result = await _httpClient!.PutAsync(requestUri: $"api/v1/refs/{TestNamespace}/bucket/willFinalizeObject.uecb", requestContent);
+                HttpResponseMessage result = await _httpClient!.PutAsync(requestUri: $"api/v1/refs/{TestNamespace}/bucket/{key}.uecb", requestContent);
                 result.EnsureSuccessStatusCode();
 
                 {
@@ -901,7 +902,7 @@ namespace Horde.Storage.FunctionalTests.References
                 ObjectRecord objectRecord = await ReferencesStore.Get(TestNamespace, bucket, key);
 
                 Assert.IsFalse(objectRecord.IsFinalized);
-                Assert.AreEqual("willFinalizeObject", objectRecord.Name.ToString());
+                Assert.AreEqual(key, objectRecord.Name);
                 Assert.AreEqual(objectHash, objectRecord.BlobIdentifier);
                 Assert.IsNotNull(objectRecord.InlinePayload);
             }
@@ -916,7 +917,7 @@ namespace Horde.Storage.FunctionalTests.References
             {
                 HttpContent requestContent = new ByteArrayContent(Array.Empty<byte>());
 
-                HttpResponseMessage result = await _httpClient!.PostAsync(requestUri: $"api/v1/refs/{TestNamespace}/bucket/willFinalizeObject/finalize/{objectHash}.uecb", requestContent);
+                HttpResponseMessage result = await _httpClient!.PostAsync(requestUri: $"api/v1/refs/{TestNamespace}/bucket/{key}/finalize/{objectHash}.uecb", requestContent);
                 result.EnsureSuccessStatusCode();
 
                 {
@@ -938,7 +939,7 @@ namespace Horde.Storage.FunctionalTests.References
                 ObjectRecord objectRecord = await ReferencesStore.Get(TestNamespace, bucket, key);
 
                 Assert.IsTrue(objectRecord.IsFinalized);
-                Assert.AreEqual("willFinalizeObject", objectRecord.Name.ToString());
+                Assert.AreEqual(key, objectRecord.Name);
                 Assert.AreEqual(objectHash, objectRecord.BlobIdentifier);
             }
         }
@@ -949,12 +950,13 @@ namespace Horde.Storage.FunctionalTests.References
             string blobContents = "This is a blob";
             byte[] blobData = Encoding.ASCII.GetBytes(blobContents);
             BlobIdentifier blobHash = BlobIdentifier.FromBlob(blobData);
+            IoHashKey key = IoHashKey.FromName("newReferenceObject");
 
             HttpContent requestContent = new ByteArrayContent(blobData);
             requestContent.Headers.ContentType = new MediaTypeHeaderValue(MediaTypeNames.Application.Octet);
             requestContent.Headers.Add(CommonHeaders.HashHeaderName, blobHash.ToString());
 
-            HttpResponseMessage result = await _httpClient!.PutAsync(requestUri: $"api/v1/refs/{TestNamespace}/bucket/newReferenceObject.uecb", requestContent);
+            HttpResponseMessage result = await _httpClient!.PutAsync(requestUri: $"api/v1/refs/{TestNamespace}/bucket/{key}.uecb", requestContent);
             result.EnsureSuccessStatusCode();
 
             {
@@ -973,7 +975,7 @@ namespace Horde.Storage.FunctionalTests.References
 
             {
                 // verify we can fetch the blob properly
-                HttpResponseMessage getResponse = await _httpClient.GetAsync($"api/v1/refs/{TestNamespace}/bucket/newReferenceObject.raw");
+                HttpResponseMessage getResponse = await _httpClient.GetAsync($"api/v1/refs/{TestNamespace}/bucket/{key}.raw");
                 getResponse.EnsureSuccessStatusCode();
                 await using MemoryStream ms = new MemoryStream();
                 await getResponse.Content.CopyToAsync(ms);
@@ -991,19 +993,19 @@ namespace Horde.Storage.FunctionalTests.References
 
             {
                 // the compact binary object still exists, so this returns a success (trying to resolve the attachment will fail)
-                HttpResponseMessage getResponse = await _httpClient.GetAsync($"api/v1/refs/{TestNamespace}/bucket/newReferenceObject.uecb");
+                HttpResponseMessage getResponse = await _httpClient.GetAsync($"api/v1/refs/{TestNamespace}/bucket/{key}.uecb");
                 getResponse.EnsureSuccessStatusCode();
             }
 
             {
                 // the compact binary object still exists, so this returns a success (trying to resolve the attachment will fail)
-                HttpResponseMessage getResponse = await _httpClient.GetAsync($"api/v1/refs/{TestNamespace}/bucket/newReferenceObject.json");
+                HttpResponseMessage getResponse = await _httpClient.GetAsync($"api/v1/refs/{TestNamespace}/bucket/{key}.json");
                 getResponse.EnsureSuccessStatusCode();
             }
 
             {
                 // we should now see a 404 as the blob is missing
-                HttpResponseMessage getResponse = await _httpClient.GetAsync($"api/v1/refs/{TestNamespace}/bucket/newReferenceObject.raw");
+                HttpResponseMessage getResponse = await _httpClient.GetAsync($"api/v1/refs/{TestNamespace}/bucket/{key}.raw");
                 Assert.AreEqual(HttpStatusCode.NotFound, getResponse.StatusCode);
                 Assert.AreEqual("application/problem+json", getResponse.Content.Headers.ContentType!.MediaType);
                 string s = await getResponse.Content.ReadAsStringAsync();
@@ -1025,27 +1027,28 @@ namespace Horde.Storage.FunctionalTests.References
             requestContent.Headers.ContentType = new MediaTypeHeaderValue(MediaTypeNames.Application.Octet);
             requestContent.Headers.Add(CommonHeaders.HashHeaderName, objectHash.ToString());
 
+            IoHashKey key = IoHashKey.FromName("deletableObject");
             // submit the object
             {
-                HttpResponseMessage result = await _httpClient!.PutAsync(requestUri: $"api/v1/refs/{TestNamespace}/bucket/deletableObject", requestContent);
+                HttpResponseMessage result = await _httpClient!.PutAsync(requestUri: $"api/v1/refs/{TestNamespace}/bucket/{key}", requestContent);
                 result.EnsureSuccessStatusCode();
             }
 
             // verify it is present
             {
-                HttpResponseMessage result = await _httpClient!.GetAsync(requestUri: $"api/v1/refs/{TestNamespace}/bucket/deletableObject.raw");
+                HttpResponseMessage result = await _httpClient!.GetAsync(requestUri: $"api/v1/refs/{TestNamespace}/bucket/{key}.raw");
                 result.EnsureSuccessStatusCode();
             }
 
             // delete the object
             {
-                HttpResponseMessage result = await _httpClient!.DeleteAsync(requestUri: $"api/v1/refs/{TestNamespace}/bucket/deletableObject");
+                HttpResponseMessage result = await _httpClient!.DeleteAsync(requestUri: $"api/v1/refs/{TestNamespace}/bucket/{key}");
                 result.EnsureSuccessStatusCode();
             }
 
             // ensure the object is not present anymore
             {
-                HttpResponseMessage result = await _httpClient!.GetAsync(requestUri: $"api/v1/refs/{TestNamespace}/bucket/deletableObject.raw");
+                HttpResponseMessage result = await _httpClient!.GetAsync(requestUri: $"api/v1/refs/{TestNamespace}/bucket/{key}.raw");
                 Assert.AreEqual(HttpStatusCode.NotFound, result.StatusCode);
             }
         }
@@ -1064,26 +1067,28 @@ namespace Horde.Storage.FunctionalTests.References
             requestContent.Headers.ContentType = new MediaTypeHeaderValue(MediaTypeNames.Application.Octet);
             requestContent.Headers.Add(CommonHeaders.HashHeaderName, objectHash.ToString());
 
+            IoHashKey key = IoHashKey.FromName("deletableObject");
+
             // submit the object into multiple buckets
             {
-                HttpResponseMessage result = await _httpClient!.PutAsync(requestUri: $"api/v1/refs/{TestNamespace}/{BucketToDelete}/deletableObject", requestContent);
+                HttpResponseMessage result = await _httpClient!.PutAsync(requestUri: $"api/v1/refs/{TestNamespace}/{BucketToDelete}/{key}", requestContent);
                 result.EnsureSuccessStatusCode();
             }
 
             {
-                HttpResponseMessage result = await _httpClient!.PutAsync(requestUri: $"api/v1/refs/{TestNamespace}/bucket/deletableObject", requestContent);
+                HttpResponseMessage result = await _httpClient!.PutAsync(requestUri: $"api/v1/refs/{TestNamespace}/bucket/{key}", requestContent);
                 result.EnsureSuccessStatusCode();
             }
 
             // verify it is present
             {
-                HttpResponseMessage result = await _httpClient!.GetAsync(requestUri: $"api/v1/refs/{TestNamespace}/bucket/deletableObject.raw");
+                HttpResponseMessage result = await _httpClient!.GetAsync(requestUri: $"api/v1/refs/{TestNamespace}/bucket/{key}.raw");
                 result.EnsureSuccessStatusCode();
             }
 
             // verify it is present
             {
-                HttpResponseMessage result = await _httpClient!.GetAsync(requestUri: $"api/v1/refs/{TestNamespace}/{BucketToDelete}/deletableObject.raw");
+                HttpResponseMessage result = await _httpClient!.GetAsync(requestUri: $"api/v1/refs/{TestNamespace}/{BucketToDelete}/{key}.raw");
                 result.EnsureSuccessStatusCode();
             }
 
@@ -1095,13 +1100,13 @@ namespace Horde.Storage.FunctionalTests.References
 
             // ensure the object is present in not deleted bucket
             {
-                HttpResponseMessage result = await _httpClient!.GetAsync(requestUri: $"api/v1/refs/{TestNamespace}/bucket/deletableObject.raw");
+                HttpResponseMessage result = await _httpClient!.GetAsync(requestUri: $"api/v1/refs/{TestNamespace}/bucket/{key}.raw");
                 result.EnsureSuccessStatusCode();
             }
 
             // ensure the object is not present anymore
             {
-                HttpResponseMessage result = await _httpClient!.GetAsync(requestUri: $"api/v1/refs/{TestNamespace}/{BucketToDelete}/deletableObject.raw");
+                HttpResponseMessage result = await _httpClient!.GetAsync(requestUri: $"api/v1/refs/{TestNamespace}/{BucketToDelete}/{key}.raw");
                 Assert.AreEqual(HttpStatusCode.NotFound, result.StatusCode);
             }
         }
@@ -1120,26 +1125,27 @@ namespace Horde.Storage.FunctionalTests.References
             requestContent.Headers.ContentType = new MediaTypeHeaderValue(MediaTypeNames.Application.Octet);
             requestContent.Headers.Add(CommonHeaders.HashHeaderName, objectHash.ToString());
 
+            IoHashKey key = IoHashKey.FromName("deletableObject");
             // submit the object into multiple namespaces
             {
-                HttpResponseMessage result = await _httpClient!.PutAsync(requestUri: $"api/v1/refs/{TestNamespace}/bucket/deletableObject", requestContent);
+                HttpResponseMessage result = await _httpClient!.PutAsync(requestUri: $"api/v1/refs/{TestNamespace}/bucket/{key}", requestContent);
                 result.EnsureSuccessStatusCode();
             }
 
             {
-                HttpResponseMessage result = await _httpClient!.PutAsync(requestUri: $"api/v1/refs/{NamespaceToBeDeleted}/bucket/deletableObject", requestContent);
+                HttpResponseMessage result = await _httpClient!.PutAsync(requestUri: $"api/v1/refs/{NamespaceToBeDeleted}/bucket/{key}", requestContent);
                 result.EnsureSuccessStatusCode();
             }
 
             // verify it is present
             {
-                HttpResponseMessage result = await _httpClient!.GetAsync(requestUri: $"api/v1/refs/{TestNamespace}/bucket/deletableObject.raw");
+                HttpResponseMessage result = await _httpClient!.GetAsync(requestUri: $"api/v1/refs/{TestNamespace}/bucket/{key}.raw");
                 result.EnsureSuccessStatusCode();
             }
 
             // verify it is present
             {
-                HttpResponseMessage result = await _httpClient!.GetAsync(requestUri: $"api/v1/refs/{NamespaceToBeDeleted}/bucket/deletableObject.raw");
+                HttpResponseMessage result = await _httpClient!.GetAsync(requestUri: $"api/v1/refs/{NamespaceToBeDeleted}/bucket/{key}.raw");
                 result.EnsureSuccessStatusCode();
             }
 
@@ -1151,13 +1157,13 @@ namespace Horde.Storage.FunctionalTests.References
 
             // ensure the object is present in not deleted namespace
             {
-                HttpResponseMessage result = await _httpClient!.GetAsync(requestUri: $"api/v1/refs/{TestNamespace}/bucket/deletableObject.raw");
+                HttpResponseMessage result = await _httpClient!.GetAsync(requestUri: $"api/v1/refs/{TestNamespace}/bucket/{key}.raw");
                 result.EnsureSuccessStatusCode();
             }
 
             // ensure the object is not present anymore
             {
-                HttpResponseMessage result = await _httpClient!.GetAsync(requestUri: $"api/v1/refs/{NamespaceToBeDeleted}/bucket/deletableObject.raw");
+                HttpResponseMessage result = await _httpClient!.GetAsync(requestUri: $"api/v1/refs/{NamespaceToBeDeleted}/bucket/{key}.raw");
                 Assert.AreEqual(HttpStatusCode.NotFound, result.StatusCode);
             }
 
@@ -1179,6 +1185,7 @@ namespace Horde.Storage.FunctionalTests.References
             const string objectContents = "This is treated as a opaque blob";
             byte[] data = Encoding.ASCII.GetBytes(objectContents);
             BlobIdentifier objectHash = BlobIdentifier.FromBlob(data);
+            IoHashKey key = IoHashKey.FromName("notUsedObject");
 
             HttpContent requestContent = new ByteArrayContent(data);
             requestContent.Headers.ContentType = new MediaTypeHeaderValue(MediaTypeNames.Application.Octet);
@@ -1186,7 +1193,7 @@ namespace Horde.Storage.FunctionalTests.References
 
             // submit a object to make sure a namespace is created
             {
-                HttpResponseMessage result = await _httpClient!.PutAsync(requestUri: $"api/v1/refs/{TestNamespace}/bucket/notUsedObject", requestContent);
+                HttpResponseMessage result = await _httpClient!.PutAsync(requestUri: $"api/v1/refs/{TestNamespace}/bucket/{key}", requestContent);
                 result.EnsureSuccessStatusCode();
             }
 
@@ -1211,17 +1218,18 @@ namespace Horde.Storage.FunctionalTests.References
             requestContent.Headers.ContentType = new MediaTypeHeaderValue(MediaTypeNames.Application.Octet);
             requestContent.Headers.Add(CommonHeaders.HashHeaderName, objectHash.ToString());
 
+            IoHashKey key = IoHashKey.FromName("oldRecord");
             // submit some contents
             {
-                HttpResponseMessage result = await _httpClient!.PutAsync(requestUri: $"api/v1/refs/{TestNamespace}/bucket/oldRecord", requestContent);
+                HttpResponseMessage result = await _httpClient!.PutAsync(requestUri: $"api/v1/refs/{TestNamespace}/bucket/{key}", requestContent);
                 result.EnsureSuccessStatusCode();
             }
 
             List<ObjectRecord> records = await ReferencesStore.GetOldestRecords(TestNamespace).ToListAsync();
 
-            ObjectRecord? oldRecord = records.Find(record => record.Name == new KeyId("oldRecord"));
+            ObjectRecord? oldRecord = records.Find(record => record.Name == key);
             Assert.IsNotNull(oldRecord);
-            Assert.AreEqual("oldRecord", oldRecord!.Name.ToString());
+            Assert.AreEqual(key, oldRecord!.Name);
             Assert.AreEqual("bucket", oldRecord.Bucket.ToString());
         }
     }
diff --git a/Engine/Source/Programs/Horde/HordeStorage/Horde.Storage/Tests/Functional/Objects/ReplicationTests.cs b/Engine/Source/Programs/Horde/HordeStorage/Horde.Storage/Tests/Functional/Objects/ReplicationTests.cs
--- a/Engine/Source/Programs/Horde/HordeStorage/Horde.Storage/Tests/Functional/Objects/ReplicationTests.cs
+++ b/Engine/Source/Programs/Horde/HordeStorage/Horde.Storage/Tests/Functional/Objects/ReplicationTests.cs
@@ -180,17 +180,17 @@ namespace Horde.Storage.FunctionalTests.References
             requestContent.Headers.Add(CommonHeaders.HashHeaderName, objectHash.ToString());
 
             {
-                HttpResponseMessage result = await _httpClient!.PutAsync(requestUri: $"api/v1/refs/{TestNamespace}/{TestBucket}/newReferenceObject.uecb", requestContent);
+                HttpResponseMessage result = await _httpClient!.PutAsync(requestUri: $"api/v1/refs/{TestNamespace}/{TestBucket}/{IoHashKey.FromName("newReferenceObject")}.uecb", requestContent);
                 result.EnsureSuccessStatusCode();
             }
 
             {
-                HttpResponseMessage result = await _httpClient!.PutAsync(requestUri: $"api/v1/refs/{TestNamespace}/{TestBucket}/secondObject.uecb", requestContent);
+                HttpResponseMessage result = await _httpClient!.PutAsync(requestUri: $"api/v1/refs/{TestNamespace}/{TestBucket}/{IoHashKey.FromName("secondObject")}.uecb", requestContent);
                 result.EnsureSuccessStatusCode();
             }
 
             {
-                HttpResponseMessage result = await _httpClient!.PutAsync(requestUri: $"api/v1/refs/{TestNamespace}/{TestBucket}/thirdObject.uecb", requestContent);
+                HttpResponseMessage result = await _httpClient!.PutAsync(requestUri: $"api/v1/refs/{TestNamespace}/{TestBucket}/{IoHashKey.FromName("thirdObject")}.uecb", requestContent);
                 result.EnsureSuccessStatusCode();
             }
             
@@ -210,7 +210,7 @@ namespace Horde.Storage.FunctionalTests.References
                     ReplicationLogEvent e = events.Events[0];
                     Assert.AreEqual(TestNamespace, e.Namespace);
                     Assert.AreEqual(TestBucket, e.Bucket);
-                    Assert.AreEqual("newReferenceObject", e.Key.ToString());
+                    Assert.AreEqual(IoHashKey.FromName("newReferenceObject"), e.Key);
                     Assert.AreEqual(objectHash, e.Blob);
                 }
 
@@ -218,7 +218,7 @@ namespace Horde.Storage.FunctionalTests.References
                     ReplicationLogEvent e = events.Events[1];
                     Assert.AreEqual(TestNamespace, e.Namespace);
                     Assert.AreEqual(TestBucket, e.Bucket);
-                    Assert.AreEqual("secondObject", e.Key.ToString());
+                    Assert.AreEqual(IoHashKey.FromName("secondObject"), e.Key);
                     Assert.AreEqual(objectHash, e.Blob);
                 }
 
@@ -226,7 +226,7 @@ namespace Horde.Storage.FunctionalTests.References
                     ReplicationLogEvent e = events.Events[2];
                     Assert.AreEqual(TestNamespace, e.Namespace);
                     Assert.AreEqual(TestBucket, e.Bucket);
-                    Assert.AreEqual("thirdObject", e.Key.ToString());
+                    Assert.AreEqual(IoHashKey.FromName("thirdObject"), e.Key);
                     Assert.AreEqual(objectHash, e.Blob);
                 }
             }
@@ -247,10 +247,10 @@ namespace Horde.Storage.FunctionalTests.References
             BlobIdentifier objectHash = BlobIdentifier.FromBlob(objectData);
 
             DateTime oldestTimestamp = DateTime.Now.AddDays(-1);
-            (string eventBucket, Guid eventId) = await _replicationLog.InsertAddEvent(TestNamespace, TestBucket, new KeyId("firstObject"), objectHash, oldestTimestamp);
-            await _replicationLog.InsertAddEvent(TestNamespace, TestBucket, new KeyId("secondObject"), objectHash, oldestTimestamp.AddHours(2.0));
-            await _replicationLog.InsertAddEvent(TestNamespace, TestBucket, new KeyId("thirdObject"), objectHash, oldestTimestamp.AddHours(3.0));
-            await _replicationLog.InsertAddEvent(TestNamespace, TestBucket, new KeyId("fourthObject"), objectHash, oldestTimestamp.AddDays(0.9));
+            (string eventBucket, Guid eventId) = await _replicationLog.InsertAddEvent(TestNamespace, TestBucket, IoHashKey.FromName("firstObject"), objectHash, oldestTimestamp);
+            await _replicationLog.InsertAddEvent(TestNamespace, TestBucket, IoHashKey.FromName("secondObject"), objectHash, oldestTimestamp.AddHours(2.0));
+            await _replicationLog.InsertAddEvent(TestNamespace, TestBucket, IoHashKey.FromName("thirdObject"), objectHash, oldestTimestamp.AddHours(3.0));
+            await _replicationLog.InsertAddEvent(TestNamespace, TestBucket, IoHashKey.FromName("fourthObject"), objectHash, oldestTimestamp.AddDays(0.9));
 
             {
                 HttpResponseMessage result = await _httpClient!.GetAsync(requestUri: $"api/v1/replication-log/incremental/{TestNamespace}?lastBucket={eventBucket}&lastEvent={eventId}");
@@ -270,7 +270,7 @@ namespace Horde.Storage.FunctionalTests.References
                     ReplicationLogEvent e = events.Events[0];
                     Assert.AreEqual(TestNamespace, e.Namespace);
                     Assert.AreEqual(TestBucket, e.Bucket);
-                    Assert.AreEqual("secondObject", e.Key.ToString());
+                    Assert.AreEqual(IoHashKey.FromName("secondObject"), e.Key);
                     Assert.AreEqual(objectHash, e.Blob);
                 }
 
@@ -278,7 +278,7 @@ namespace Horde.Storage.FunctionalTests.References
                     ReplicationLogEvent e = events.Events[1];
                     Assert.AreEqual(TestNamespace, e.Namespace);
                     Assert.AreEqual(TestBucket, e.Bucket);
-                    Assert.AreEqual("thirdObject", e.Key.ToString());
+                    Assert.AreEqual(IoHashKey.FromName("thirdObject"), e.Key);
                     Assert.AreEqual(objectHash, e.Blob);
                 }
 
@@ -286,7 +286,7 @@ namespace Horde.Storage.FunctionalTests.References
                     ReplicationLogEvent e = events.Events[2];
                     Assert.AreEqual(TestNamespace, e.Namespace);
                     Assert.AreEqual(TestBucket, e.Bucket);
-                    Assert.AreEqual("fourthObject", e.Key.ToString());
+                    Assert.AreEqual(IoHashKey.FromName("fourthObject"), e.Key);
                     Assert.AreEqual(objectHash, e.Blob);
                 }
             }
@@ -307,12 +307,12 @@ namespace Horde.Storage.FunctionalTests.References
 
             DateTime oldestTimestamp = DateTime.Now.AddDays(-3);
             // insert multiple objects in the same time bucket, verifying that we correctly get only the objects after this
-            await _replicationLog.InsertAddEvent(TestNamespace, TestBucket, new KeyId("firstObject"), objectHash, oldestTimestamp);
-            await _replicationLog.InsertAddEvent(TestNamespace, TestBucket, new KeyId("secondObject"), objectHash, oldestTimestamp.AddHours(2.0));
-            (string eventBucket, Guid eventId) = await _replicationLog.InsertAddEvent(TestNamespace, TestBucket, new KeyId("thirdObject"), objectHash, oldestTimestamp.AddHours(2.1));
-            await _replicationLog.InsertAddEvent(TestNamespace, TestBucket, new KeyId("fourthObject"), objectHash, oldestTimestamp.AddHours(2.11));
-            await _replicationLog.InsertAddEvent(TestNamespace, TestBucket, new KeyId("fifthObject"), objectHash, oldestTimestamp.AddHours(2.12));
-            await _replicationLog.InsertAddEvent(TestNamespace, TestBucket, new KeyId("sixthObject"), objectHash, oldestTimestamp.AddDays(2.13));
+            await _replicationLog.InsertAddEvent(TestNamespace, TestBucket, IoHashKey.FromName("firstObject"), objectHash, oldestTimestamp);
+            await _replicationLog.InsertAddEvent(TestNamespace, TestBucket, IoHashKey.FromName("secondObject"), objectHash, oldestTimestamp.AddHours(2.0));
+            (string eventBucket, Guid eventId) = await _replicationLog.InsertAddEvent(TestNamespace, TestBucket, IoHashKey.FromName("thirdObject"), objectHash, oldestTimestamp.AddHours(2.1));
+            await _replicationLog.InsertAddEvent(TestNamespace, TestBucket, IoHashKey.FromName("fourthObject"), objectHash, oldestTimestamp.AddHours(2.11));
+            await _replicationLog.InsertAddEvent(TestNamespace, TestBucket, IoHashKey.FromName("fifthObject"), objectHash, oldestTimestamp.AddHours(2.12));
+            await _replicationLog.InsertAddEvent(TestNamespace, TestBucket, IoHashKey.FromName("sixthObject"), objectHash, oldestTimestamp.AddDays(2.13));
 
             {
                 HttpResponseMessage result = await _httpClient!.GetAsync(requestUri: $"api/v1/replication-log/incremental/{TestNamespace}?lastBucket={eventBucket}&lastEvent={eventId}");
@@ -332,7 +332,7 @@ namespace Horde.Storage.FunctionalTests.References
                     ReplicationLogEvent e = events.Events[0];
                     Assert.AreEqual(TestNamespace, e.Namespace);
                     Assert.AreEqual(TestBucket, e.Bucket);
-                    Assert.AreEqual("fourthObject", e.Key.ToString());
+                    Assert.AreEqual(IoHashKey.FromName("fourthObject"), e.Key);
                     Assert.AreEqual(objectHash, e.Blob);
                 }
 
@@ -340,7 +340,7 @@ namespace Horde.Storage.FunctionalTests.References
                     ReplicationLogEvent e = events.Events[1];
                     Assert.AreEqual(TestNamespace, e.Namespace);
                     Assert.AreEqual(TestBucket, e.Bucket);
-                    Assert.AreEqual("fifthObject", e.Key.ToString());
+                    Assert.AreEqual(IoHashKey.FromName("fifthObject"), e.Key);
                     Assert.AreEqual(objectHash, e.Blob);
                 }
 
@@ -348,7 +348,7 @@ namespace Horde.Storage.FunctionalTests.References
                     ReplicationLogEvent e = events.Events[2];
                     Assert.AreEqual(TestNamespace, e.Namespace);
                     Assert.AreEqual(TestBucket, e.Bucket);
-                    Assert.AreEqual("sixthObject", e.Key.ToString());
+                    Assert.AreEqual(IoHashKey.FromName("sixthObject"), e.Key);
                     Assert.AreEqual(objectHash, e.Blob);
                 }
             }
@@ -368,10 +368,10 @@ namespace Horde.Storage.FunctionalTests.References
             BlobIdentifier objectHash = BlobIdentifier.FromBlob(objectData);
 
             DateTime oldestTimestamp = DateTime.Now.AddDays(-1);
-            (string eventBucket, Guid eventId) = await _replicationLog.InsertAddEvent(TestNamespace, TestBucket, new KeyId("firstObject"), objectHash, oldestTimestamp);
-            await _replicationLog.InsertAddEvent(TestNamespace, TestBucket, new KeyId("secondObject"), objectHash, oldestTimestamp.AddHours(1));
-            await _replicationLog.InsertAddEvent(TestNamespace, TestBucket, new KeyId("thirdObject"), objectHash, oldestTimestamp.AddHours(1.5));
-            await _replicationLog.InsertAddEvent(TestNamespace, TestBucket, new KeyId("fourthObject"), objectHash, oldestTimestamp.AddDays(0.9));
+            (string eventBucket, Guid eventId) = await _replicationLog.InsertAddEvent(TestNamespace, TestBucket, IoHashKey.FromName("firstObject"), objectHash, oldestTimestamp);
+            await _replicationLog.InsertAddEvent(TestNamespace, TestBucket, IoHashKey.FromName("secondObject"), objectHash, oldestTimestamp.AddHours(1));
+            await _replicationLog.InsertAddEvent(TestNamespace, TestBucket, IoHashKey.FromName("thirdObject"), objectHash, oldestTimestamp.AddHours(1.5));
+            await _replicationLog.InsertAddEvent(TestNamespace, TestBucket, IoHashKey.FromName("fourthObject"), objectHash, oldestTimestamp.AddDays(0.9));
 
             // start from the second event
             const int eventsToFetch = 2;
@@ -391,7 +391,7 @@ namespace Horde.Storage.FunctionalTests.References
                     ReplicationLogEvent e = events.Events[0];
                     Assert.AreEqual(TestNamespace, e.Namespace);
                     Assert.AreEqual(TestBucket, e.Bucket);
-                    Assert.AreEqual("secondObject", e.Key.ToString());
+                    Assert.AreEqual(IoHashKey.FromName("secondObject"), e.Key);
                     Assert.AreEqual(objectHash, e.Blob);
                 }
 
@@ -399,7 +399,7 @@ namespace Horde.Storage.FunctionalTests.References
                     ReplicationLogEvent e = events.Events[1];
                     Assert.AreEqual(TestNamespace, e.Namespace);
                     Assert.AreEqual(TestBucket, e.Bucket);
-                    Assert.AreEqual("thirdObject", e.Key.ToString());
+                    Assert.AreEqual(IoHashKey.FromName("thirdObject"), e.Key);
                     Assert.AreEqual(objectHash, e.Blob);
                 }
             }
@@ -419,7 +419,7 @@ namespace Horde.Storage.FunctionalTests.References
             byte[] objectData = writer.Save();
             BlobIdentifier objectHash = BlobIdentifier.FromBlob(objectData);
 
-            await _replicationLog.InsertAddEvent(TestNamespace, TestBucket, new KeyId("firstObject"), objectHash, DateTime.Now.AddDays(-1));
+            await _replicationLog.InsertAddEvent(TestNamespace, TestBucket, IoHashKey.FromName("firstObject"), objectHash, DateTime.Now.AddDays(-1));
 
             string eventBucket = "rep-00000000";
             Guid eventId = Guid.NewGuid();
@@ -465,10 +465,10 @@ namespace Horde.Storage.FunctionalTests.References
             BlobIdentifier objectHash = BlobIdentifier.FromBlob(objectData);
 
             DateTime oldestTimestamp = DateTime.Now.AddDays(-1);
-            await _replicationLog.InsertAddEvent(TestNamespace, TestBucket, new KeyId("firstObject"), objectHash, oldestTimestamp);
-            await _replicationLog.InsertAddEvent(TestNamespace, TestBucket, new KeyId("secondObject"), objectHash, oldestTimestamp.AddHours(1));
-            await _replicationLog.InsertAddEvent(TestNamespace, TestBucket, new KeyId("thirdObject"), objectHash, oldestTimestamp.AddHours(1.5));
-            await _replicationLog.InsertAddEvent(TestNamespace, TestBucket, new KeyId("fourthObject"), objectHash, oldestTimestamp.AddDays(0.9));
+            await _replicationLog.InsertAddEvent(TestNamespace, TestBucket, IoHashKey.FromName("firstObject"), objectHash, oldestTimestamp);
+            await _replicationLog.InsertAddEvent(TestNamespace, TestBucket, IoHashKey.FromName("secondObject"), objectHash, oldestTimestamp.AddHours(1));
+            await _replicationLog.InsertAddEvent(TestNamespace, TestBucket, IoHashKey.FromName("thirdObject"), objectHash, oldestTimestamp.AddHours(1.5));
+            await _replicationLog.InsertAddEvent(TestNamespace, TestBucket, IoHashKey.FromName("fourthObject"), objectHash, oldestTimestamp.AddDays(0.9));
 
             // create a snapshot
             ReplicationLogSnapshotBuilder snapshotBuilder = ActivatorUtilities.CreateInstance<ReplicationLogSnapshotBuilder>(_server!.Services);
@@ -504,10 +504,10 @@ namespace Horde.Storage.FunctionalTests.References
             BlobIdentifier objectHash = BlobIdentifier.FromBlob(objectData);
 
             DateTime oldestTimestamp = DateTime.Now.AddDays(-1);
-            await _replicationLog.InsertAddEvent(TestNamespace, TestBucket, new KeyId("firstObject"), objectHash, oldestTimestamp);
-            await _replicationLog.InsertAddEvent(TestNamespace, TestBucket, new KeyId("secondObject"), objectHash, oldestTimestamp.AddHours(1));
-            await _replicationLog.InsertAddEvent(TestNamespace, TestBucket, new KeyId("thirdObject"), objectHash, oldestTimestamp.AddHours(1.5));
-            (string lastEventBucket, Guid lastEventId) = await _replicationLog.InsertAddEvent(TestNamespace, TestBucket, new KeyId("fourthObject"), objectHash, oldestTimestamp.AddDays(0.7));
+            await _replicationLog.InsertAddEvent(TestNamespace, TestBucket, IoHashKey.FromName("firstObject"), objectHash, oldestTimestamp);
+            await _replicationLog.InsertAddEvent(TestNamespace, TestBucket, IoHashKey.FromName("secondObject"), objectHash, oldestTimestamp.AddHours(1));
+            await _replicationLog.InsertAddEvent(TestNamespace, TestBucket, IoHashKey.FromName("thirdObject"), objectHash, oldestTimestamp.AddHours(1.5));
+            (string lastEventBucket, Guid lastEventId) = await _replicationLog.InsertAddEvent(TestNamespace, TestBucket, IoHashKey.FromName("fourthObject"), objectHash, oldestTimestamp.AddDays(0.7));
 
             // verify the objects were added
             {
@@ -519,7 +519,7 @@ namespace Horde.Storage.FunctionalTests.References
                     ReplicationLogEvent e = logEvents[0];
                     Assert.AreEqual(TestNamespace, e.Namespace);
                     Assert.AreEqual(TestBucket, e.Bucket);
-                    Assert.AreEqual("firstObject", e.Key.ToString());
+                    Assert.AreEqual(IoHashKey.FromName("firstObject"), e.Key);
                     Assert.AreEqual(objectHash, e.Blob);
                 }
 
@@ -527,7 +527,7 @@ namespace Horde.Storage.FunctionalTests.References
                     ReplicationLogEvent e = logEvents[1];
                     Assert.AreEqual(TestNamespace, e.Namespace);
                     Assert.AreEqual(TestBucket, e.Bucket);
-                    Assert.AreEqual("secondObject", e.Key.ToString());
+                    Assert.AreEqual(IoHashKey.FromName("secondObject"), e.Key);
                     Assert.AreEqual(objectHash, e.Blob);
                 }
 
@@ -535,7 +535,7 @@ namespace Horde.Storage.FunctionalTests.References
                     ReplicationLogEvent e = logEvents[2];
                     Assert.AreEqual(TestNamespace, e.Namespace);
                     Assert.AreEqual(TestBucket, e.Bucket);
-                    Assert.AreEqual("thirdObject", e.Key.ToString());
+                    Assert.AreEqual(IoHashKey.FromName("thirdObject"), e.Key);
                     Assert.AreEqual(objectHash, e.Blob);
                 }
 
@@ -543,7 +543,7 @@ namespace Horde.Storage.FunctionalTests.References
                     ReplicationLogEvent e = logEvents[3];
                     Assert.AreEqual(TestNamespace, e.Namespace);
                     Assert.AreEqual(TestBucket, e.Bucket);
-                    Assert.AreEqual("fourthObject", e.Key.ToString());
+                    Assert.AreEqual(IoHashKey.FromName("fourthObject"), e.Key);
                     Assert.AreEqual(objectHash, e.Blob);
                 }
             }
@@ -567,10 +567,10 @@ namespace Horde.Storage.FunctionalTests.References
             Assert.AreEqual(lastEventId, snapshot.LastEvent);
             Assert.AreEqual(lastEventBucket, snapshot.LastBucket);
             
-            Assert.IsTrue(snapshot.LiveObjects.Any(o => o.Bucket == TestBucket && o.Key == new KeyId("firstObject")));
-            Assert.IsTrue(snapshot.LiveObjects.Any(o => o.Bucket == TestBucket && o.Key == new KeyId("secondObject")));
-            Assert.IsTrue(snapshot.LiveObjects.Any(o => o.Bucket == TestBucket && o.Key == new KeyId("thirdObject")));
-            Assert.IsTrue(snapshot.LiveObjects.Any(o => o.Bucket == TestBucket && o.Key == new KeyId("fourthObject")));
+            Assert.IsTrue(snapshot.LiveObjects.Any(o => o.Bucket == TestBucket && o.Key == IoHashKey.FromName("firstObject")));
+            Assert.IsTrue(snapshot.LiveObjects.Any(o => o.Bucket == TestBucket && o.Key == IoHashKey.FromName("secondObject")));
+            Assert.IsTrue(snapshot.LiveObjects.Any(o => o.Bucket == TestBucket && o.Key == IoHashKey.FromName("thirdObject")));
+            Assert.IsTrue(snapshot.LiveObjects.Any(o => o.Bucket == TestBucket && o.Key == IoHashKey.FromName("fourthObject")));
         }
 
 
@@ -586,10 +586,10 @@ namespace Horde.Storage.FunctionalTests.References
             BlobIdentifier objectHash = BlobIdentifier.FromBlob(objectData);
 
             DateTime oldestTimestamp = DateTime.Now.AddDays(-1);
-            await _replicationLog.InsertAddEvent(TestNamespace, TestBucket, new KeyId("firstObject"), objectHash, oldestTimestamp);
-            await _replicationLog.InsertAddEvent(TestNamespace, TestBucket, new KeyId("secondObject"), objectHash, oldestTimestamp.AddHours(1));
-            await _replicationLog.InsertAddEvent(TestNamespace, TestBucket, new KeyId("thirdObject"), objectHash, oldestTimestamp.AddHours(1.5));
-            (string lastEventBucket, Guid lastEventId) = await _replicationLog.InsertAddEvent(TestNamespace, TestBucket, new KeyId("fourthObject"), objectHash, oldestTimestamp.AddDays(0.9));
+            await _replicationLog.InsertAddEvent(TestNamespace, TestBucket, IoHashKey.FromName("firstObject"), objectHash, oldestTimestamp);
+            await _replicationLog.InsertAddEvent(TestNamespace, TestBucket, IoHashKey.FromName("secondObject"), objectHash, oldestTimestamp.AddHours(1));
+            await _replicationLog.InsertAddEvent(TestNamespace, TestBucket, IoHashKey.FromName("thirdObject"), objectHash, oldestTimestamp.AddHours(1.5));
+            (string lastEventBucket, Guid lastEventId) = await _replicationLog.InsertAddEvent(TestNamespace, TestBucket, IoHashKey.FromName("fourthObject"), objectHash, oldestTimestamp.AddDays(0.9));
 
             // verify the objects were added
             List<ReplicationLogEvent> logEvents = await _replicationLog.Get(TestNamespace, null, null).ToListAsync();
@@ -643,10 +643,10 @@ namespace Horde.Storage.FunctionalTests.References
             BlobIdentifier objectHash = BlobIdentifier.FromBlob(objectData);
 
             DateTime oldestTimestamp = DateTime.Now.AddDays(-1);
-            await _replicationLog.InsertAddEvent(TestNamespace, TestBucket, new KeyId("firstObject"), objectHash, oldestTimestamp);
-            await _replicationLog.InsertAddEvent(TestNamespace, TestBucket, new KeyId("secondObject"), objectHash, oldestTimestamp.AddHours(1));
-            await _replicationLog.InsertAddEvent(TestNamespace, TestBucket, new KeyId("thirdObject"), objectHash, oldestTimestamp.AddHours(1.5));
-            await _replicationLog.InsertAddEvent(TestNamespace, TestBucket, new KeyId("fourthObject"), objectHash, oldestTimestamp.AddDays(0.9));
+            await _replicationLog.InsertAddEvent(TestNamespace, TestBucket, IoHashKey.FromName("firstObject"), objectHash, oldestTimestamp);
+            await _replicationLog.InsertAddEvent(TestNamespace, TestBucket, IoHashKey.FromName("secondObject"), objectHash, oldestTimestamp.AddHours(1));
+            await _replicationLog.InsertAddEvent(TestNamespace, TestBucket, IoHashKey.FromName("thirdObject"), objectHash, oldestTimestamp.AddHours(1.5));
+            await _replicationLog.InsertAddEvent(TestNamespace, TestBucket, IoHashKey.FromName("fourthObject"), objectHash, oldestTimestamp.AddDays(0.9));
 
             // verify the objects were added
             List<ReplicationLogEvent> logEvents = await _replicationLog.Get(TestNamespace, null, null).ToListAsync();
@@ -668,8 +668,8 @@ namespace Horde.Storage.FunctionalTests.References
             ReplicationLogSnapshot snapshot = await ReplicationLogSnapshot.DeserializeSnapshot(blobContents.Stream);
 
             // insert more events
-            await _replicationLog.InsertAddEvent(TestNamespace, TestBucket, new KeyId("fifthObject"), objectHash, oldestTimestamp.AddDays(0.91));
-            await _replicationLog.InsertAddEvent(TestNamespace, TestBucket, new KeyId("sixthObject"), objectHash, oldestTimestamp.AddDays(0.92));
+            await _replicationLog.InsertAddEvent(TestNamespace, TestBucket, IoHashKey.FromName("fifthObject"), objectHash, oldestTimestamp.AddDays(0.91));
+            await _replicationLog.InsertAddEvent(TestNamespace, TestBucket, IoHashKey.FromName("sixthObject"), objectHash, oldestTimestamp.AddDays(0.92));
 
             // verify the new events can be found when resuming from the snapshot
             {
@@ -687,7 +687,7 @@ namespace Horde.Storage.FunctionalTests.References
                     ReplicationLogEvent e = events.Events[0];
                     Assert.AreEqual(TestNamespace, e.Namespace);
                     Assert.AreEqual(TestBucket, e.Bucket);
-                    Assert.AreEqual("fifthObject", e.Key.ToString());
+                    Assert.AreEqual(IoHashKey.FromName("fifthObject"), e.Key);
                     Assert.AreEqual(objectHash, e.Blob);
                 }
 
@@ -695,7 +695,7 @@ namespace Horde.Storage.FunctionalTests.References
                     ReplicationLogEvent e = events.Events[1];
                     Assert.AreEqual(TestNamespace, e.Namespace);
                     Assert.AreEqual(TestBucket, e.Bucket);
-                    Assert.AreEqual("sixthObject", e.Key.ToString());
+                    Assert.AreEqual(IoHashKey.FromName("sixthObject"), e.Key);
                     Assert.AreEqual(objectHash, e.Blob);
                 }
             }
@@ -718,7 +718,7 @@ namespace Horde.Storage.FunctionalTests.References
             List<BlobIdentifier> createdSnapshots = new List<BlobIdentifier>();
             for (int i = 0; i < countOfSnapshotsToCreate ; i++)
             {
-                await _replicationLog.InsertAddEvent(TestNamespace, TestBucket, new KeyId($"object {i}"), objectHash);
+                await _replicationLog.InsertAddEvent(TestNamespace, TestBucket, IoHashKey.FromName($"object {i}"), objectHash);
 
                 ReplicationLogSnapshotBuilder snapshotBuilder = ActivatorUtilities.CreateInstance<ReplicationLogSnapshotBuilder>(_server!.Services);
                 Assert.IsNotNull(snapshotBuilder);
diff --git a/Engine/Source/Programs/Horde/HordeStorage/Jupiter.Common/Implementation/IoHashKey.cs b/Engine/Source/Programs/Horde/HordeStorage/Jupiter.Common/Implementation/IoHashKey.cs
new file mode 100644
--- /dev/null
+++ b/Engine/Source/Programs/Horde/HordeStorage/Jupiter.Common/Implementation/IoHashKey.cs
@@ -0,0 +1,121 @@
+// Copyright Epic Games, Inc. All Rights Reserved.
+
+using System;
+using System.ComponentModel;
+using System.Globalization;
+using System.Text;
+using Newtonsoft.Json;
+
+namespace Jupiter.Implementation
+{
+    [JsonConverter(typeof(IoHashKeyJsonConverter))]
+    [TypeConverter(typeof(IoHashKeyTypeConverter))]
+
+    public readonly struct IoHashKey: IEquatable<IoHashKey>
+    {
+        public IoHashKey(string key)
+        {
+             _text = key.ToLower();
+ 
+            if (_text.Length != 40)
+            {
+                throw new ArgumentException("IoHashKeys must be exactly 40 bytes.");
+            }
+ 
+            for (int Idx = 0; Idx < _text.Length; Idx++)
+            {
+                if (!IsValidCharacter(_text[Idx]))
+                {
+                    throw new ArgumentException($"{_text} is not a valid namespace id");
+                }
+            }
+        }
+
+        static bool IsValidCharacter(char c)
+        {
+            if (c >= 'a' && c <= 'z')
+            {
+                return true;
+            }
+            if (c >= '0' && c <= '9')
+            {
+                return true;
+            }
+            return false;
+        }
+        readonly string _text;
+
+        public bool Equals(IoHashKey other)
+        {
+            return string.Equals(_text , other._text, StringComparison.Ordinal);
+        }
+
+        public override bool Equals(object? obj)
+        {
+            return obj is IoHashKey other && Equals(other);
+        }
+
+        public override int GetHashCode()
+        {
+            return _text.GetHashCode();
+        }
+
+        public override string ToString()
+        {
+            return _text;
+        }
+
+        public static bool operator ==(IoHashKey left, IoHashKey right)
+        {
+            return left.Equals(right);
+        }
+ 
+        public static bool operator !=(IoHashKey left, IoHashKey right)
+        {
+            return !left.Equals(right);
+        }
+
+        public static IoHashKey FromName(string s)
+        {
+            return new IoHashKey(ContentHash.FromBlob(Encoding.UTF8.GetBytes(s)).ToString());
+        }
+    }
+
+    public class IoHashKeyJsonConverter: JsonConverter<IoHashKey>
+    {
+        public override void WriteJson(JsonWriter writer, IoHashKey value, JsonSerializer serializer)
+        {
+            writer.WriteValue(value!.ToString());
+        }
+
+        public override IoHashKey ReadJson(JsonReader reader, Type objectType, IoHashKey existingValue, bool hasExistingValue, Newtonsoft.Json.JsonSerializer serializer)
+        {
+            string? s = (string?)reader.Value;
+
+            return new IoHashKey(s!);
+        }
+    }
+
+    public sealed class IoHashKeyTypeConverter : TypeConverter
+    {
+        /// <inheritdoc/>
+        public override bool CanConvertFrom(ITypeDescriptorContext? context, Type sourceType)
+        {
+            if (sourceType == typeof(string))
+                return true;
+
+            return base.CanConvertFrom(context, sourceType);
+        }
+ 
+        public override object? ConvertFrom(ITypeDescriptorContext? context, CultureInfo? culture, object value)
+        {
+            if (value is string s)
+            {
+                return new IoHashKey(s);
+            }
+
+            return base.ConvertFrom(context, culture, value);  
+        }
+    }
+}
+
