Horde Storage

Added concept of storage pools which allows us to control which namespaces share storage and which do not. Currently only implemented for S3 (resulting in seperate s3 buckets).
Will be used to seperate ddc1 and ddc2 data from each other greatly simplifying the GC processes.

Also moved the auth map from a seperate setting to a container for everything related to the namespace (Claims and Storage Pool currently with more things to come).

diff --git a/Engine/Source/Programs/Horde/HordeStorage/Callisto/appsettings.Testing.json b/Engine/Source/Programs/Horde/HordeStorage/Callisto/appsettings.Testing.json
--- a/Engine/Source/Programs/Horde/HordeStorage/Callisto/appsettings.Testing.json
+++ b/Engine/Source/Programs/Horde/HordeStorage/Callisto/appsettings.Testing.json
@@ -18,12 +18,6 @@
     "Method": "Disabled"
   },
 
-  "Authorization": {
-    "NamespaceToClaim": {
-      "ue4.ddc": "*"
-    }
-  },
-
   "Jupiter": {
     "CurrentSite": "Test"
   },
diff --git a/Engine/Source/Programs/Horde/HordeStorage/Horde.Storage/Implementation/Blob/AmazonS3Store.cs b/Engine/Source/Programs/Horde/HordeStorage/Horde.Storage/Implementation/Blob/AmazonS3Store.cs
--- a/Engine/Source/Programs/Horde/HordeStorage/Horde.Storage/Implementation/Blob/AmazonS3Store.cs
+++ b/Engine/Source/Programs/Horde/HordeStorage/Horde.Storage/Implementation/Blob/AmazonS3Store.cs
@@ -8,10 +8,12 @@ using System.Threading.Tasks;
 using Amazon.S3;
 using Amazon.S3.Model;
 using Horde.Storage.Implementation.Blob;
+using Jupiter;
 using Jupiter.Common.Implementation;
 using Jupiter.Implementation;
 using Microsoft.Extensions.Options;
 using Serilog;
+using KeyNotFoundException = System.Collections.Generic.KeyNotFoundException;
 
 namespace Horde.Storage.Implementation
 {
@@ -20,14 +22,16 @@ namespace Horde.Storage.Implementation
         private readonly ILogger _logger = Log.ForContext<AmazonS3Store>();
         private readonly IAmazonS3 _amazonS3;
         private readonly IBlobIndex _blobIndex;
+        private readonly IOptionsMonitor<NamespaceSettings> _namespaceSettings;
         private readonly S3Settings _settings;
         private readonly HashSet<string> _bucketAccessPolicyApplied = new HashSet<string>();
         private readonly HashSet<string> _bucketExistenceChecked = new HashSet<string>();
 
-        public AmazonS3Store(IAmazonS3 amazonS3, IOptionsMonitor<S3Settings> settings, IBlobIndex blobIndex)
+        public AmazonS3Store(IAmazonS3 amazonS3, IOptionsMonitor<S3Settings> settings, IBlobIndex blobIndex, IOptionsMonitor<NamespaceSettings> namespaceSettings)
         {
             _amazonS3 = amazonS3;
             _blobIndex = blobIndex;
+            _namespaceSettings = namespaceSettings;
             _settings = settings.CurrentValue;
         }
 
@@ -120,7 +124,16 @@ namespace Horde.Storage.Implementation
 
         private string GetBucketName(NamespaceId ns)
         {
-            return $"{_settings.BucketName}";
+            try
+            {
+                string storagePool = _namespaceSettings.CurrentValue.GetPoliciesForNs(ns).StoragePool;
+                string storagePoolSuffix = string.IsNullOrEmpty(storagePool) ? "" : $"-{storagePool}";
+                return $"{_settings.BucketName}{storagePoolSuffix}";
+            }
+            catch (KeyNotFoundException)
+            {
+                throw new NamespaceNotFoundException(ns);
+            }
         }
 
 
diff --git a/Engine/Source/Programs/Horde/HordeStorage/Horde.Storage/Implementation/Blob/FileSystemStore.cs b/Engine/Source/Programs/Horde/HordeStorage/Horde.Storage/Implementation/Blob/FileSystemStore.cs
--- a/Engine/Source/Programs/Horde/HordeStorage/Horde.Storage/Implementation/Blob/FileSystemStore.cs
+++ b/Engine/Source/Programs/Horde/HordeStorage/Horde.Storage/Implementation/Blob/FileSystemStore.cs
@@ -168,7 +168,7 @@ namespace Horde.Storage.Implementation
 
         private IEnumerable<(BlobIdentifier, DateTime)> DoListOldObjects(NamespaceId ns)
         {
-            DirectoryInfo di = new DirectoryInfo(Path.Combine(_settings.CurrentValue.RootDir, ns.ToString()));
+            DirectoryInfo di = GetFilesystemPath(ns);
             if (!di.Exists)
             {
                 yield break;
diff --git a/Engine/Source/Programs/Horde/HordeStorage/Horde.Storage/Tests/UnitTests/S3StoreTests.cs b/Engine/Source/Programs/Horde/HordeStorage/Horde.Storage/Tests/UnitTests/S3StoreTests.cs
--- a/Engine/Source/Programs/Horde/HordeStorage/Horde.Storage/Tests/UnitTests/S3StoreTests.cs
+++ b/Engine/Source/Programs/Horde/HordeStorage/Horde.Storage/Tests/UnitTests/S3StoreTests.cs
@@ -1,6 +1,7 @@
 ﻿// Copyright Epic Games, Inc. All Rights Reserved.
 
 using System;
+using System.Collections.Generic;
 using System.IO;
 using System.Text;
 using System.Threading.Tasks;
@@ -8,6 +9,7 @@ using Amazon.S3;
 using Amazon.S3.Model;
 using Horde.Storage.Implementation;
 using Horde.Storage.Implementation.Blob;
+using Jupiter;
 using Jupiter.Implementation;
 using Microsoft.Extensions.Options;
 using Microsoft.VisualStudio.TestTools.UnitTesting;
@@ -27,7 +29,18 @@ namespace Horde.Storage.UnitTests
             S3Settings settings = new S3Settings{BucketName = "tests-foo"};
             IOptionsMonitor<S3Settings> settingsMock = Mock.Of<IOptionsMonitor<S3Settings>>(_ => _.CurrentValue == settings);
 
-            AmazonS3Store store = new AmazonS3Store(s3Mock.Object, settingsMock, Mock.Of<IBlobIndex>());
+            NamespaceSettings namespaceSettings = new NamespaceSettings{ Policies = new Dictionary<string, NamespaceSettings.PerNamespaceSettings>()
+            {
+                {
+                    Namespace.ToString(), new NamespaceSettings.PerNamespaceSettings {
+                        StoragePool = ""
+                    }
+                }
+                
+            }};
+            IOptionsMonitor<NamespaceSettings> namespaceSettingsMock = Mock.Of<IOptionsMonitor<NamespaceSettings>>(_ => _.CurrentValue == namespaceSettings);
+
+            AmazonS3Store store = new AmazonS3Store(s3Mock.Object, settingsMock, Mock.Of<IBlobIndex>(), namespaceSettingsMock);
             byte[] content = Encoding.ASCII.GetBytes("test content");
             BlobIdentifier blobIdentifier = BlobIdentifier.FromBlob(content);
             Task task = store.PutObject(Namespace, content.AsMemory(), blobIdentifier);
@@ -42,6 +55,39 @@ namespace Horde.Storage.UnitTests
             s3Mock.VerifyNoOtherCalls();
         }
 
+        [TestMethod]
+        public async Task PutBufferSuccessStoragePool()
+        {
+            Mock<IAmazonS3> s3Mock = new Mock<IAmazonS3>();
+            S3Settings settings = new S3Settings{BucketName = "tests-foo"};
+            IOptionsMonitor<S3Settings> settingsMock = Mock.Of<IOptionsMonitor<S3Settings>>(_ => _.CurrentValue == settings);
+
+            NamespaceSettings namespaceSettings = new NamespaceSettings{ Policies = new Dictionary<string, NamespaceSettings.PerNamespaceSettings>()
+            {
+                {
+                    Namespace.ToString(), new NamespaceSettings.PerNamespaceSettings {
+                        StoragePool = "storagepool"
+                    }
+                }
+                
+            }};
+            IOptionsMonitor<NamespaceSettings> namespaceSettingsMock = Mock.Of<IOptionsMonitor<NamespaceSettings>>(_ => _.CurrentValue == namespaceSettings);
+
+            AmazonS3Store store = new AmazonS3Store(s3Mock.Object, settingsMock, Mock.Of<IBlobIndex>(), namespaceSettingsMock);
+            byte[] content = Encoding.ASCII.GetBytes("test content");
+            BlobIdentifier blobIdentifier = BlobIdentifier.FromBlob(content);
+            Task task = store.PutObject(Namespace, content.AsMemory(), blobIdentifier);
+            await task;
+
+            s3Mock.Verify( s3 => s3.DoesS3BucketExistAsync("tests-foo-storagepool"));
+            s3Mock.Verify( s3 => s3.PutBucketAsync(It.Is<PutBucketRequest>(p => p.BucketName =="tests-foo-storagepool"), default));
+            s3Mock.Verify(s3 =>
+                s3.PutObjectAsync(It.Is<PutObjectRequest>(p => p.BucketName == "tests-foo-storagepool" && p.Key == blobIdentifier.AsS3Key()), default));
+
+            s3Mock.Verify( s3 => s3.PutPublicAccessBlockAsync(It.IsAny<PutPublicAccessBlockRequest>(), default));
+            s3Mock.VerifyNoOtherCalls();
+        }
+
         [TestMethod]
         [ExpectedException(typeof(Exception))]
         public async Task PutBufferError()
@@ -52,9 +98,19 @@ namespace Horde.Storage.UnitTests
             Mock<IAmazonS3> s3Mock = new Mock<IAmazonS3>();
             S3Settings settings = new S3Settings{BucketName = "tests-foo"};
             IOptionsMonitor<S3Settings> settingsMock = Mock.Of<IOptionsMonitor<S3Settings>>(_ => _.CurrentValue == settings);
+            NamespaceSettings namespaceSettings = new NamespaceSettings{ Policies = new Dictionary<string, NamespaceSettings.PerNamespaceSettings>()
+            {
+                {
+                    Namespace.ToString(), new NamespaceSettings.PerNamespaceSettings {
+                        StoragePool = ""
+                    }
+                }
+                
+            }};
+            IOptionsMonitor<NamespaceSettings> namespaceSettingsMock = Mock.Of<IOptionsMonitor<NamespaceSettings>>(_ => _.CurrentValue == namespaceSettings);
 
             s3Mock.Setup(s3 => s3.PutObjectAsync(It.IsAny<PutObjectRequest>(), default)).Throws<Exception>();
-            AmazonS3Store store = new AmazonS3Store(s3Mock.Object, settingsMock, Mock.Of<IBlobIndex>());
+            AmazonS3Store store = new AmazonS3Store(s3Mock.Object, settingsMock, Mock.Of<IBlobIndex>(), namespaceSettingsMock);
             Task task = store.PutObject(Namespace, content, blob);
             await task;
 
@@ -75,7 +131,18 @@ namespace Horde.Storage.UnitTests
             S3Settings settings = new S3Settings{BucketName = "tests-foo"};
             IOptionsMonitor<S3Settings> settingsMock = Mock.Of<IOptionsMonitor<S3Settings>>(_ => _.CurrentValue == settings);
 
-            AmazonS3Store store = new AmazonS3Store(s3Mock.Object, settingsMock, Mock.Of<IBlobIndex>());
+            NamespaceSettings namespaceSettings = new NamespaceSettings{ Policies = new Dictionary<string, NamespaceSettings.PerNamespaceSettings>()
+            {
+                {
+                    Namespace.ToString(), new NamespaceSettings.PerNamespaceSettings {
+                        StoragePool = ""
+                    }
+                }
+                
+            }};
+            IOptionsMonitor<NamespaceSettings> namespaceSettingsMock = Mock.Of<IOptionsMonitor<NamespaceSettings>>(_ => _.CurrentValue == namespaceSettings);
+
+            AmazonS3Store store = new AmazonS3Store(s3Mock.Object, settingsMock, Mock.Of<IBlobIndex>(), namespaceSettingsMock);
             await using BlobContents blobContents = await store.GetObject(Namespace, blob);
 
             s3Mock.Verify(s3 => s3.GetObjectAsync("tests-foo", blob.AsS3Key(), default));
@@ -90,7 +157,18 @@ namespace Horde.Storage.UnitTests
             S3Settings settings = new S3Settings{BucketName = "tests-foo"};
             IOptionsMonitor<S3Settings> settingsMock = Mock.Of<IOptionsMonitor<S3Settings>>(_ => _.CurrentValue == settings);
 
-            AmazonS3Store store = new AmazonS3Store(s3Mock.Object, settingsMock, Mock.Of<IBlobIndex>());
+            NamespaceSettings namespaceSettings = new NamespaceSettings{ Policies = new Dictionary<string, NamespaceSettings.PerNamespaceSettings>()
+            {
+                {
+                    Namespace.ToString(), new NamespaceSettings.PerNamespaceSettings {
+                        StoragePool = ""
+                    }
+                }
+                
+            }};
+            IOptionsMonitor<NamespaceSettings> namespaceSettingsMock = Mock.Of<IOptionsMonitor<NamespaceSettings>>(_ => _.CurrentValue == namespaceSettings);
+
+            AmazonS3Store store = new AmazonS3Store(s3Mock.Object, settingsMock, Mock.Of<IBlobIndex>(), namespaceSettingsMock);
             Task task = store.DeleteObject(Namespace, blob);
             await task;
 
diff --git a/Engine/Source/Programs/Horde/HordeStorage/Horde.Storage/appsettings.Development.json b/Engine/Source/Programs/Horde/HordeStorage/Horde.Storage/appsettings.Development.json
--- a/Engine/Source/Programs/Horde/HordeStorage/Horde.Storage/appsettings.Development.json
+++ b/Engine/Source/Programs/Horde/HordeStorage/Horde.Storage/appsettings.Development.json
@@ -29,6 +29,19 @@
     "contentIdStoreImplementation": "memory"
   },
 
+  // settings s3 defaults to use a local minio instance, same as tests, not enabled by default but useful to do local testing
+  "S3": {
+    "ConnectionString": "http://127.0.0.1:9000",
+    "ForceAWSPathStyle": true, // required by minio,
+    "SetBucketPolicies": false, // minio does not support bucket access policies
+    "BucketName": "local-horde-storage-bucket"
+  },
+  "AWSCredentials": {
+    // minio and health checks expect access keys to have been set
+    "AwsAccessKey": "FSYPTVCPKC2IV5D1U8RE",
+    "AwsSecretKey": "WaUf4Zq+J6YMlvbCVvTRkQJIPc7DvZwTm5UP0tPJ"
+  },
+
   "Callisto": {
     "ConnectionString": "http://localhost"
   },
@@ -57,12 +70,6 @@
     "Method": "Disabled"
   },
 
-  "Authorization": {
-    "NamespaceToClaim": {
-      "ue4.ddc": "*"
-    }
-  },
-
   "Mongo": {
     "ConnectionString": "mongodb://root:abc123@localhost:27017",
     "RequireTls12": false
@@ -83,6 +90,18 @@
     "Enabled": false
   },
 
+  "Namespaces": {
+    "Policies": {
+      "ue4.ddc": {
+        "Claims": [ "UE4" ]
+      },
+      "ue.ddc": {
+        "Claims": [ "UE4" ],
+        "StoragePool": "ddc2" 
+      }
+    }
+  },
+
   "Kestrel": {
     "Endpoints": {
       "Http": {
diff --git a/Engine/Source/Programs/Horde/HordeStorage/Horde.Storage/appsettings.Testing.json b/Engine/Source/Programs/Horde/HordeStorage/Horde.Storage/appsettings.Testing.json
--- a/Engine/Source/Programs/Horde/HordeStorage/Horde.Storage/appsettings.Testing.json
+++ b/Engine/Source/Programs/Horde/HordeStorage/Horde.Storage/appsettings.Testing.json
@@ -74,10 +74,20 @@
     "Method": "Disabled"
   },
 
-  "Authorization": {
-    "NamespaceToClaim": {
-      "ue4.ddc": "*",
-      "TestNamespace": "*"
+  "Namespaces": {
+    "Policies": {
+      "ue4.ddc": {
+        "Claims": [ "*" ]
+      },
+      "TestNamespace": {
+        "Claims": [ "*" ]
+      },
+      "test-namespace": {
+        "Claims": [ "*" ]
+      },
+      "testbucket": {
+        "Claims": [ "*" ]
+      }
     }
   },
 
diff --git a/Engine/Source/Programs/Horde/HordeStorage/Jupiter.Common/BaseStartup.cs b/Engine/Source/Programs/Horde/HordeStorage/Jupiter.Common/BaseStartup.cs
--- a/Engine/Source/Programs/Horde/HordeStorage/Jupiter.Common/BaseStartup.cs
+++ b/Engine/Source/Programs/Horde/HordeStorage/Jupiter.Common/BaseStartup.cs
@@ -11,6 +11,7 @@ using System.Text;
 using System.Threading.Tasks;
 using Amazon;
 using Datadog.Trace;
+using Jupiter.Implementation;
 using Microsoft.AspNetCore.Authentication;
 using Microsoft.AspNetCore.Authorization;
 using Microsoft.AspNetCore.Builder;
@@ -67,8 +68,8 @@ namespace Jupiter
 
             services.AddOptions<ServiceAccountAuthOptions>().Bind(Configuration.GetSection("ServiceAccounts")).ValidateDataAnnotations();
 
-            services.AddOptions<AuthorizationSettings>().Bind(Configuration.GetSection("Authorization")).ValidateDataAnnotations();
             services.AddOptions<JupiterSettings>().Bind(Configuration.GetSection("Jupiter")).ValidateDataAnnotations();
+            services.AddOptions<NamespaceSettings>().Bind(Configuration.GetSection("Namespaces")).ValidateDataAnnotations();
 
             // this is the same as invoke MvcBuilder.AddJsonOptions but with a service provider passed so we can use DI in the options creation
             // see https://stackoverflow.com/questions/53288633/net-core-api-custom-json-resolver-based-on-request-values
@@ -462,6 +463,29 @@ namespace Jupiter
         //public Dictionary<string, int> DisableAuthOnPorts { get; set; } = new ();
     }
 
+    public class NamespaceSettings
+    {
+        public Dictionary<string, PerNamespaceSettings> Policies { get; set; } = new Dictionary<string, PerNamespaceSettings>();
+
+        public class PerNamespaceSettings
+        {
+            public string[] Claims { get; set; } = Array.Empty<string>();
+            public string StoragePool { get; set; } = "";
+        }
+
+        public PerNamespaceSettings GetPoliciesForNs(NamespaceId ns)
+        {
+            return Policies[ns.ToString()];
+        /*
+            if(Policies.TryGetValue(ns.ToString(), out PerNamespaceSettings? settings))
+            {
+                return settings;
+            }
+
+            return DefaultNamespaceSettings;*/
+        }
+    }
+
     public class DatadogTraceMiddleware
     {
         private readonly RequestDelegate _next;
