Horde Storage
Added tracking of which blobs are referenced by which refs, a inverse lookup from what we normally do. Will be used by the GC to determine which blobs can be removed.

diff --git a/Engine/Source/Programs/Horde/HordeStorage/Horde.Storage/Horde.Storage.Startup.cs b/Engine/Source/Programs/Horde/HordeStorage/Horde.Storage/Horde.Storage.Startup.cs
--- a/Engine/Source/Programs/Horde/HordeStorage/Horde.Storage/Horde.Storage.Startup.cs
+++ b/Engine/Source/Programs/Horde/HordeStorage/Horde.Storage/Horde.Storage.Startup.cs
@@ -349,6 +349,9 @@ namespace Horde.Storage
             replicatedSession.Execute(new SimpleStatement("CREATE TYPE IF NOT EXISTS blob_identifier (hash blob)"));
             replicatedSession.UserDefinedTypes.Define(UdtMap.For<ScyllaBlobIdentifier>("blob_identifier", DefaultKeyspaceName));
 
+            replicatedSession.Execute(new SimpleStatement("CREATE TYPE IF NOT EXISTS object_reference (bucket text, key text)"));
+            replicatedSession.UserDefinedTypes.Define(UdtMap.For<ScyllaObjectReference>("object_reference", DefaultKeyspaceName));
+
             string localKeyspaceName = $"jupiter_local_{settings.LocalKeyspaceSuffix}";
 
             Dictionary<string, string> replicationStrategyLocal = ReplicationStrategies.CreateSimpleStrategyReplicationProperty(2);
diff --git a/Engine/Source/Programs/Horde/HordeStorage/Horde.Storage/Implementation/Blob/IBlobIndex.cs b/Engine/Source/Programs/Horde/HordeStorage/Horde.Storage/Implementation/Blob/IBlobIndex.cs
--- a/Engine/Source/Programs/Horde/HordeStorage/Horde.Storage/Implementation/Blob/IBlobIndex.cs
+++ b/Engine/Source/Programs/Horde/HordeStorage/Horde.Storage/Implementation/Blob/IBlobIndex.cs
@@ -10,7 +10,9 @@ namespace Horde.Storage.Implementation.Blob
     {
         public class BlobInfo
         {
-            public HashSet<string> Regions = new HashSet<string>();
+            public HashSet<string> Regions { get; init; } = new HashSet<string>();
+            public NamespaceId Namespace { get; init; }
+            public List<(BucketId, IoHashKey)> References { get; init; } = new List<(BucketId, IoHashKey)>();
         }
 
         Task AddBlobToIndex(NamespaceId ns, BlobIdentifier id);
@@ -19,5 +21,6 @@ namespace Horde.Storage.Implementation.Blob
 
         Task<bool> RemoveBlobFromIndex(NamespaceId ns, BlobIdentifier id);
         Task<bool> BlobExistsInRegion(NamespaceId ns, BlobIdentifier blobIdentifier);
+        Task AddRefToBlobs(NamespaceId ns, BucketId bucket, IoHashKey key, BlobIdentifier[] blobs);
     }
 }
diff --git a/Engine/Source/Programs/Horde/HordeStorage/Horde.Storage/Implementation/Blob/MemoryBlobIndex.cs b/Engine/Source/Programs/Horde/HordeStorage/Horde.Storage/Implementation/Blob/MemoryBlobIndex.cs
--- a/Engine/Source/Programs/Horde/HordeStorage/Horde.Storage/Implementation/Blob/MemoryBlobIndex.cs
+++ b/Engine/Source/Programs/Horde/HordeStorage/Horde.Storage/Implementation/Blob/MemoryBlobIndex.cs
@@ -1,5 +1,6 @@
 // Copyright Epic Games, Inc. All Rights Reserved.
 
+using System;
 using System.Collections.Concurrent;
 using System.Collections.Generic;
 using System.Threading.Tasks;
@@ -13,6 +14,7 @@ public class MemoryBlobIndex : IBlobIndex
 {
     private class MemoryBlobInfo : IBlobIndex.BlobInfo
     {
+
     }
 
     private readonly ConcurrentDictionary<NamespaceId, ConcurrentDictionary<BlobIdentifier, MemoryBlobInfo>> _index = new ();
@@ -31,10 +33,7 @@ public class MemoryBlobIndex : IBlobIndex
     public Task AddBlobToIndex(NamespaceId ns, BlobIdentifier id)
     {
         ConcurrentDictionary<BlobIdentifier, MemoryBlobInfo> index = GetNamespaceContainer(ns);
-        index[id] = new MemoryBlobInfo
-        {
-            Regions = new HashSet<string> { _jupiterSettings.CurrentValue.CurrentSite }
-        };
+        index[id] = NewBlobInfo(ns);
         return Task.CompletedTask;
     }
 
@@ -60,4 +59,35 @@ public class MemoryBlobIndex : IBlobIndex
         IBlobIndex.BlobInfo? blobInfo = await GetBlobInfo(ns, blobIdentifier);
         return blobInfo?.Regions.Contains(_jupiterSettings.CurrentValue.CurrentSite) ?? false;
     }
+
+    public Task AddRefToBlobs(NamespaceId ns, BucketId bucket, IoHashKey key, BlobIdentifier[] blobs)
+    {
+        foreach (BlobIdentifier id in blobs)
+        {
+            ConcurrentDictionary<BlobIdentifier, MemoryBlobInfo> index = GetNamespaceContainer(ns);
+
+            index.AddOrUpdate(id, _ =>
+            {
+                MemoryBlobInfo info = NewBlobInfo(ns);
+                info.References.Add((bucket, key));
+                return info;
+            }, (_, info) =>
+            {
+                info.References.Add((bucket, key));
+                return info;
+            });
+        };
+
+        return Task.CompletedTask;
+    }
+
+    private MemoryBlobInfo NewBlobInfo(NamespaceId ns)
+    {
+        MemoryBlobInfo info = new MemoryBlobInfo
+        {
+            Regions = new HashSet<string> { _jupiterSettings.CurrentValue.CurrentSite },
+            Namespace = ns,
+        };
+        return info;
+    }
 }
diff --git a/Engine/Source/Programs/Horde/HordeStorage/Horde.Storage/Implementation/Blob/ScyllaBlobIndex.cs b/Engine/Source/Programs/Horde/HordeStorage/Horde.Storage/Implementation/Blob/ScyllaBlobIndex.cs
--- a/Engine/Source/Programs/Horde/HordeStorage/Horde.Storage/Implementation/Blob/ScyllaBlobIndex.cs
+++ b/Engine/Source/Programs/Horde/HordeStorage/Horde.Storage/Implementation/Blob/ScyllaBlobIndex.cs
@@ -1,6 +1,5 @@
 // Copyright Epic Games, Inc. All Rights Reserved.
 
-using System;
 using System.Collections.Generic;
 using System.Linq;
 using System.Threading.Tasks;
@@ -29,6 +28,7 @@ public class ScyllaBlobIndex : IBlobIndex
             namespace text,
             blob_id frozen<blob_identifier>,
             regions set<text>,
+            references set<frozen<object_reference>>,
             PRIMARY KEY ((namespace, blob_id))
         );"
         ));
@@ -38,21 +38,26 @@ public class ScyllaBlobIndex : IBlobIndex
     {
         using Scope _ = Tracer.Instance.StartActive("scylla.insert_blob_index");
 
-        await _mapper.UpdateAsync<ScyllaBlobIndexTable>("SET regions = regions + ? WHERE namespace = ? AND blob_id = ?", new string[] {_jupiterSettings.CurrentValue.CurrentSite}, ns.ToString(), new ScyllaBlobIdentifier(id));
+        await _mapper.UpdateAsync<ScyllaBlobIndexTable>("SET regions = regions + ? WHERE namespace = ? AND blob_id = ?",
+            new string[] { _jupiterSettings.CurrentValue.CurrentSite }, ns.ToString(), new ScyllaBlobIdentifier(id));
     }
 
     public async Task<IBlobIndex.BlobInfo?> GetBlobInfo(NamespaceId ns, BlobIdentifier id)
     {
         using Scope _ = Tracer.Instance.StartActive("scylla.fetch_blob_index");
 
-        ScyllaBlobIndexTable? blobIndex = await _mapper.FirstOrDefaultAsync<ScyllaBlobIndexTable>("WHERE namespace = ? AND blob_id = ?", ns.ToString(), new ScyllaBlobIdentifier(id));
+        ScyllaBlobIndexTable? blobIndex =
+            await _mapper.FirstOrDefaultAsync<ScyllaBlobIndexTable>("WHERE namespace = ? AND blob_id = ?",
+                ns.ToString(), new ScyllaBlobIdentifier(id));
 
         if (blobIndex == null)
             return null;
 
-        return new IBlobIndex.BlobInfo()
+        return new IBlobIndex.BlobInfo
         {
-            Regions = blobIndex.Regions
+            Regions = blobIndex.Regions,
+            Namespace = new NamespaceId(blobIndex.Namespace),
+            References = blobIndex.References.Select(reference => reference.AsTuple()).ToList()
         };
     }
 
@@ -60,7 +65,8 @@ public class ScyllaBlobIndex : IBlobIndex
     {
         // TODO: Should this only remove the current region, and the actual row isnt removed until all regions have been removed? Seems overly complicated
         using Scope _ = Tracer.Instance.StartActive("scylla.remove_from_blob_index");
-        await _mapper.DeleteAsync<ScyllaBlobIndexTable>("WHERE namespace = ? AND blob_id = ?", ns.ToString(), new ScyllaBlobIdentifier(id));
+        await _mapper.DeleteAsync<ScyllaBlobIndexTable>("WHERE namespace = ? AND blob_id = ?", ns.ToString(),
+            new ScyllaBlobIdentifier(id));
 
         return true;
     }
@@ -70,6 +76,27 @@ public class ScyllaBlobIndex : IBlobIndex
         IBlobIndex.BlobInfo? blobInfo = await GetBlobInfo(ns, blobIdentifier);
         return blobInfo?.Regions.Contains(_jupiterSettings.CurrentValue.CurrentSite) ?? false;
     }
+
+    public Task AddRefToBlobs(NamespaceId ns, BucketId bucket, IoHashKey key, BlobIdentifier[] blobs)
+    {
+        using Scope _ = Tracer.Instance.StartActive("scylla.add_ref_blobs");
+
+        string nsAsString = ns.ToString();
+        ScyllaObjectReference r = new ScyllaObjectReference(bucket, key);
+        Task[] refUpdateTasks = new Task[blobs.Length];
+        for (int i = 0; i < blobs.Length; i++)
+        {
+            BlobIdentifier id = blobs[i];
+            refUpdateTasks[i] = _mapper.UpdateAsync<ScyllaBlobIndexTable>(
+                "SET references = references + ? WHERE namespace = ? AND blob_id = ?",
+                new[] { r },
+                nsAsString,
+                new ScyllaBlobIdentifier(id));
+        }
+
+        return Task.WhenAll(refUpdateTasks);
+    }
+
 }
 
 
@@ -81,13 +108,15 @@ class ScyllaBlobIndexTable
         Namespace = null!;
         BlobId = null!;
         Regions = null!;
+        References = null!;
     }
 
-    public ScyllaBlobIndexTable(string @namespace, BlobIdentifier blobId, HashSet<string> regions)
+    public ScyllaBlobIndexTable(string @namespace, BlobIdentifier blobId, HashSet<string> regions, List<ScyllaObjectReference> references)
     {
         Namespace = @namespace;
         BlobId =  new ScyllaBlobIdentifier(blobId);
         Regions = regions;
+        References = references;
     }
 
     [Cassandra.Mapping.Attributes.PartitionKey]
@@ -99,4 +128,7 @@ class ScyllaBlobIndexTable
 
     [Cassandra.Mapping.Attributes.Column("regions")]
     public HashSet<string> Regions { get; set; }
+
+    [Cassandra.Mapping.Attributes.Column("references")]
+    public List<ScyllaObjectReference> References { get; set; }
 }
diff --git a/Engine/Source/Programs/Horde/HordeStorage/Horde.Storage/Implementation/ObjectService.cs b/Engine/Source/Programs/Horde/HordeStorage/Horde.Storage/Implementation/ObjectService.cs
--- a/Engine/Source/Programs/Horde/HordeStorage/Horde.Storage/Implementation/ObjectService.cs
+++ b/Engine/Source/Programs/Horde/HordeStorage/Horde.Storage/Implementation/ObjectService.cs
@@ -7,6 +7,7 @@ using System.Linq;
 using System.Threading.Tasks;
 using async_enumerable_dotnet;
 using Datadog.Trace;
+using Horde.Storage.Implementation.Blob;
 using Jupiter.Implementation;
 using Jupiter.Utils;
 
@@ -31,13 +32,15 @@ namespace Horde.Storage.Implementation
         private readonly IBlobService _blobService;
         private readonly IReferenceResolver _referenceResolver;
         private readonly IReplicationLog _replicationLog;
+        private readonly IBlobIndex _blobIndex;
 
-        public ObjectService(IReferencesStore referencesStore, IBlobService blobService, IReferenceResolver referenceResolver, IReplicationLog replicationLog)
+        public ObjectService(IReferencesStore referencesStore, IBlobService blobService, IReferenceResolver referenceResolver, IReplicationLog replicationLog, IBlobIndex blobIndex)
         {
             _referencesStore = referencesStore;
             _blobService = blobService;
             _referenceResolver = referenceResolver;
             _replicationLog = replicationLog;
+            _blobIndex = blobIndex;
         }
 
         public async Task<(ObjectRecord, BlobContents)> Get(NamespaceId ns, BucketId bucket, IoHashKey key, string[]? fields = null)
@@ -66,7 +69,7 @@ namespace Horde.Storage.Implementation
             Task objectStorePut = _referencesStore.Put(ns, bucket, key, blobHash, payload.Data, isFinalized);
 
             Task<BlobIdentifier> blobStorePut = _blobService.PutObject(ns, payload.Data, blobHash);
-
+            Task addRefToBodyTask = _blobIndex.AddRefToBlobs(ns, bucket, key, new [] {blobHash});
             BlobIdentifier[] missingReferences = Array.Empty<BlobIdentifier>();
             if (hasReferences)
             {
@@ -74,7 +77,12 @@ namespace Horde.Storage.Implementation
                 try
                 {
                     IAsyncEnumerable<BlobIdentifier> references = _referenceResolver.ResolveReferences(ns, payload);
-                    missingReferences = await _blobService.FilterOutKnownBlobs(ns, references);
+                    // TODO: we no longer use this as a async enumerable, which will slow down the resolve reference a little bit
+                    BlobIdentifier[] referencesArray = await references.ToArrayAsync();
+                    Task addRefsTask = _blobIndex.AddRefToBlobs(ns, bucket, key, referencesArray);
+                    Task<BlobIdentifier[]> checkReferencesTask = _blobService.FilterOutKnownBlobs(ns, referencesArray);
+                    await Task.WhenAll(addRefsTask, checkReferencesTask);
+                    missingReferences = await checkReferencesTask;
                 }
                 catch (PartialReferenceResolveException e)
                 {
@@ -82,7 +90,7 @@ namespace Horde.Storage.Implementation
                 }
             }
 
-            await Task.WhenAll(objectStorePut, blobStorePut);
+            await Task.WhenAll(objectStorePut, blobStorePut, addRefToBodyTask);
 
             if (missingReferences.Length == 0)
             {
diff --git a/Engine/Source/Programs/Horde/HordeStorage/Horde.Storage/Implementation/Objects/ScyllaReferencesStore.cs b/Engine/Source/Programs/Horde/HordeStorage/Horde.Storage/Implementation/Objects/ScyllaReferencesStore.cs
--- a/Engine/Source/Programs/Horde/HordeStorage/Horde.Storage/Implementation/Objects/ScyllaReferencesStore.cs
+++ b/Engine/Source/Programs/Horde/HordeStorage/Horde.Storage/Implementation/Objects/ScyllaReferencesStore.cs
@@ -92,13 +92,10 @@ namespace Horde.Storage.Implementation
         public async Task Finalize(NamespaceId ns, BucketId bucket, IoHashKey name, BlobIdentifier blobHash)
         {
             using Scope _ = Tracer.Instance.StartActive("scylla.finalize");
-            ObjectRecord o = await Get(ns, bucket, name);
-            if (!o.BlobIdentifier.Equals(blobHash))
-                throw new ObjectHashMismatchException(ns, bucket, name, blobHash, o.BlobIdentifier);
 
-            AppliedInfo<ScyllaObject> info = await _mapper.UpdateIfAsync<ScyllaObject>("SET is_finalized=true WHERE namespace=? AND bucket=? AND name=?", ns.ToString(), bucket.ToString(), name.ToString());
+            AppliedInfo<ScyllaObject> info = await _mapper.UpdateIfAsync<ScyllaObject>("SET is_finalized=true WHERE namespace=? AND bucket=? AND name=? IF payload_hash=?", ns.ToString(), bucket.ToString(), name.ToString(), blobHash.ToString());
             if (!info.Applied)
-                throw new InvalidOperationException("Failed to finalize object even though it existed and hashes matched");
+                throw new ObjectHashMismatchException(ns, bucket, name, blobHash, info.Existing.PayloadHash!.AsBlobIdentifier());
         }
 
         private async Task CreateTTLRecord(NamespaceId ns, BucketId bucket, IoHashKey name, sbyte partitionIndex)
@@ -302,6 +299,30 @@ namespace Horde.Storage.Implementation
         }
     }
 
+    public class ScyllaObjectReference
+    {
+        // used by the cassandra mapper
+        public ScyllaObjectReference()
+        {
+            Bucket = null!;
+            Key = null!;
+        }
+
+        public ScyllaObjectReference(BucketId bucket, IoHashKey key)
+        {
+            Bucket = bucket.ToString();
+            Key = key.ToString();
+        }
+
+        public string Bucket { get;set; }
+        public string Key { get; set; }
+
+        public (BucketId, IoHashKey) AsTuple()
+        {
+            return (new BucketId(Bucket), new IoHashKey(Key));
+        }
+    }
+
     public class ScyllaSettings
     {
         public long InlineBlobMaxSize { get; set; } = 32 * 1024; // default to 32 kb blobs max
diff --git a/Engine/Source/Programs/Horde/HordeStorage/Horde.Storage/Tests/Functional/Blob/BlobIndexTests.cs b/Engine/Source/Programs/Horde/HordeStorage/Horde.Storage/Tests/Functional/Blob/BlobIndexTests.cs
--- a/Engine/Source/Programs/Horde/HordeStorage/Horde.Storage/Tests/Functional/Blob/BlobIndexTests.cs
+++ b/Engine/Source/Programs/Horde/HordeStorage/Horde.Storage/Tests/Functional/Blob/BlobIndexTests.cs
@@ -2,12 +2,16 @@
 
 using System;
 using System.Collections.Generic;
+using System.IO;
 using System.Net.Http;
 using System.Net.Http.Headers;
 using System.Net.Mime;
 using System.Text;
 using System.Threading.Tasks;
+using Cassandra;
+using Horde.Storage.Implementation;
 using Horde.Storage.Implementation.Blob;
+using Jupiter;
 using Jupiter.Implementation;
 using Microsoft.AspNetCore.Hosting;
 using Microsoft.AspNetCore.TestHost;
@@ -15,7 +19,7 @@ using Microsoft.Extensions.Configuration;
 using Microsoft.Extensions.DependencyInjection;
 using Microsoft.VisualStudio.TestTools.UnitTesting;
 using Serilog;
-using Serilog.Core;
+using Logger = Serilog.Core.Logger;
 
 namespace Horde.Storage.FunctionalTests.Storage
 {
@@ -25,22 +29,6 @@ namespace Horde.Storage.FunctionalTests.Storage
         protected HttpClient? _httpClient;
 
         protected readonly NamespaceId TestNamespaceName = new NamespaceId("testbucket");
-        protected const string SmallFileContents = "Small file contents";
-        protected const string AnotherFileContents = "Another file with contents";
-        protected const string DeletableFileContents = "Delete Me";
-        protected const string OldFileContents = "a old blob used for testing cutoff filtering";
-
-        protected readonly BlobIdentifier _smallFileHash =
-            BlobIdentifier.FromBlob(Encoding.ASCII.GetBytes(SmallFileContents));
-
-        protected readonly BlobIdentifier _anotherFileHash =
-            BlobIdentifier.FromBlob(Encoding.ASCII.GetBytes(AnotherFileContents));
-
-        protected readonly BlobIdentifier _deleteFileHash =
-            BlobIdentifier.FromBlob(Encoding.ASCII.GetBytes(DeletableFileContents));
-
-        protected readonly BlobIdentifier _oldBlobFileHash =
-            BlobIdentifier.FromBlob(Encoding.ASCII.GetBytes(OldFileContents));
 
         [TestInitialize]
         public async Task Setup()
@@ -72,13 +60,13 @@ namespace Horde.Storage.FunctionalTests.Storage
         protected abstract IEnumerable<KeyValuePair<string, string>> GetSettings();
 
         protected abstract Task Seed(IServiceProvider serverServices);
-        protected abstract Task Teardown();
+        protected abstract Task Teardown(IServiceProvider serverServices);
 
 
         [TestCleanup]
         public async Task MyTeardown()
         {
-            await Teardown();
+            await Teardown(_server!.Services);
         }
 
         
@@ -100,6 +88,37 @@ namespace Horde.Storage.FunctionalTests.Storage
             Assert.IsTrue(blobInfo.Regions.Contains("test"));
         }
 
+        [TestMethod]
+        public async Task UploadRef()
+        {
+            CompactBinaryWriter writer = new CompactBinaryWriter();
+            writer.BeginObject();
+            writer.AddString("thisIsAField", "stringField");
+            writer.EndObject();
+
+            byte[] objectData = writer.Save();
+            BlobIdentifier objectHash = BlobIdentifier.FromBlob(objectData);
+
+            HttpContent requestContent = new ByteArrayContent(objectData);
+            requestContent.Headers.ContentType = new MediaTypeHeaderValue(CustomMediaTypeNames.UnrealCompactBinary);
+            requestContent.Headers.Add(CommonHeaders.HashHeaderName, objectHash.ToString());
+            IoHashKey putKey = IoHashKey.FromName("newReferenceUploadObject");
+            HttpResponseMessage result = await _httpClient!.PutAsync(requestUri: $"api/v1/refs/{TestNamespaceName}/bucket/{putKey}.uecb", requestContent);
+            result.EnsureSuccessStatusCode();
+
+            IBlobIndex? index = _server!.Services.GetService<IBlobIndex>();
+            Assert.IsNotNull(index);
+            IBlobIndex.BlobInfo? blobInfo = await index.GetBlobInfo(TestNamespaceName, objectHash);
+
+            Assert.IsNotNull(blobInfo);
+            Assert.IsTrue(blobInfo.Regions.Contains("test"));
+            Assert.AreEqual(1, blobInfo.References.Count);
+
+            (BucketId bucket, IoHashKey key) = blobInfo.References[0];
+            Assert.AreEqual("bucket", bucket.ToString());
+            Assert.AreEqual(putKey, key);
+        }
+
         [TestMethod]
         public async Task DeleteBlob()
         {
@@ -145,7 +164,7 @@ namespace Horde.Storage.FunctionalTests.Storage
             return Task.CompletedTask;
         }
 
-        protected override Task Teardown()
+        protected override Task Teardown(IServiceProvider serverServices)
         {
             return Task.CompletedTask;
         }
@@ -164,9 +183,12 @@ namespace Horde.Storage.FunctionalTests.Storage
             return Task.CompletedTask;
         }
 
-        protected override Task Teardown()
+        protected override async Task Teardown(IServiceProvider provider)
         {
-            return Task.CompletedTask;
+            IScyllaSessionManager scyllaSessionManager = provider.GetService<IScyllaSessionManager>()!;
+
+            ISession replicatedKeyspace = scyllaSessionManager.GetSessionForReplicatedKeyspace();
+            await replicatedKeyspace.ExecuteAsync(new SimpleStatement("DROP TABLE IF EXISTS blob_index"));
         }
     }
 }
diff --git a/Engine/Source/Programs/Horde/HordeStorage/Horde.Storage/appsettings.Development.json b/Engine/Source/Programs/Horde/HordeStorage/Horde.Storage/appsettings.Development.json
--- a/Engine/Source/Programs/Horde/HordeStorage/Horde.Storage/appsettings.Development.json
+++ b/Engine/Source/Programs/Horde/HordeStorage/Horde.Storage/appsettings.Development.json
@@ -26,7 +26,7 @@
 
     "replicationLogWriterImplementation": "memory",
 
-    "contentIdStoreImplementation": "memory"
+    "blobIndexImplementation": "memory"
   },
 
   // settings s3 defaults to use a local minio instance, same as tests, not enabled by default but useful to do local testing
@@ -75,6 +75,11 @@
     "RequireTls12": false
   },
 
+  "Scylla": {
+    "LocalDatacenterName": "datacenter1",
+    "LocalKeyspaceSuffix": "local"
+  },
+
   "DynamoDb": {
     "ConnectionString": "localhost:8000",
     "UseGlobalTables": false
