Horde: Fix conversion of default logger events to LogEvent objects leaving key names inside properties.

diff --git a/Engine/Source/Programs/Horde/Horde.Agent.Tests/LogParserTests.cs b/Engine/Source/Programs/Horde/Horde.Agent.Tests/LogParserTests.cs
--- a/Engine/Source/Programs/Horde/Horde.Agent.Tests/LogParserTests.cs
+++ b/Engine/Source/Programs/Horde/Horde.Agent.Tests/LogParserTests.cs
@@ -5,6 +5,7 @@ using HordeAgent.Parser;
 using HordeAgent.Utility;
 using HordeCommon;
 using Microsoft.Extensions.Logging;
+using Microsoft.Extensions.Logging.Abstractions;
 using Microsoft.VisualStudio.TestPlatform.Common.Utilities;
 using Microsoft.VisualStudio.TestTools.UnitTesting;
 using NuGet.Frameworks;
@@ -52,6 +53,35 @@ namespace HordeAgentTests
 			}
 		}
 
+		class JsonLoggerImpl : JsonLogger
+		{
+			public List<LogEvent> Lines = new List<LogEvent>();
+
+			public JsonLoggerImpl(bool? Warnings, ILogger Inner) 
+				: base(Warnings, Inner)
+			{
+			}
+
+			protected override void WriteFormattedEvent(LogLevel Level, byte[] Line)
+			{
+				Lines.Add(LogEvent.Read(Line));
+			}
+		}
+
+		[TestMethod]
+		public void JsonLoggerTest()
+		{
+			JsonLoggerImpl Impl = new JsonLoggerImpl(null, NullLogger.Instance);
+			Impl.LogInformation("Hello {0}", "world");
+			Impl.LogInformation("Hello {Text}", "world");
+
+			Assert.AreEqual(Impl.Lines.Count, 2);
+			Assert.AreEqual(Impl.Lines[0].ToString(), "Hello world");
+			Assert.AreEqual(Impl.Lines[0].Properties!["0"].ToString()!, "world");
+			Assert.AreEqual(Impl.Lines[1].ToString(), "Hello world");
+			Assert.AreEqual(Impl.Lines[1].Properties!["Text"].ToString(), "world");
+		}
+
 		[TestMethod]
 		public void StructuredOutputMatcher()
 		{
diff --git a/Engine/Source/Programs/Shared/EpicGames.Core/LogEvent.cs b/Engine/Source/Programs/Shared/EpicGames.Core/LogEvent.cs
--- a/Engine/Source/Programs/Shared/EpicGames.Core/LogEvent.cs
+++ b/Engine/Source/Programs/Shared/EpicGames.Core/LogEvent.cs
@@ -114,6 +114,18 @@ namespace EpicGames.Core
 			this.Exception = Exception;
 		}
 
+		/// <summary>
+		/// Read a log event from a utf-8 encoded json byte array
+		/// </summary>
+		/// <param name="Data"></param>
+		/// <returns></returns>
+		public static LogEvent Read(ReadOnlySpan<byte> Data)
+		{
+			Utf8JsonReader Reader = new Utf8JsonReader(Data);
+			Reader.Read();
+			return Read(ref Reader);
+		}
+
 		/// <summary>
 		/// Read a log event from Json
 		/// </summary>
diff --git a/Engine/Source/Programs/Shared/EpicGames.Core/MessageTemplate.cs b/Engine/Source/Programs/Shared/EpicGames.Core/MessageTemplate.cs
--- a/Engine/Source/Programs/Shared/EpicGames.Core/MessageTemplate.cs
+++ b/Engine/Source/Programs/Shared/EpicGames.Core/MessageTemplate.cs
@@ -42,7 +42,7 @@ namespace EpicGames.Core
 						ReadOnlySpan<char> Name = Format.AsSpan(Offset, Length);
 						if (int.TryParse(Name, out int Number))
 						{
-							OutputProperties[Name.ToString()] = Properties.ElementAtOrDefault(Number);
+							OutputProperties[Name.ToString()] = Properties.ElementAtOrDefault(Number).Value;
 						}
 					}
 				}
