Horde: More descriptive error return paths. Errors now return a JSON log event with information about the problem, including a specific error message and properties.

- HordeControllerBase defines methods for generating standard 404, 503 responses with typed data.
- Unhandled exceptions include an exception field detailing the source of the error.
- Added a string typeconverter for AgentId and SubResourceId, so they can be bound directly to controller routes.
- Modified SessionId to be an alias for ObjectId<ISession>.

diff --git a/Engine/Source/Programs/Horde/Horde.Build/Acls/AclService.cs b/Engine/Source/Programs/Horde/Horde.Build/Acls/AclService.cs
--- a/Engine/Source/Programs/Horde/Horde.Build/Acls/AclService.cs
+++ b/Engine/Source/Programs/Horde/Horde.Build/Acls/AclService.cs
@@ -21,6 +21,7 @@ using System.Threading.Tasks;
 namespace HordeServer.Services
 {
 	using UserId = ObjectId<IUser>;
+	using SessionId = ObjectId<ISession>;
 
 	/// <summary>
 	/// Cache of global ACL
@@ -216,7 +217,7 @@ namespace HordeServer.Services
 		/// </summary>
 		/// <param name="SessionId">The session id</param>
 		/// <returns>New claim instance</returns>
-		public static AclClaim GetSessionClaim(ObjectId SessionId)
+		public static AclClaim GetSessionClaim(SessionId SessionId)
 		{
 			return new AclClaim(HordeClaimTypes.AgentSessionId, SessionId.ToString());
 		}
@@ -243,21 +244,21 @@ namespace HordeServer.Services
 
 	static class ClaimExtensions
 	{
-		public static bool HasSessionClaim(this ClaimsPrincipal User, ObjectId SessionId)
+		public static bool HasSessionClaim(this ClaimsPrincipal User, SessionId SessionId)
 		{
 			return User.HasClaim(HordeClaimTypes.AgentSessionId, SessionId.ToString());
 		}
 
-		public static ObjectId? GetSessionClaim(this ClaimsPrincipal User)
+		public static SessionId? GetSessionClaim(this ClaimsPrincipal User)
 		{
 			Claim? Claim = User.FindFirst(HordeClaimTypes.AgentSessionId);
-			if (Claim == null || !ObjectId.TryParse(Claim.Value, out ObjectId SessionId))
+			if (Claim == null || !SessionId.TryParse(Claim.Value, out SessionId SessionIdValue))
 			{
 				return null;
 			}
 			else
 			{
-				return SessionId;
+				return SessionIdValue;
 			}
 		}
 		
diff --git a/Engine/Source/Programs/Horde/Horde.Build/Agents/AgentId.cs b/Engine/Source/Programs/Horde/Horde.Build/Agents/AgentId.cs
--- a/Engine/Source/Programs/Horde/Horde.Build/Agents/AgentId.cs
+++ b/Engine/Source/Programs/Horde/Horde.Build/Agents/AgentId.cs
@@ -22,6 +22,7 @@ namespace HordeServer.Utilities
 	/// <summary>
 	/// Normalized hostname of an agent
 	/// </summary>
+	[TypeConverter(typeof(AgentIdTypeConverter))]
 	[RedisConverter(typeof(AgentIdRedisConverter))]
 	[BsonSerializer(typeof(AgentIdBsonSerializer))]
 	public struct AgentId : IEquatable<AgentId>
@@ -167,19 +168,41 @@ namespace HordeServer.Utilities
 		}
 	}
 
+
 	/// <summary>
-	/// Extension methods for AgentId types
+	/// Type converter from strings to PropertyFilter objects
 	/// </summary>
-	static class AgentIdExtensions
+	sealed class AgentIdTypeConverter : TypeConverter
 	{
-		/// <summary>
-		/// Create an AgentId from a string
-		/// </summary>
-		/// <param name="Text">String to parse</param>
-		/// <returns>New agent id</returns>
-		public static AgentId ToAgentId(this string Text)
+		/// <inheritdoc/>
+		public override bool CanConvertFrom(ITypeDescriptorContext? Context, Type SourceType)
 		{
-			return new AgentId(Text);
+			return SourceType == typeof(string);
+		}
+
+		/// <inheritdoc/>
+		public override object ConvertFrom(ITypeDescriptorContext? Context, CultureInfo? Culture, object Value)
+		{
+			return new AgentId((string)Value);
+		}
+
+		/// <inheritdoc/>
+		public override bool CanConvertTo(ITypeDescriptorContext? Context, Type? DestinationType)
+		{
+			return DestinationType == typeof(string);
+		}
+
+		/// <inheritdoc/>
+		public override object? ConvertTo(ITypeDescriptorContext? Context, CultureInfo? Culture, object? Value, Type DestinationType)
+		{
+			if (DestinationType == typeof(string))
+			{
+				return Value?.ToString();
+			}
+			else
+			{
+				return null;
+			}
 		}
 	}
 }
diff --git a/Engine/Source/Programs/Horde/Horde.Build/Agents/IAgent.cs b/Engine/Source/Programs/Horde/Horde.Build/Agents/IAgent.cs
--- a/Engine/Source/Programs/Horde/Horde.Build/Agents/IAgent.cs
+++ b/Engine/Source/Programs/Horde/Horde.Build/Agents/IAgent.cs
@@ -28,6 +28,7 @@ namespace HordeServer.Models
 	using LeaseId = ObjectId<ILease>;
 	using LogId = ObjectId<ILogFile>;
 	using PoolId = StringId<IPool>;
+	using SessionId = ObjectId<ISession>;
 	using StreamId = StringId<IStream>;
 	using AgentSoftwareVersion = StringId<IAgentSoftwareCollection>;
 	using AgentSoftwareChannelName = StringId<AgentSoftwareChannels>;
@@ -402,7 +403,7 @@ namespace HordeServer.Models
 		/// <summary>
 		/// The current session id, if it's online
 		/// </summary>
-		public ObjectId? SessionId { get; }
+		public SessionId? SessionId { get; }
 
 		/// <summary>
 		/// Time at which the current session expires. 
diff --git a/Engine/Source/Programs/Horde/Horde.Build/Agents/IAgentCollection.cs b/Engine/Source/Programs/Horde/Horde.Build/Agents/IAgentCollection.cs
--- a/Engine/Source/Programs/Horde/Horde.Build/Agents/IAgentCollection.cs
+++ b/Engine/Source/Programs/Horde/Horde.Build/Agents/IAgentCollection.cs
@@ -20,6 +20,7 @@ namespace HordeServer.Collections
 {
 	using PoolId = StringId<IPool>;
 	using AgentSoftwareChannelName = StringId<AgentSoftwareChannels>;
+	using SessionId = ObjectId<ISession>;
 
 	/// <summary>
 	/// Interface for a collection of agent documents
@@ -112,7 +113,7 @@ namespace HordeServer.Collections
 		/// <param name="DynamicPools">New list of dynamic pools for the agent</param>
 		/// <param name="Version">Current version of the agent software</param>
 		/// <returns>New agent state</returns>
-		Task<IAgent?> TryStartSessionAsync(IAgent Agent, ObjectId SessionId, DateTime SessionExpiresAt, AgentStatus Status, IReadOnlyList<string> Properties, IReadOnlyDictionary<string, int> Resources, IReadOnlyList<PoolId> DynamicPools, string? Version);
+		Task<IAgent?> TryStartSessionAsync(IAgent Agent, SessionId SessionId, DateTime SessionExpiresAt, AgentStatus Status, IReadOnlyList<string> Properties, IReadOnlyDictionary<string, int> Resources, IReadOnlyList<PoolId> DynamicPools, string? Version);
 
 		/// <summary>
 		/// Attempt to update the agent state
diff --git a/Engine/Source/Programs/Horde/Horde.Build/Agents/ILease.cs b/Engine/Source/Programs/Horde/Horde.Build/Agents/ILease.cs
--- a/Engine/Source/Programs/Horde/Horde.Build/Agents/ILease.cs
+++ b/Engine/Source/Programs/Horde/Horde.Build/Agents/ILease.cs
@@ -18,6 +18,7 @@ namespace HordeServer.Models
 	using LeaseId = ObjectId<ILease>;
 	using LogId = ObjectId<ILogFile>;
 	using PoolId = StringId<IPool>;
+	using SessionId = ObjectId<ISession>;
 	using StreamId = StringId<IStream>;
 
 	/// <summary>
@@ -43,7 +44,7 @@ namespace HordeServer.Models
 		/// <summary>
 		/// Unique id of the agent session
 		/// </summary>
-		public ObjectId SessionId { get; }
+		public SessionId SessionId { get; }
 
 		/// <summary>
 		/// The stream this lease belongs to
diff --git a/Engine/Source/Programs/Horde/Horde.Build/Agents/ILeaseCollection.cs b/Engine/Source/Programs/Horde/Horde.Build/Agents/ILeaseCollection.cs
--- a/Engine/Source/Programs/Horde/Horde.Build/Agents/ILeaseCollection.cs
+++ b/Engine/Source/Programs/Horde/Horde.Build/Agents/ILeaseCollection.cs
@@ -17,6 +17,7 @@ namespace HordeServer.Collections
 	using LeaseId = ObjectId<ILease>;
 	using LogId = ObjectId<ILogFile>;
 	using PoolId = StringId<IPool>;
+	using SessionId = ObjectId<ISession>;
 	using StreamId = StringId<IStream>;
 
 	/// <summary>
@@ -37,7 +38,7 @@ namespace HordeServer.Collections
 		/// <param name="StartTime">Start time of the lease</param>
 		/// <param name="Payload">Payload for the lease</param>
 		/// <returns>Async task</returns>
-		Task<ILease> AddAsync(LeaseId Id, string Name, AgentId AgentId, ObjectId SessionId, StreamId? StreamId, PoolId? PoolId, LogId? LogId, DateTime StartTime, byte[] Payload);
+		Task<ILease> AddAsync(LeaseId Id, string Name, AgentId AgentId, SessionId SessionId, StreamId? StreamId, PoolId? PoolId, LogId? LogId, DateTime StartTime, byte[] Payload);
 
 		/// <summary>
 		/// Deletes a lease from the collection
@@ -63,7 +64,7 @@ namespace HordeServer.Collections
 		/// <param name="Index">Index of the first result to return</param>
 		/// <param name="Count">Number of results to return</param>
 		/// <returns>List of leases matching the given criteria</returns>
-		Task<List<ILease>> FindLeasesAsync(AgentId? AgentId = null, ObjectId? SessionId = null, DateTime? MinTime = null, DateTime? MaxTime = null, int? Index = null, int? Count = null);
+		Task<List<ILease>> FindLeasesAsync(AgentId? AgentId = null, SessionId? SessionId = null, DateTime? MinTime = null, DateTime? MaxTime = null, int? Index = null, int? Count = null);
 
 		/// <summary>
 		/// Finds all active leases
diff --git a/Engine/Source/Programs/Horde/Horde.Build/Agents/ISession.cs b/Engine/Source/Programs/Horde/Horde.Build/Agents/ISession.cs
--- a/Engine/Source/Programs/Horde/Horde.Build/Agents/ISession.cs
+++ b/Engine/Source/Programs/Horde/Horde.Build/Agents/ISession.cs
@@ -11,6 +11,8 @@ using System.Threading.Tasks;
 
 namespace HordeServer.Models
 {
+	using SessionId = ObjectId<ISession>;
+
 	/// <summary>
 	/// Information about an agent session.
 	/// </summary>
@@ -19,7 +21,7 @@ namespace HordeServer.Models
 		/// <summary>
 		/// Unique id for this session
 		/// </summary>
-		public ObjectId Id { get; }
+		public SessionId Id { get; }
 
 		/// <summary>
 		/// The agent id
diff --git a/Engine/Source/Programs/Horde/Horde.Build/Agents/ISessionCollection.cs b/Engine/Source/Programs/Horde/Horde.Build/Agents/ISessionCollection.cs
--- a/Engine/Source/Programs/Horde/Horde.Build/Agents/ISessionCollection.cs
+++ b/Engine/Source/Programs/Horde/Horde.Build/Agents/ISessionCollection.cs
@@ -14,6 +14,8 @@ using System.Threading.Tasks;
 
 namespace HordeServer.Collections
 {
+	using SessionId = ObjectId<ISession>;
+
 	/// <summary>
 	/// Interface for a collection of session documents
 	/// </summary>
@@ -28,14 +30,14 @@ namespace HordeServer.Collections
 		/// <param name="Properties">Properties of this agent at the time the session started</param>
 		/// <param name="Resources">Resources which the agent has</param>
 		/// <param name="Version">Version of the agent software</param>
-		Task<ISession> AddAsync(ObjectId Id, AgentId AgentId, DateTime StartTime, IReadOnlyList<string>? Properties, IReadOnlyDictionary<string, int>? Resources, string? Version);
+		Task<ISession> AddAsync(SessionId Id, AgentId AgentId, DateTime StartTime, IReadOnlyList<string>? Properties, IReadOnlyDictionary<string, int>? Resources, string? Version);
 
 		/// <summary>
 		/// Gets information about a particular session
 		/// </summary>
 		/// <param name="SessionId">The unique session id</param>
 		/// <returns>The session information</returns>
-		Task<ISession?> GetAsync(ObjectId SessionId);
+		Task<ISession?> GetAsync(SessionId SessionId);
 
 		/// <summary>
 		/// Find sessions for the given agent
@@ -64,13 +66,13 @@ namespace HordeServer.Collections
 		/// <param name="Properties">The agent properties</param>
 		/// <param name="Resources">Resources which the agent has</param>
 		/// <returns>Async task</returns>
-		Task UpdateAsync(ObjectId SessionId, DateTime FinishTime, IReadOnlyList<string> Properties, IReadOnlyDictionary<string, int> Resources);
+		Task UpdateAsync(SessionId SessionId, DateTime FinishTime, IReadOnlyList<string> Properties, IReadOnlyDictionary<string, int> Resources);
 
 		/// <summary>
 		/// Delete a session from the collection
 		/// </summary>
 		/// <param name="SessionId">The session id</param>
 		/// <returns>Async task</returns>
-		Task DeleteAsync(ObjectId SessionId);
+		Task DeleteAsync(SessionId SessionId);
 	}
 }
diff --git a/Engine/Source/Programs/Horde/Horde.Build/Agents/Impl/AgentCollection.cs b/Engine/Source/Programs/Horde/Horde.Build/Agents/Impl/AgentCollection.cs
--- a/Engine/Source/Programs/Horde/Horde.Build/Agents/Impl/AgentCollection.cs
+++ b/Engine/Source/Programs/Horde/Horde.Build/Agents/Impl/AgentCollection.cs
@@ -21,6 +21,7 @@ namespace HordeServer.Collections.Impl
 {
 	using AgentSoftwareChannelName = StringId<AgentSoftwareChannels>;
 	using PoolId = StringId<IPool>;
+	using SessionId = ObjectId<ISession>;
 
 	/// <summary>
 	/// Collection of agent documents
@@ -60,7 +61,7 @@ namespace HordeServer.Collections.Impl
 			[BsonRequired, BsonId]
 			public AgentId Id { get; set; }
 
-			public ObjectId? SessionId { get; set; }
+			public SessionId? SessionId { get; set; }
 			public DateTime? SessionExpiresAt { get; set; }
 
 			public AgentStatus Status { get; set; }
@@ -426,7 +427,7 @@ namespace HordeServer.Collections.Impl
 		}
 
 		/// <inheritdoc/>
-		public async Task<IAgent?> TryStartSessionAsync(IAgent AgentInterface, ObjectId SessionId, DateTime SessionExpiresAt, AgentStatus Status, IReadOnlyList<string> Properties, IReadOnlyDictionary<string, int> Resources, IReadOnlyList<PoolId> DynamicPools, string? Version)
+		public async Task<IAgent?> TryStartSessionAsync(IAgent AgentInterface, SessionId SessionId, DateTime SessionExpiresAt, AgentStatus Status, IReadOnlyList<string> Properties, IReadOnlyDictionary<string, int> Resources, IReadOnlyList<PoolId> DynamicPools, string? Version)
 		{
 			AgentDocument Agent = (AgentDocument)AgentInterface;
 			List<string> NewProperties = Properties.OrderBy(x => x, StringComparer.OrdinalIgnoreCase).ToList();
diff --git a/Engine/Source/Programs/Horde/Horde.Build/Agents/Impl/AgentService.cs b/Engine/Source/Programs/Horde/Horde.Build/Agents/Impl/AgentService.cs
--- a/Engine/Source/Programs/Horde/Horde.Build/Agents/Impl/AgentService.cs
+++ b/Engine/Source/Programs/Horde/Horde.Build/Agents/Impl/AgentService.cs
@@ -26,6 +26,7 @@ namespace HordeServer.Services
 	using AgentSoftwareChannelName = StringId<AgentSoftwareChannels>;
 	using LeaseId = ObjectId<ILease>;
 	using PoolId = StringId<IPool>;
+	using SessionId = ObjectId<ISession>;
 
 	/// <summary>
 	/// Wraps funtionality for manipulating agents
@@ -145,7 +146,7 @@ namespace HordeServer.Services
 		/// </summary>
 		/// <param name="SessionId">The session id</param>
 		/// <returns>Bearer token for the agent</returns>
-		public string IssueSessionToken(ObjectId SessionId)
+		public string IssueSessionToken(SessionId SessionId)
 		{
 			List<AclClaim> Claims = new List<AclClaim>();
 			Claims.Add(AclService.AgentClaim);
@@ -285,7 +286,7 @@ namespace HordeServer.Services
 					}
 
 					// Create a new session document
-					ISession NewSession = await Sessions.AddAsync(ObjectId.GenerateNewId(), Agent.Id, Clock.UtcNow, Properties, Resources, Version);
+					ISession NewSession = await Sessions.AddAsync(SessionId.GenerateNewId(), Agent.Id, Clock.UtcNow, Properties, Resources, Version);
 					DateTime SessionExpiresAt = UtcNow + SessionExpiryTime;
 
 					// Get the new dynamic pools for the agent
@@ -490,7 +491,7 @@ namespace HordeServer.Services
 		/// <summary>
 		/// 
 		/// </summary>
-		public async Task<IAgent?> UpdateSessionWithWaitAsync(IAgent InAgent, ObjectId SessionId, AgentStatus Status, IReadOnlyList<string>? Properties, IReadOnlyDictionary<string, int>? Resources, IList<HordeCommon.Rpc.Messages.Lease> NewLeases, CancellationToken CancellationToken)
+		public async Task<IAgent?> UpdateSessionWithWaitAsync(IAgent InAgent, SessionId SessionId, AgentStatus Status, IReadOnlyList<string>? Properties, IReadOnlyDictionary<string, int>? Resources, IList<HordeCommon.Rpc.Messages.Lease> NewLeases, CancellationToken CancellationToken)
 		{
 			IAgent? Agent = InAgent;
 
@@ -516,7 +517,7 @@ namespace HordeServer.Services
 		/// <param name="Resources">New agent resources</param>
 		/// <param name="NewLeases">New list of leases for this session</param>
 		/// <returns>Updated agent state</returns>
-		public async Task<IAgent?> UpdateSessionAsync(IAgent InAgent, ObjectId SessionId, AgentStatus Status, IReadOnlyList<string>? Properties, IReadOnlyDictionary<string, int>? Resources, IList<HordeCommon.Rpc.Messages.Lease> NewLeases)
+		public async Task<IAgent?> UpdateSessionAsync(IAgent InAgent, SessionId SessionId, AgentStatus Status, IReadOnlyList<string>? Properties, IReadOnlyDictionary<string, int>? Resources, IList<HordeCommon.Rpc.Messages.Lease> NewLeases)
 		{
 			DateTime UtcNow = Clock.UtcNow;
 
@@ -656,7 +657,7 @@ namespace HordeServer.Services
 			}
 
 			// Save off the session id and current leases
-			ObjectId SessionId = Agent.SessionId.Value;
+			SessionId SessionId = Agent.SessionId.Value;
 			List<AgentLease> Leases = new List<AgentLease>(Agent.Leases);
 
 			// Clear the current session
@@ -709,7 +710,7 @@ namespace HordeServer.Services
 		/// <param name="Index">Index of the first result to return</param>
 		/// <param name="Count">Number of results to return</param>
 		/// <returns>List of leases matching the given criteria</returns>
-		public Task<List<ILease>> FindLeasesAsync(AgentId? AgentId, ObjectId? SessionId, DateTime? StartTime, DateTime? FinishTime, int Index, int Count)
+		public Task<List<ILease>> FindLeasesAsync(AgentId? AgentId, SessionId? SessionId, DateTime? StartTime, DateTime? FinishTime, int Index, int Count)
 		{
 			return Leases.FindLeasesAsync(AgentId, SessionId, StartTime, FinishTime, Index, Count);
 		}
@@ -775,7 +776,7 @@ namespace HordeServer.Services
 		/// </summary>
 		/// <param name="SessionId">The unique session id</param>
 		/// <returns>The session information</returns>
-		public Task<ISession?> GetSessionAsync(ObjectId SessionId)
+		public Task<ISession?> GetSessionAsync(SessionId SessionId)
 		{
 			return Sessions.GetAsync(SessionId);
 		}
diff --git a/Engine/Source/Programs/Horde/Horde.Build/Agents/Impl/LeaseCollection.cs b/Engine/Source/Programs/Horde/Horde.Build/Agents/Impl/LeaseCollection.cs
--- a/Engine/Source/Programs/Horde/Horde.Build/Agents/Impl/LeaseCollection.cs
+++ b/Engine/Source/Programs/Horde/Horde.Build/Agents/Impl/LeaseCollection.cs
@@ -17,6 +17,7 @@ namespace HordeServer.Collections.Impl
 	using LeaseId = ObjectId<ILease>;
 	using LogId = ObjectId<ILogFile>;
 	using PoolId = StringId<IPool>;
+	using SessionId = ObjectId<ISession>;
 	using StreamId = StringId<IStream>;
 
 	/// <summary>
@@ -32,7 +33,7 @@ namespace HordeServer.Collections.Impl
 			public LeaseId Id { get; set; }
 			public string Name { get; set; }
 			public AgentId AgentId { get; set; }
-			public ObjectId SessionId { get; set; }
+			public SessionId SessionId { get; set; }
 			public StreamId? StreamId { get; set; }
 			public PoolId? PoolId { get; set; }
 			public LogId? LogId { get; set; }
@@ -54,7 +55,7 @@ namespace HordeServer.Collections.Impl
 				Payload = null!;
 			}
 
-			public LeaseDocument(LeaseId Id, string Name, AgentId AgentId, ObjectId SessionId, StreamId? StreamId, PoolId? PoolId, LogId? LogId, DateTime StartTime, byte[] Payload)
+			public LeaseDocument(LeaseId Id, string Name, AgentId AgentId, SessionId SessionId, StreamId? StreamId, PoolId? PoolId, LogId? LogId, DateTime StartTime, byte[] Payload)
 			{
 				this.Id = Id;
 				this.Name = Name;
@@ -91,7 +92,7 @@ namespace HordeServer.Collections.Impl
 		}
 
 		/// <inheritdoc/>
-		public async Task<ILease> AddAsync(LeaseId Id, string Name, AgentId AgentId, ObjectId SessionId, StreamId? StreamId, PoolId? PoolId, LogId? LogId, DateTime StartTime, byte[] Payload)
+		public async Task<ILease> AddAsync(LeaseId Id, string Name, AgentId AgentId, SessionId SessionId, StreamId? StreamId, PoolId? PoolId, LogId? LogId, DateTime StartTime, byte[] Payload)
 		{
 			LeaseDocument Lease = new LeaseDocument(Id, Name, AgentId, SessionId, StreamId, PoolId, LogId, StartTime, Payload);
 			await Leases.ReplaceOneAsync(x => x.Id == Id, Lease, new ReplaceOptions { IsUpsert = true });
@@ -111,7 +112,7 @@ namespace HordeServer.Collections.Impl
 		}
 
 		/// <inheritdoc/>
-		public async Task<List<ILease>> FindLeasesAsync(AgentId? AgentId, ObjectId? SessionId, DateTime? MinTime, DateTime? MaxTime, int? Index, int? Count)
+		public async Task<List<ILease>> FindLeasesAsync(AgentId? AgentId, SessionId? SessionId, DateTime? MinTime, DateTime? MaxTime, int? Index, int? Count)
 		{
 			FilterDefinitionBuilder<LeaseDocument> FilterBuilder = Builders<LeaseDocument>.Filter;
 
diff --git a/Engine/Source/Programs/Horde/Horde.Build/Agents/Impl/SessionCollection.cs b/Engine/Source/Programs/Horde/Horde.Build/Agents/Impl/SessionCollection.cs
--- a/Engine/Source/Programs/Horde/Horde.Build/Agents/Impl/SessionCollection.cs
+++ b/Engine/Source/Programs/Horde/Horde.Build/Agents/Impl/SessionCollection.cs
@@ -14,6 +14,8 @@ using System.Threading.Tasks;
 
 namespace HordeServer.Collections.Impl
 {
+	using SessionId = ObjectId<ISession>;
+
 	/// <summary>
 	/// Collection of session documents
 	/// </summary>
@@ -25,7 +27,7 @@ namespace HordeServer.Collections.Impl
 		class SessionDocument : ISession
 		{
 			[BsonRequired, BsonId]
-			public ObjectId Id { get; set; }
+			public SessionId Id { get; set; }
 
 			[BsonRequired]
 			public AgentId AgentId { get; set; }
@@ -43,7 +45,7 @@ namespace HordeServer.Collections.Impl
 			{
 			}
 
-			public SessionDocument(ObjectId Id, AgentId AgentId, DateTime StartTime, IReadOnlyList<string>? Properties, IReadOnlyDictionary<string, int>? Resources, string? Version)
+			public SessionDocument(SessionId Id, AgentId AgentId, DateTime StartTime, IReadOnlyList<string>? Properties, IReadOnlyDictionary<string, int>? Resources, string? Version)
 			{
 				this.Id = Id;
 				this.AgentId = AgentId;
@@ -82,7 +84,7 @@ namespace HordeServer.Collections.Impl
 		}
 
 		/// <inheritdoc/>
-		public async Task<ISession> AddAsync(ObjectId Id, AgentId AgentId, DateTime StartTime, IReadOnlyList<string>? Properties, IReadOnlyDictionary<string, int>? Resources, string? Version)
+		public async Task<ISession> AddAsync(SessionId Id, AgentId AgentId, DateTime StartTime, IReadOnlyList<string>? Properties, IReadOnlyDictionary<string, int>? Resources, string? Version)
 		{
 			SessionDocument NewSession = new SessionDocument(Id, AgentId, StartTime, Properties, Resources, Version);
 			await Sessions.InsertOneAsync(NewSession);
@@ -90,7 +92,7 @@ namespace HordeServer.Collections.Impl
 		}
 
 		/// <inheritdoc/>
-		public async Task<ISession?> GetAsync(ObjectId SessionId)
+		public async Task<ISession?> GetAsync(SessionId SessionId)
 		{
 			return await Sessions.Find(x => x.Id == SessionId).FirstOrDefaultAsync();
 		}
@@ -122,7 +124,7 @@ namespace HordeServer.Collections.Impl
 		}
 
 		/// <inheritdoc/>
-		public Task UpdateAsync(ObjectId SessionId, DateTime FinishTime, IReadOnlyList<string> Properties, IReadOnlyDictionary<string, int> Resources)
+		public Task UpdateAsync(SessionId SessionId, DateTime FinishTime, IReadOnlyList<string> Properties, IReadOnlyDictionary<string, int> Resources)
 		{
 			UpdateDefinition<SessionDocument> Update = Builders<SessionDocument>.Update
 				.Set(x => x.FinishTime, FinishTime)
@@ -133,7 +135,7 @@ namespace HordeServer.Collections.Impl
 		}
 
 		/// <inheritdoc/>
-		public Task DeleteAsync(ObjectId SessionId)
+		public Task DeleteAsync(SessionId SessionId)
 		{
 			return Sessions.DeleteOneAsync(x => x.Id == SessionId);
 		}
diff --git a/Engine/Source/Programs/Horde/Horde.Build/Controllers/AdminController.cs b/Engine/Source/Programs/Horde/Horde.Build/Controllers/AdminController.cs
--- a/Engine/Source/Programs/Horde/Horde.Build/Controllers/AdminController.cs
+++ b/Engine/Source/Programs/Horde/Horde.Build/Controllers/AdminController.cs
@@ -61,7 +61,7 @@ namespace HordeServer.Controllers
 	[ApiController]
 	[Authorize]
 	[Route("[controller]")]
-	public class AdminController : ControllerBase
+	public class AdminController : HordeControllerBase
 	{
 		/// <summary>
 		/// The database service singleton
@@ -98,39 +98,6 @@ namespace HordeServer.Controllers
 			this.Settings = Settings;
 		}
 
-		/// <summary>
-		/// Force a reset on the database
-		/// </summary>
-		[HttpPost]
-		[Route("/api/v1/admin/reset")]
-		public async Task<ActionResult> ForceResetAsync([FromQuery] string? Instance = null)
-		{
-			if(!await AclService.AuthorizeAsync(AclAction.AdminWrite, User))
-			{
-				return Forbid();
-			}
-
-			for (; ; )
-			{
-				Globals Globals = await DatabaseService.GetGlobalsAsync();
-				if (Instance == null)
-				{
-					return NotFound($"Missing code query parameter. Set to {Globals.InstanceId} to reset.");
-				}
-				if (Globals.InstanceId != Instance.ToObjectId())
-				{
-					return NotFound($"Incorrect code query parameter. Should be {Globals.InstanceId}.");
-				}
-
-				Globals.ForceReset = true;
-
-				if (await DatabaseService.TryUpdateSingletonAsync(Globals))
-				{
-					return Ok("Database will be reinitialized on next restart");
-				}
-			}
-		}
-
 		/// <summary>
 		/// Upgrade the database to the latest schema
 		/// </summary>
@@ -141,7 +108,7 @@ namespace HordeServer.Controllers
 		{
 			if (!await AclService.AuthorizeAsync(AclAction.AdminWrite, User))
 			{
-				return Forbid();
+				return Forbid(AclAction.AdminWrite);
 			}
 
 			await UpgradeService.UpgradeSchemaAsync(FromVersion);
@@ -158,7 +125,7 @@ namespace HordeServer.Controllers
 		{
 			if (!await AclService.AuthorizeAsync(AclAction.IssueBearerToken, User))
 			{
-				return Forbid();
+				return Forbid(AclAction.IssueBearerToken);
 			}
 
 			return AclService.IssueBearerToken(User.Claims, GetDefaultExpiryTime());
@@ -175,7 +142,7 @@ namespace HordeServer.Controllers
 		{
 			if (!await AclService.AuthorizeAsync(AclAction.AdminWrite, User))
 			{
-				return Forbid();
+				return Forbid(AclAction.AdminWrite);
 			}
 
 			List<Claim> Claims = new List<Claim>();
@@ -194,7 +161,7 @@ namespace HordeServer.Controllers
 		{
 			if (!await AclService.AuthorizeAsync(AclAction.AdminWrite, User))
 			{
-				return Forbid();
+				return Forbid(AclAction.AdminWrite);
 			}
 
 			List<AclClaim> Claims = new List<AclClaim>();
@@ -214,7 +181,7 @@ namespace HordeServer.Controllers
 		{
 			if (!await AclService.AuthorizeAsync(AclAction.AdminWrite, User))
 			{
-				return Forbid();
+				return Forbid(AclAction.AdminWrite);
 			}
 
 			List<AclClaim> Claims = new List<AclClaim>();
@@ -234,7 +201,7 @@ namespace HordeServer.Controllers
 		{
 			if (!await AclService.AuthorizeAsync(AclAction.AdminRead, User))
 			{
-				return Forbid();
+				return Forbid(AclAction.AdminRead);
 			}
 
 			List<AclClaim> Claims = new List<AclClaim>();
@@ -254,7 +221,7 @@ namespace HordeServer.Controllers
 		{
 			if (!await AclService.AuthorizeAsync(AclAction.AdminRead, User))
 			{
-				return Forbid();
+				return Forbid(AclAction.AdminRead);
 			}
 
 			List<AclClaim> Claims = new List<AclClaim>();
@@ -274,7 +241,7 @@ namespace HordeServer.Controllers
 		{
 			if (!await AclService.AuthorizeAsync(AclAction.AdminRead, User))
 			{
-				return Forbid();
+				return Forbid(AclAction.AdminRead);
 			}
 
 			List<AclClaim> Claims = new List<AclClaim>();
diff --git a/Engine/Source/Programs/Horde/Horde.Build/Controllers/AgentsController.cs b/Engine/Source/Programs/Horde/Horde.Build/Controllers/AgentsController.cs
--- a/Engine/Source/Programs/Horde/Horde.Build/Controllers/AgentsController.cs
+++ b/Engine/Source/Programs/Horde/Horde.Build/Controllers/AgentsController.cs
@@ -20,6 +20,7 @@ namespace HordeServer.Controllers
 	using AgentSoftwareChannelName = StringId<AgentSoftwareChannels>;
 	using LeaseId = ObjectId<ILease>;
 	using PoolId = StringId<IPool>;
+	using SessionId = ObjectId<ISession>;
 
 	/// <summary>
 	/// Controller for the /api/v1/agents endpoint
@@ -27,7 +28,7 @@ namespace HordeServer.Controllers
 	[ApiController]
 	[Authorize]
 	[Route("[controller]")]
-	public class AgentsController : ControllerBase
+	public class AgentsController : HordeControllerBase
 	{
 		/// <summary>
 		/// Singleton instance of the ACL service
@@ -68,7 +69,7 @@ namespace HordeServer.Controllers
 		{
 			if(!await AclService.AuthorizeAsync(AclAction.CreateAgent, User))
 			{
-				return Forbid();
+				return Forbid(AclAction.CreateAgent);
 			}
 
 			AgentSoftwareChannelName? Channel = String.IsNullOrEmpty(Request.Channel) ? (AgentSoftwareChannelName?)null : new AgentSoftwareChannelName(Request.Channel);
@@ -113,12 +114,12 @@ namespace HordeServer.Controllers
 		[HttpGet]
 		[Route("/api/v1/agents/{AgentId}")]
 		[ProducesResponseType(typeof(GetAgentResponse), 200)]
-		public async Task<ActionResult<object>> GetAgentAsync(string AgentId, [FromQuery] PropertyFilter? Filter = null)
+		public async Task<ActionResult<object>> GetAgentAsync(AgentId AgentId, [FromQuery] PropertyFilter? Filter = null)
 		{
-			IAgent? Agent = await AgentService.GetAgentAsync(new AgentId(AgentId));
+			IAgent? Agent = await AgentService.GetAgentAsync(AgentId);
 			if (Agent == null)
 			{
-				return NotFound();
+				return NotFound(AgentId);
 			}
 
 			GlobalPermissionsCache Cache = new GlobalPermissionsCache();
@@ -134,7 +135,7 @@ namespace HordeServer.Controllers
 		/// <returns>Http result code</returns>
 		[HttpPut]
 		[Route("/api/v1/agents/{AgentId}")]
-		public async Task<ActionResult> UpdateAgentAsync(string AgentId, [FromBody] UpdateAgentRequest Update)
+		public async Task<ActionResult> UpdateAgentAsync(AgentId AgentId, [FromBody] UpdateAgentRequest Update)
 		{
 			List<PoolId>? UpdatePools = Update.Pools?.ConvertAll(x => new PoolId(x));
 			AgentSoftwareChannelName? Channel = String.IsNullOrEmpty(Update.Channel) ? (AgentSoftwareChannelName?)null : new AgentSoftwareChannelName(Update.Channel);
@@ -144,18 +145,18 @@ namespace HordeServer.Controllers
 			GlobalPermissionsCache Cache = new GlobalPermissionsCache();
 			for (; ; )
 			{
-				IAgent? Agent = await AgentService.GetAgentAsync(new AgentId(AgentId));
+				IAgent? Agent = await AgentService.GetAgentAsync(AgentId);
 				if (Agent == null)
 				{
-					return NotFound();
+					return NotFound(AgentId);
 				}
 				if (!await AgentService.AuthorizeAsync(Agent, AclAction.UpdateAgent, User, Cache))
 				{
-					return Forbid();
+					return Forbid(AclAction.UpdateAgent, AgentId);
 				}
 				if (Update.Acl != null && !await AgentService.AuthorizeAsync(Agent, AclAction.ChangePermissions, User, Cache))
 				{
-					return Forbid();
+					return Forbid(AclAction.ChangePermissions, AgentId);
 				}
 
 				IAgent? NewAgent = await AgentService.Agents.TryUpdateSettingsAsync(Agent, Update.Enabled, Update.RequestConform, Update.RequestRestart, Update.RequestShutdown, Channel, Update.Pools?.ConvertAll(x => new PoolId(x)), Acl.Merge(Agent.Acl, Update.Acl), Update.Comment);
@@ -205,16 +206,16 @@ namespace HordeServer.Controllers
 		/// <returns>Http result code</returns>
 		[HttpDelete]
 		[Route("/api/v1/agents/{AgentId}")]
-		public async Task<ActionResult> DeleteAgentAsync(string AgentId)
+		public async Task<ActionResult> DeleteAgentAsync(AgentId AgentId)
 		{
-			IAgent? Agent = await AgentService.GetAgentAsync(new AgentId(AgentId));
+			IAgent? Agent = await AgentService.GetAgentAsync(AgentId);
 			if (Agent == null)
 			{
-				return NotFound();
+				return NotFound(AgentId);
 			}
 			if (!await AgentService.AuthorizeAsync(Agent, AclAction.DeleteAgent, User, null))
 			{
-				return Forbid();
+				return Forbid(AclAction.DeleteAgent, AgentId);
 			}
 
 			await AgentService.DeleteAgentAsync(Agent);
@@ -232,12 +233,12 @@ namespace HordeServer.Controllers
 		/// <returns>Information about the requested agent</returns>
 		[HttpGet]
 		[Route("/api/v1/agents/{AgentId}/history")]
-		public async Task GetAgentHistoryAsync(string AgentId, [FromQuery] DateTime? MinTime = null, [FromQuery] DateTime? MaxTime = null, [FromQuery] int Index = 0, [FromQuery] int Count = 50)
+		public async Task GetAgentHistoryAsync(AgentId AgentId, [FromQuery] DateTime? MinTime = null, [FromQuery] DateTime? MaxTime = null, [FromQuery] int Index = 0, [FromQuery] int Count = 50)
 		{
 			Response.ContentType = "application/json";
 			Response.StatusCode = 200;
 			await Response.StartAsync();
-			await AgentService.Agents.GetLogger(AgentId.ToAgentId()).FindAsync(Response.BodyWriter, MinTime, MaxTime, Index, Count);
+			await AgentService.Agents.GetLogger(AgentId).FindAsync(Response.BodyWriter, MinTime, MaxTime, Index, Count);
 		}
 
 		/// <summary>
@@ -251,21 +252,19 @@ namespace HordeServer.Controllers
 		/// <returns>Sessions </returns>
 		[HttpGet]
 		[Route("/api/v1/agents/{AgentId}/sessions")]
-		public async Task<ActionResult<List<GetAgentSessionResponse>>> FindSessionsAsync(string AgentId, [FromQuery] DateTimeOffset? StartTime, [FromQuery] DateTimeOffset? FinishTime, [FromQuery] int Index = 0, [FromQuery] int Count = 50)
+		public async Task<ActionResult<List<GetAgentSessionResponse>>> FindSessionsAsync(AgentId AgentId, [FromQuery] DateTimeOffset? StartTime, [FromQuery] DateTimeOffset? FinishTime, [FromQuery] int Index = 0, [FromQuery] int Count = 50)
 		{
-			AgentId AgentIdValue = new AgentId(AgentId);
-
-			IAgent? Agent = await AgentService.GetAgentAsync(AgentIdValue);
+			IAgent? Agent = await AgentService.GetAgentAsync(AgentId);
 			if (Agent == null)
 			{
-				return NotFound();
+				return NotFound(AgentId);
 			}
 			if (!await AgentService.AuthorizeAsync(Agent, AclAction.ViewSession, User, null))
 			{
-				return Forbid();
+				return Forbid(AclAction.ViewSession, AgentId);
 			}
 
-			List<ISession> Sessions = await AgentService.FindSessionsAsync(AgentIdValue, StartTime?.UtcDateTime, FinishTime?.UtcDateTime, Index, Count);
+			List<ISession> Sessions = await AgentService.FindSessionsAsync(AgentId, StartTime?.UtcDateTime, FinishTime?.UtcDateTime, Index, Count);
 			return Sessions.ConvertAll(x => new GetAgentSessionResponse(x));
 		}
 
@@ -277,22 +276,20 @@ namespace HordeServer.Controllers
 		/// <returns>Sessions </returns>
 		[HttpGet]
 		[Route("/api/v1/agents/{AgentId}/sessions/{SessionId}")]
-		public async Task<ActionResult<GetAgentSessionResponse>> GetSessionAsync(string AgentId, string SessionId)
+		public async Task<ActionResult<GetAgentSessionResponse>> GetSessionAsync(AgentId AgentId, SessionId SessionId)
 		{
-			AgentId AgentIdValue = new AgentId(AgentId);
-
-			IAgent? Agent = await AgentService.GetAgentAsync(AgentIdValue);
+			IAgent? Agent = await AgentService.GetAgentAsync(AgentId);
 			if (Agent == null)
 			{
-				return NotFound();
+				return NotFound(AgentId);
 			}
 			if (!await AgentService.AuthorizeAsync(Agent, AclAction.ViewSession, User, null))
 			{
-				return Forbid();
+				return Forbid(AclAction.ViewSession, AgentId);
 			}
 
-			ISession? Session = await AgentService.GetSessionAsync(SessionId.ToObjectId());
-			if(Session == null || Session.AgentId != AgentIdValue)
+			ISession? Session = await AgentService.GetSessionAsync(SessionId);
+			if(Session == null || Session.AgentId != AgentId)
 			{
 				return NotFound();
 			}
@@ -312,21 +309,19 @@ namespace HordeServer.Controllers
 		/// <returns>Sessions </returns>
 		[HttpGet]
 		[Route("/api/v1/agents/{AgentId}/leases")]
-		public async Task<ActionResult<List<GetAgentLeaseResponse>>> FindLeasesAsync(string AgentId, [FromQuery] string? SessionId, [FromQuery] DateTimeOffset? StartTime, [FromQuery] DateTimeOffset? FinishTime, [FromQuery] int Index = 0, [FromQuery] int Count = 1000)
+		public async Task<ActionResult<List<GetAgentLeaseResponse>>> FindLeasesAsync(AgentId AgentId, [FromQuery] SessionId? SessionId, [FromQuery] DateTimeOffset? StartTime, [FromQuery] DateTimeOffset? FinishTime, [FromQuery] int Index = 0, [FromQuery] int Count = 1000)
 		{
-			AgentId AgentIdValue = new AgentId(AgentId);
-
-			IAgent? Agent = await AgentService.GetAgentAsync(AgentIdValue);
+			IAgent? Agent = await AgentService.GetAgentAsync(AgentId);
 			if (Agent == null)
 			{
-				return NotFound();
+				return NotFound(AgentId);
 			}
 			if (!await AgentService.AuthorizeAsync(Agent, AclAction.ViewLeases, User, null))
 			{
-				return Forbid();
+				return Forbid(AclAction.ViewLeases, AgentId);
 			}
 
-			List<ILease> Leases = await AgentService.FindLeasesAsync(AgentIdValue, SessionId?.ToObjectId(), StartTime?.UtcDateTime, FinishTime?.UtcDateTime, Index, Count);
+			List<ILease> Leases = await AgentService.FindLeasesAsync(AgentId, SessionId, StartTime?.UtcDateTime, FinishTime?.UtcDateTime, Index, Count);
 			return Leases.ConvertAll(x => new GetAgentLeaseResponse(x));
 		}
 
@@ -343,17 +338,17 @@ namespace HordeServer.Controllers
 			IAgent? Agent = await AgentService.GetAgentAsync(AgentId);
 			if(Agent == null)
 			{
-				return NotFound();
+				return NotFound(AgentId);
 			}
 			if (!await AclService.AuthorizeAsync(AclAction.ViewLeases, User))
 			{
-				return Forbid();
+				return Forbid(AclAction.ViewLeases, AgentId);
 			}
 
 			ILease? Lease = await AgentService.GetLeaseAsync(LeaseId);
 			if (Lease == null || Lease.AgentId != AgentId)
 			{
-				return NotFound();
+				return NotFound(AgentId, LeaseId);
 			}
 
 			return new GetAgentLeaseResponse(Lease);
diff --git a/Engine/Source/Programs/Horde/Horde.Build/Controllers/ConfigController.cs b/Engine/Source/Programs/Horde/Horde.Build/Controllers/ConfigController.cs
--- a/Engine/Source/Programs/Horde/Horde.Build/Controllers/ConfigController.cs
+++ b/Engine/Source/Programs/Horde/Horde.Build/Controllers/ConfigController.cs
@@ -11,6 +11,7 @@ using System.Threading.Tasks;
 using System.Text.Json;
 using Microsoft.Extensions.Logging;
 using System.Runtime.InteropServices;
+using HordeServer.Utilities;
 
 namespace HordeServer.Controllers
 {
@@ -21,7 +22,7 @@ namespace HordeServer.Controllers
 	[ApiController]
 	[Authorize]
 	[Route("[controller]")]
-	public class ConfigController : ControllerBase
+	public class ConfigController : HordeControllerBase
 	{
 		/// <summary>
 		/// The database service singleton
@@ -75,13 +76,12 @@ namespace HordeServer.Controllers
 		{
 			if (!await AclService.AuthorizeAsync(AclAction.AdminWrite, User))
 			{
-				return Forbid();
+				return Forbid(AclAction.AdminWrite);
 			}
 
 			if (!Settings.CurrentValue.SingleInstance)
 			{
-				Logger.LogError("Updating global configuration settings with ConfigController currently only supported in single instance mode");
-				return BadRequest();
+				return BadRequest("Updating global configuration settings with ConfigController currently only supported in single instance mode");
 			}
 
 			await ConfigService.UpdateGlobalConfig(Request);
@@ -99,12 +99,12 @@ namespace HordeServer.Controllers
 		{
 			if (!await AclService.AuthorizeAsync(AclAction.AdminRead, User))
 			{
-				return Forbid();
+				return Forbid(AclAction.AdminRead);
 			}
 
 			if (string.IsNullOrEmpty(Settings.CurrentValue.ConfigPath))
 			{
-				return BadRequest();
+				return BadRequest("Missing config path for settings");
 			}
 
 			FileReference GlobalConfigFile = new FileReference(Settings.CurrentValue.ConfigPath);
@@ -128,7 +128,7 @@ namespace HordeServer.Controllers
 		{
 			if (!await AclService.AuthorizeAsync(AclAction.AdminRead, User))
 			{
-				return Forbid();
+				return Forbid(AclAction.AdminRead);
 			}
 
 			GetServerSettingsResponse Response = new GetServerSettingsResponse(Settings.CurrentValue);
@@ -147,13 +147,12 @@ namespace HordeServer.Controllers
 		{
 			if (!await AclService.AuthorizeAsync(AclAction.AdminWrite, User))
 			{
-				return Forbid();
+				return Forbid(AclAction.AdminWrite);
 			}
 
 			if (!Settings.CurrentValue.SingleInstance)
 			{
-				Logger.LogError("Updating server settings from ConfigController currently only supported in single instance mode");
-				return BadRequest();
+				return BadRequest("Updating server settings from ConfigController currently only supported in single instance mode");
 			}
 
 			if (Request.Settings == null || Request.Settings.Count == 0)
diff --git a/Engine/Source/Programs/Horde/Horde.Build/Controllers/ErrorController.cs b/Engine/Source/Programs/Horde/Horde.Build/Controllers/ErrorController.cs
--- a/Engine/Source/Programs/Horde/Horde.Build/Controllers/ErrorController.cs
+++ b/Engine/Source/Programs/Horde/Horde.Build/Controllers/ErrorController.cs
@@ -1,7 +1,11 @@
 // Copyright Epic Games, Inc. All Rights Reserved.
 
+using EpicGames.Core;
+using HordeServer.Utilities;
 using Microsoft.AspNetCore.Diagnostics;
+using Microsoft.AspNetCore.Http;
 using Microsoft.AspNetCore.Mvc;
+using Microsoft.Extensions.Logging;
 
 namespace HordeServer.Controllers
 {
@@ -10,7 +14,7 @@ namespace HordeServer.Controllers
 	/// </summary>
 	[ApiController]
 	[ApiExplorerSettings(IgnoreApi = true)]
-	public class ErrorController : ControllerBase
+	public class ErrorController : HordeControllerBase
 	{
 		/// <summary>
 		/// Handle an exception and generate a problem response
@@ -24,8 +28,7 @@ namespace HordeServer.Controllers
 			{
 				return NoContent();
 			}
-
-			return Problem(detail: Context.Error.StackTrace, title: Context.Error.Message);
+			return StatusCode(StatusCodes.Status500InternalServerError, LogEvent.Create(LogLevel.Error, KnownLogEvents.None, Context.Error, Context.Error.Message));
 		}
 	}
 }
diff --git a/Engine/Source/Programs/Horde/Horde.Build/Controllers/IssuesController.cs b/Engine/Source/Programs/Horde/Horde.Build/Controllers/IssuesController.cs
--- a/Engine/Source/Programs/Horde/Horde.Build/Controllers/IssuesController.cs
+++ b/Engine/Source/Programs/Horde/Horde.Build/Controllers/IssuesController.cs
@@ -32,7 +32,7 @@ namespace HordeServer.Controllers
 	[Authorize]
 	[ApiController]
 	[Route("[controller]")]
-	public class IssuesController : ControllerBase
+	public class IssuesController : HordeControllerBase
 	{
 		private readonly IIssueCollection IssueCollection;
 		private readonly IIssueService IssueService;
@@ -216,11 +216,11 @@ namespace HordeServer.Controllers
 				IJob? Job = await JobService.GetJobAsync(JobId.Value);
 				if (Job == null)
 				{
-					return NotFound();
+					return NotFound(JobId.Value);
 				}
 				if(!await JobService.AuthorizeAsync(Job, AclAction.ViewJob, this.User, null))
 				{
-					return Forbid();
+					return Forbid(AclAction.ViewJob, JobId.Value);
 				}
 
 				IGraph Graph = await JobService.GetGraphAsync(Job);
diff --git a/Engine/Source/Programs/Horde/Horde.Build/Controllers/JobsController.cs b/Engine/Source/Programs/Horde/Horde.Build/Controllers/JobsController.cs
--- a/Engine/Source/Programs/Horde/Horde.Build/Controllers/JobsController.cs
+++ b/Engine/Source/Programs/Horde/Horde.Build/Controllers/JobsController.cs
@@ -40,7 +40,7 @@ namespace HordeServer.Controllers
 	[ApiController]
 	[Authorize]
 	[Route("[controller]")]
-	public class JobsController : ControllerBase
+	public class JobsController : HordeControllerBase
 	{
 		private readonly AclService AclService;
 		private readonly IGraphCollection Graphs;
@@ -82,11 +82,11 @@ namespace HordeServer.Controllers
 			IStream? Stream = await StreamService.GetStreamAsync(new StreamId(Create.StreamId));
 			if (Stream == null)
 			{
-				return BadRequest("Invalid StreamId parameter");
+				return NotFound(Create.StreamId);
 			}
 			if (!await StreamService.AuthorizeAsync(Stream, AclAction.CreateJob, User, null))
 			{
-				return Forbid();
+				return Forbid(AclAction.CreateJob, Stream.Id);
 			}
 
 			// Get the name of the template ref
@@ -100,17 +100,17 @@ namespace HordeServer.Controllers
 			}
 			if (!await StreamService.AuthorizeAsync(Stream, TemplateRef, AclAction.CreateJob, User, null))
 			{
-				return Forbid();
+				return Forbid(AclAction.CreateJob, Stream.Id);
 			}
 
 			ITemplate? Template = await TemplateCollection.GetAsync(TemplateRef.Hash);
 			if (Template == null)
 			{
-				return BadRequest($"Missing template referenced by {Create.TemplateId}");
+				return BadRequest("Missing template referenced by {TemplateId}", Create.TemplateId);
 			}
 			if (!Template.AllowPreflights && Create.PreflightChange > 0)
 			{
-				return BadRequest("Template does not allow preflights");
+				return BadRequest("Template {TemplateId} does not allow preflights", Create.TemplateId);
 			}
 
 			// Get the name of the new job
@@ -220,15 +220,15 @@ namespace HordeServer.Controllers
 			IJob? Job = await JobService.GetJobAsync(JobId);
 			if (Job == null)
 			{
-				return NotFound();
+				return NotFound(JobId);
 			}
 			if (!await JobService.AuthorizeAsync(Job, AclAction.DeleteJob, User, null))
 			{
-				return Forbid();
+				return Forbid(AclAction.DeleteJob, JobId);
 			}
 			if (!await JobService.DeleteJobAsync(Job))
 			{
-				return NotFound();
+				return NotFound(JobId);
 			}
 			return Ok();
 		}
@@ -246,17 +246,17 @@ namespace HordeServer.Controllers
 			IJob? Job = await JobService.GetJobAsync(JobId);
 			if (Job == null)
 			{
-				return NotFound();
+				return NotFound(JobId);
 			}
 
 			StreamPermissionsCache PermissionsCache = new StreamPermissionsCache();
 			if (!await JobService.AuthorizeAsync(Job, AclAction.UpdateJob, User, PermissionsCache))
 			{
-				return Forbid();
+				return Forbid(AclAction.UpdateJob, JobId);
 			}
 			if (Request.Acl != null && !await JobService.AuthorizeAsync(Job, AclAction.ChangePermissions, User, PermissionsCache))
 			{
-				return Forbid();
+				return Forbid(AclAction.ChangePermissions, JobId);
 			}
 
 			// Convert legacy behavior of clearing out the argument to setting the aborted flag
@@ -275,7 +275,7 @@ namespace HordeServer.Controllers
 			IJob? NewJob = await JobService.UpdateJobAsync(Job, Name: Request.Name, Priority: Request.Priority, AutoSubmit: Request.AutoSubmit, AbortedByUserId: AbortedByUserId, Arguments: Request.Arguments);
 			if (NewJob == null)
 			{
-				return NotFound();
+				return NotFound(JobId);
 			}
 			return Ok();
 		}
@@ -293,11 +293,11 @@ namespace HordeServer.Controllers
 			IJob? Job = await JobService.GetJobAsync(JobId);
 			if (Job == null)
 			{
-				return NotFound();
+				return NotFound(JobId);
 			}
 			if (!await JobService.AuthorizeAsync(Job, AclAction.CreateSubscription, User, null))
 			{
-				return Forbid();
+				return Forbid(AclAction.CreateSubscription, JobId);
 			}
 
 			ObjectId TriggerId = Job.NotificationTriggerId ?? ObjectId.GenerateNewId();
@@ -305,7 +305,7 @@ namespace HordeServer.Controllers
 			Job = await JobService.UpdateJobAsync(Job, null, null, null, null, TriggerId, null, null);
 			if (Job == null)
 			{
-				return NotFound();
+				return NotFound(JobId);
 			}
 
 			await NotificationService.UpdateSubscriptionsAsync(TriggerId, User, Request.Email, Request.Slack);
@@ -324,11 +324,11 @@ namespace HordeServer.Controllers
 			IJob? Job = await JobService.GetJobAsync(JobId);
 			if (Job == null)
 			{
-				return NotFound();
+				return NotFound(JobId);
 			}
 			if (!await JobService.AuthorizeAsync(Job, AclAction.CreateSubscription, User, null))
 			{
-				return Forbid();
+				return Forbid(AclAction.CreateSubscription, JobId);
 			}
 
 			INotificationSubscription? Subscription;
@@ -358,13 +358,13 @@ namespace HordeServer.Controllers
 			IJob? Job = await JobService.GetJobAsync(JobId);
 			if (Job == null)
 			{
-				return NotFound();
+				return NotFound(JobId);
 			}
 
 			StreamPermissionsCache Cache = new StreamPermissionsCache();
 			if (!await JobService.AuthorizeAsync(Job, AclAction.ViewJob, User, Cache))
 			{
-				return Forbid();
+				return Forbid(AclAction.ViewJob, JobId);
 			}
 			if (ModifiedAfter != null && Job.UpdateTimeUtc <= ModifiedAfter.Value)
 			{
@@ -496,11 +496,11 @@ namespace HordeServer.Controllers
 			IJob? Job = await JobService.GetJobAsync(JobId);
 			if (Job == null)
 			{
-				return NotFound();
+				return NotFound(JobId);
 			}
 			if (!await JobService.AuthorizeAsync(Job, AclAction.ViewJob, User, null))
 			{
-				return Forbid();
+				return Forbid(AclAction.ViewJob, Job.Id);
 			}
 
 			IGraph Graph = await JobService.GetGraphAsync(Job);
@@ -521,11 +521,11 @@ namespace HordeServer.Controllers
 			IJob? Job = await JobService.GetJobAsync(JobId);
 			if (Job == null)
 			{
-				return NotFound();
+				return NotFound(JobId);
 			}
 			if (!await JobService.AuthorizeAsync(Job, AclAction.ViewJob, User, null))
 			{
-				return Forbid();
+				return Forbid(AclAction.ViewJob, JobId);
 			}
 
 			IJobTiming JobTiming = await JobService.GetJobTimingAsync(Job);
@@ -567,17 +567,17 @@ namespace HordeServer.Controllers
 			IJob? Job = await JobService.GetJobAsync(JobId);
 			if (Job == null || Job.TemplateHash == null)
 			{
-				return NotFound();
+				return NotFound(JobId);
 			}
 			if (!await JobService.AuthorizeAsync(Job, AclAction.ViewJob, User, null))
 			{
-				return Forbid();
+				return Forbid(AclAction.ViewJob, JobId);
 			}
 
 			ITemplate? Template = await TemplateCollection.GetAsync(Job.TemplateHash);
 			if(Template == null)
 			{
-				return NotFound();
+				return NotFound(Job.StreamId, Job.TemplateId);
 			}
 
 			return new GetTemplateResponse(Template).ApplyFilter(Filter);
@@ -794,11 +794,11 @@ namespace HordeServer.Controllers
 				IJob? Job = await JobService.GetJobAsync(JobId);
 				if (Job == null)
 				{
-					return NotFound();
+					return NotFound(JobId);
 				}
 				if (!await JobService.AuthorizeAsync(Job, AclAction.ExecuteJob, User, null))
 				{
-					return Forbid();
+					return Forbid(AclAction.ExecuteJob, JobId);
 				}
 
 				IGraph Graph = await JobService.GetGraphAsync(Job);
@@ -826,11 +826,11 @@ namespace HordeServer.Controllers
 			IJob? Job = await JobService.GetJobAsync(JobId);
 			if (Job == null)
 			{
-				return NotFound();
+				return NotFound(JobId);
 			}
 			if (!await JobService.AuthorizeAsync(Job, AclAction.ViewJob, User, null))
 			{
-				return Forbid();
+				return Forbid(AclAction.ViewJob, JobId);
 			}
 
 			IGraph Graph = await JobService.GetGraphAsync(Job);
@@ -852,17 +852,17 @@ namespace HordeServer.Controllers
 			IJob? Job = await JobService.GetJobAsync(JobId);
 			if (Job == null)
 			{
-				return NotFound();
+				return NotFound(JobId);
 			}
 			if (!await JobService.AuthorizeAsync(Job, AclAction.ViewJob, User, null))
 			{
-				return Forbid();
+				return Forbid(AclAction.ViewJob, JobId);
 			}
 
 			IGraph Graph = await JobService.GetGraphAsync(Job);
 			if (GroupIdx < 0 || GroupIdx >= Graph.Groups.Count)
 			{
-				return NotFound();
+				return NotFound(JobId, GroupIdx);
 			}
 
 			return new GetGroupResponse(Graph.Groups[GroupIdx], Graph.Groups).ApplyFilter(Filter);
@@ -882,17 +882,17 @@ namespace HordeServer.Controllers
 			IJob? Job = await JobService.GetJobAsync(JobId);
 			if (Job == null)
 			{
-				return NotFound();
+				return NotFound(JobId);
 			}
 			if (!await JobService.AuthorizeAsync(Job, AclAction.ViewJob, User, null))
 			{
-				return Forbid();
+				return Forbid(AclAction.ViewJob, JobId);
 			}
 
 			IGraph Graph = await JobService.GetGraphAsync(Job);
 			if (GroupIdx < 0 || GroupIdx >= Graph.Groups.Count)
 			{
-				return NotFound();
+				return NotFound(JobId, GroupIdx);
 			}
 
 			return Graph.Groups[GroupIdx].Nodes.ConvertAll(x => new GetNodeResponse(x, Graph.Groups).ApplyFilter(Filter));
@@ -913,17 +913,21 @@ namespace HordeServer.Controllers
 			IJob? Job = await JobService.GetJobAsync(JobId);
 			if (Job == null)
 			{
-				return NotFound();
+				return NotFound(JobId);
 			}
 			if (!await JobService.AuthorizeAsync(Job, AclAction.ViewJob, User, null))
 			{
-				return Forbid();
+				return Forbid(AclAction.ViewJob, JobId);
 			}
 
 			IGraph Graph = await JobService.GetGraphAsync(Job);
-			if (GroupIdx < 0 || GroupIdx >= Graph.Groups.Count || NodeIdx < 0 || NodeIdx >= Graph.Groups[GroupIdx].Nodes.Count)
+			if (GroupIdx < 0 || GroupIdx >= Graph.Groups.Count)
 			{
-				return NotFound();
+				return NotFound(JobId, GroupIdx);
+			}
+			if(NodeIdx < 0 || NodeIdx >= Graph.Groups[GroupIdx].Nodes.Count)
+			{
+				return NotFound(JobId, GroupIdx, NodeIdx);
 			}
 
 			return new GetNodeResponse(Graph.Groups[GroupIdx].Nodes[NodeIdx], Graph.Groups).ApplyFilter(Filter);
@@ -943,15 +947,13 @@ namespace HordeServer.Controllers
 			IJob? Job = await JobService.GetJobAsync(JobId);
 			if (Job == null)
 			{
-				return NotFound();
+				return NotFound(JobId);
 			}
 			if (!await JobService.AuthorizeAsync(Job, AclAction.ViewJob, User, null))
 			{
-				return Forbid();
+				return Forbid(AclAction.ViewJob, JobId);
 			}
 
-			IGraph Graph = await JobService.GetGraphAsync(Job);
-
 			List<object> Responses = new List<object>();
 			foreach (IJobStepBatch Batch in Job.Batches)
 			{
@@ -969,30 +971,28 @@ namespace HordeServer.Controllers
 		/// <param name="Request">Updates to apply to the node</param>
 		[HttpPut]
 		[Route("/api/v1/jobs/{JobId}/batches/{BatchId}")]
-		public async Task<ActionResult> UpdateBatchAsync(JobId JobId, string BatchId, [FromBody] UpdateBatchRequest Request)
+		public async Task<ActionResult> UpdateBatchAsync(JobId JobId, SubResourceId BatchId, [FromBody] UpdateBatchRequest Request)
 		{
-			SubResourceId BatchIdValue = BatchId.ToSubResourceId();
-
 			IJob? Job = await JobService.GetJobAsync(JobId);
 			if (Job == null)
 			{
-				return NotFound();
+				return NotFound(JobId);
 			}
 
-			IJobStepBatch? Batch = Job.Batches.FirstOrDefault(x => x.Id == BatchIdValue);
+			IJobStepBatch? Batch = Job.Batches.FirstOrDefault(x => x.Id == BatchId);
 			if (Batch == null)
 			{
-				return NotFound();
+				return NotFound(JobId, BatchId);
 			}
 			if (Batch.SessionId == null || !User.HasSessionClaim(Batch.SessionId.Value))
 			{
-				return Forbid();
+				return Forbid("Missing session claim for job {JobId} batch {BatchId}", JobId, BatchId);
 			}
 
-			IJob? NewJob = await JobService.UpdateBatchAsync(Job, BatchIdValue, Request.LogId?.ToObjectId<ILogFile>(), Request.State);
+			IJob? NewJob = await JobService.UpdateBatchAsync(Job, BatchId, Request.LogId?.ToObjectId<ILogFile>(), Request.State);
 			if (NewJob == null)
 			{
-				return NotFound();
+				return NotFound(JobId);
 			}
 			return Ok();
 		}
@@ -1007,31 +1007,29 @@ namespace HordeServer.Controllers
 		[HttpGet]
 		[Route("/api/v1/jobs/{JobId}/batches/{BatchId}")]
 		[ProducesResponseType(typeof(GetBatchResponse), 200)]
-		public async Task<ActionResult<object>> GetBatchAsync(JobId JobId, string BatchId, [FromQuery] PropertyFilter? Filter = null)
+		public async Task<ActionResult<object>> GetBatchAsync(JobId JobId, SubResourceId BatchId, [FromQuery] PropertyFilter? Filter = null)
 		{
-			SubResourceId BatchIdValue = BatchId.ToSubResourceId();
-
 			IJob? Job = await JobService.GetJobAsync(JobId);
 			if (Job == null)
 			{
-				return NotFound();
+				return NotFound(JobId);
 			}
 			if (!await JobService.AuthorizeAsync(Job, AclAction.ViewJob, User, null))
 			{
-				return Forbid();
+				return Forbid(AclAction.ViewJob, JobId);
 			}
 
 			IGraph Graph = await JobService.GetGraphAsync(Job);
 			foreach (IJobStepBatch Batch in Job.Batches)
 			{
-				if (Batch.Id == BatchIdValue)
+				if (Batch.Id == BatchId)
 				{
 					GetBatchResponse Response = await CreateBatchResponseAsync(Batch);
 					return Response.ApplyFilter(Filter);
 				}
 			}
 
-			return NotFound();
+			return NotFound(JobId, BatchId);
 		}
 
 		/// <summary>
@@ -1044,27 +1042,23 @@ namespace HordeServer.Controllers
 		[HttpGet]
 		[Route("/api/v1/jobs/{JobId}/batches/{BatchId}/steps")]
 		[ProducesResponseType(typeof(List<GetStepResponse>), 200)]
-		public async Task<ActionResult<List<object>>> GetStepsAsync(JobId JobId, string BatchId, [FromQuery] PropertyFilter? Filter = null)
+		public async Task<ActionResult<List<object>>> GetStepsAsync(JobId JobId, SubResourceId BatchId, [FromQuery] PropertyFilter? Filter = null)
 		{
-			SubResourceId BatchIdValue = BatchId.ToSubResourceId();
-
 			IJob? Job = await JobService.GetJobAsync(JobId);
 			if (Job == null)
 			{
-				return NotFound();
+				return NotFound(JobId);
 			}
 			if (!await JobService.AuthorizeAsync(Job, AclAction.ViewJob, User, null))
 			{
-				return Forbid();
+				return Forbid(AclAction.ViewJob, JobId);
 			}
 
 			IGraph Graph = await JobService.GetGraphAsync(Job);
 			foreach (IJobStepBatch Batch in Job.Batches)
 			{
-				if (Batch.Id == BatchIdValue)
+				if (Batch.Id == BatchId)
 				{
-					INodeGroup Group = Graph.Groups[Batch.GroupIdx];
-
 					List<object> Responses = new List<object>();
 					foreach (IJobStep Step in Batch.Steps)
 					{
@@ -1075,7 +1069,7 @@ namespace HordeServer.Controllers
 				}
 			}
 
-			return NotFound();
+			return NotFound(JobId, BatchId);
 		}
 
 		/// <summary>
@@ -1087,24 +1081,21 @@ namespace HordeServer.Controllers
 		/// <param name="Request">Updates to apply to the node</param>
 		[HttpPut]
 		[Route("/api/v1/jobs/{JobId}/batches/{BatchId}/steps/{StepId}")]
-		public async Task<ActionResult<UpdateStepResponse>> UpdateStepAsync(JobId JobId, string BatchId, string StepId, [FromBody] UpdateStepRequest Request)
+		public async Task<ActionResult<UpdateStepResponse>> UpdateStepAsync(JobId JobId, SubResourceId BatchId, SubResourceId StepId, [FromBody] UpdateStepRequest Request)
 		{
-			SubResourceId BatchIdValue = BatchId.ToSubResourceId();
-			SubResourceId StepIdValue = StepId.ToSubResourceId();
-
 			IJob? Job = await JobService.GetJobAsync(JobId);
 			if (Job == null)
 			{
-				return NotFound();
+				return NotFound(JobId);
 			}
 
 			// Check permissions for updating this step. Only the agent executing the step can modify the state of it.
 			if (Request.State != JobStepState.Unspecified || Request.Outcome != JobStepOutcome.Unspecified)
 			{
-				IJobStepBatch? Batch = Job.Batches.FirstOrDefault(x => x.Id == BatchIdValue);
+				IJobStepBatch? Batch = Job.Batches.FirstOrDefault(x => x.Id == BatchId);
 				if (Batch == null)
 				{
-					return NotFound();
+					return NotFound(JobId, BatchId);
 				}
 				if (!Batch.SessionId.HasValue || !User.HasSessionClaim(Batch.SessionId.Value))
 				{
@@ -1115,14 +1106,14 @@ namespace HordeServer.Controllers
 			{
 				if (!await JobService.AuthorizeAsync(Job, AclAction.RetryJobStep, User, null))
 				{
-					return Forbid();
+					return Forbid(AclAction.RetryJobStep, JobId);
 				}
 			}
 			if (Request.Properties != null)
 			{
 				if (!await JobService.AuthorizeAsync(Job, AclAction.UpdateJob, User, null))
 				{
-					return Forbid();
+					return Forbid(AclAction.UpdateJob, JobId);
 				}
 			}
 
@@ -1131,16 +1122,16 @@ namespace HordeServer.Controllers
 
 			try
 			{
-				IJob? NewJob = await JobService.UpdateStepAsync(Job, BatchIdValue, StepIdValue, Request.State, Request.Outcome, Request.AbortRequested, AbortByUser, Request.LogId?.ToObjectId<ILogFile>(), null, RetryByUser, Request.Priority, null, Request.Properties);
+				IJob? NewJob = await JobService.UpdateStepAsync(Job, BatchId, StepId, Request.State, Request.Outcome, Request.AbortRequested, AbortByUser, Request.LogId?.ToObjectId<ILogFile>(), null, RetryByUser, Request.Priority, null, Request.Properties);
 				if (NewJob == null)
 				{
-					return NotFound();
+					return NotFound(JobId);
 				}
 
 				UpdateStepResponse Response = new UpdateStepResponse();
 				if (Request.Retry ?? false)
 				{
-					JobStepRefId? RetriedStepId = FindRetriedStep(Job, BatchIdValue, StepIdValue);
+					JobStepRefId? RetriedStepId = FindRetriedStep(Job, BatchId, StepId);
 					if (RetriedStepId != null)
 					{
 						Response.BatchId = RetriedStepId.Value.BatchId.ToString();
@@ -1196,38 +1187,35 @@ namespace HordeServer.Controllers
 		[HttpGet]
 		[Route("/api/v1/jobs/{JobId}/batches/{BatchId}/steps/{StepId}")]
 		[ProducesResponseType(typeof(GetStepResponse), 200)]
-		public async Task<ActionResult<object>> GetStepAsync(JobId JobId, string BatchId, string StepId, [FromQuery] PropertyFilter? Filter = null)
+		public async Task<ActionResult<object>> GetStepAsync(JobId JobId, SubResourceId BatchId, SubResourceId StepId, [FromQuery] PropertyFilter? Filter = null)
 		{
-			SubResourceId BatchIdValue = BatchId.ToSubResourceId();
-			SubResourceId StepIdValue = StepId.ToSubResourceId();
-
 			IJob? Job = await JobService.GetJobAsync(JobId);
 			if (Job == null)
 			{
-				return NotFound();
+				return NotFound(JobId);
 			}
 			if (!await JobService.AuthorizeAsync(Job, AclAction.ViewJob, User, null))
 			{
-				return Forbid();
+				return Forbid(AclAction.ViewJob, JobId);
 			}
 
 			foreach (IJobStepBatch Batch in Job.Batches)
 			{
-				if (Batch.Id == BatchIdValue)
+				if (Batch.Id == BatchId)
 				{
 					foreach (IJobStep Step in Batch.Steps)
 					{
-						if (Step.Id == StepIdValue)
+						if (Step.Id == StepId)
 						{
 							GetStepResponse Response = await CreateStepResponseAsync(Step);
 							return Response.ApplyFilter(Filter);
 						}
 					}
-					break;
+					return NotFound(JobId, BatchId, StepId);
 				}
 			}
 
-			return NotFound();
+			return NotFound(JobId, BatchId);
 		}
 
 		/// <summary>
@@ -1240,38 +1228,35 @@ namespace HordeServer.Controllers
 		/// <returns>Information about the requested job</returns>
 		[HttpPut]
 		[Route("/api/v1/jobs/{JobId}/batches/{BatchId}/steps/{StepId}/notifications")]
-		public async Task<ActionResult> UpdateStepNotificationsAsync(JobId JobId, string BatchId, string StepId, [FromBody] UpdateNotificationsRequest Request)
+		public async Task<ActionResult> UpdateStepNotificationsAsync(JobId JobId, SubResourceId BatchId, SubResourceId StepId, [FromBody] UpdateNotificationsRequest Request)
 		{
-			SubResourceId BatchIdValue = BatchId.ToSubResourceId();
-			SubResourceId StepIdValue = StepId.ToSubResourceId();
-
 			IJob? Job = await JobService.GetJobAsync(JobId);
 			if (Job == null)
 			{
-				return NotFound();
+				return NotFound(JobId);
 			}
 			if (!await JobService.AuthorizeAsync(Job, AclAction.CreateSubscription, User, null))
 			{
-				return Forbid();
+				return Forbid(AclAction.CreateSubscription, JobId);
 			}
 
-			if (!Job.TryGetBatch(BatchIdValue, out IJobStepBatch? Batch))
+			if (!Job.TryGetBatch(BatchId, out IJobStepBatch? Batch))
 			{
-				return NotFound();
+				return NotFound(JobId, BatchId);
 			}
 
-			if (!Batch.TryGetStep(StepIdValue, out IJobStep? Step))
+			if (!Batch.TryGetStep(StepId, out IJobStep? Step))
 			{
-				return NotFound();
+				return NotFound(JobId, BatchId, StepId);
 			}
 
 			ObjectId? TriggerId = Step.NotificationTriggerId;
 			if (TriggerId == null)
 			{
 				TriggerId = ObjectId.GenerateNewId();
-				if (await JobService.UpdateStepAsync(Job, BatchIdValue, StepIdValue, JobStepState.Unspecified, JobStepOutcome.Unspecified, null, null, null, TriggerId, null, null, null) == null)
+				if (await JobService.UpdateStepAsync(Job, BatchId, StepId, JobStepState.Unspecified, JobStepOutcome.Unspecified, null, null, null, TriggerId, null, null, null) == null)
 				{
-					return NotFound();
+					return NotFound(JobId, BatchId, StepId);
 				}
 			}
 
@@ -1288,22 +1273,26 @@ namespace HordeServer.Controllers
 		/// <returns>Information about the requested job</returns>
 		[HttpGet]
 		[Route("/api/v1/jobs/{JobId}/batches/{BatchId}/steps/{StepId}/notifications")]
-		public async Task<ActionResult<GetNotificationResponse>> GetStepNotificationsAsync(JobId JobId, string BatchId, string StepId)
+		public async Task<ActionResult<GetNotificationResponse>> GetStepNotificationsAsync(JobId JobId, SubResourceId BatchId, SubResourceId StepId)
 		{
 			IJob? Job = await JobService.GetJobAsync(JobId);
 			if (Job == null)
 			{
-				return NotFound();
+				return NotFound(JobId);
 			}
 
 			IJobStep? Step;
-			if(!Job.TryGetStep(BatchId.ToSubResourceId(), StepId.ToSubResourceId(), out Step))
+			if (!Job.TryGetBatch(BatchId, out IJobStepBatch? Batch))
 			{
-				return NotFound();
+				return NotFound(JobId, BatchId);
+			}
+			if (!Batch.TryGetStep(StepId, out Step))
+			{
+				return NotFound(JobId, BatchId, StepId);
 			}
 			if (!await JobService.AuthorizeAsync(Job, AclAction.CreateSubscription, User, null))
 			{
-				return Forbid();
+				return Forbid(AclAction.CreateSubscription, JobId);
 			}
 
 			INotificationSubscription? Subscription;
@@ -1328,26 +1317,27 @@ namespace HordeServer.Controllers
 		/// <returns>List of nodes to be executed</returns>
 		[HttpGet]
 		[Route("/api/v1/jobs/{JobId}/batches/{BatchId}/steps/{StepId}/artifacts/{*Name}")]
-		public async Task<ActionResult> GetArtifactAsync(JobId JobId, string BatchId, string StepId, string Name)
+		public async Task<ActionResult> GetArtifactAsync(JobId JobId, SubResourceId BatchId, SubResourceId StepId, string Name)
 		{
-			SubResourceId BatchIdValue = BatchId.ToSubResourceId();
-			SubResourceId StepIdValue = StepId.ToSubResourceId();
-
 			IJob? Job = await JobService.GetJobAsync(JobId);
 			if (Job == null)
 			{
-				return NotFound();
+				return NotFound(JobId);
 			}
 			if (!await JobService.AuthorizeAsync(Job, AclAction.ViewJob, User, null))
 			{
-				return Forbid();
+				return Forbid(AclAction.ViewJob, JobId);
 			}
-			if (!Job.TryGetStep(BatchId.ToSubResourceId(), StepId.ToSubResourceId(), out _))
+			if (!Job.TryGetBatch(BatchId, out IJobStepBatch? Batch))
 			{
-				return NotFound();
+				return NotFound(JobId, BatchId);
+			}
+			if (!Batch.TryGetStep(StepId, out _))
+			{
+				return NotFound(JobId, BatchId, StepId);
 			}
 
-			List<IArtifact> Artifacts = await ArtifactCollection.GetArtifactsAsync(JobId, StepIdValue, Name);
+			List<IArtifact> Artifacts = await ArtifactCollection.GetArtifactsAsync(JobId, StepId, Name);
 			if (Artifacts.Count == 0)
 			{
 				return NotFound();
@@ -1366,26 +1356,27 @@ namespace HordeServer.Controllers
 		/// <returns>List of nodes to be executed</returns>
 		[HttpGet]
 		[Route("/api/v1/jobs/{JobId}/batches/{BatchId}/steps/{StepId}/trace")]
-		public async Task<ActionResult> GetStepTraceAsync(JobId JobId, string BatchId, string StepId)
+		public async Task<ActionResult> GetStepTraceAsync(JobId JobId, SubResourceId BatchId, SubResourceId StepId)
 		{
-			SubResourceId BatchIdValue = BatchId.ToSubResourceId();
-			SubResourceId StepIdValue = StepId.ToSubResourceId();
-
 			IJob? Job = await JobService.GetJobAsync(JobId);
 			if (Job == null)
 			{
-				return NotFound();
+				return NotFound(JobId);
 			}
 			if (!await JobService.AuthorizeAsync(Job, AclAction.ViewJob, User, null))
 			{
-				return Forbid();
+				return Forbid(AclAction.ViewJob, JobId);
 			}
-			if (!Job.TryGetStep(BatchId.ToSubResourceId(), StepId.ToSubResourceId(), out _))
+			if (!Job.TryGetBatch(BatchId, out IJobStepBatch? Batch))
 			{
-				return NotFound();
+				return NotFound(JobId, BatchId);
+			}
+			if (!Batch.TryGetStep(StepId, out _))
+			{
+				return NotFound(JobId, BatchId, StepId);
 			}
 
-			List<IArtifact> Artifacts = await ArtifactCollection.GetArtifactsAsync(JobId, StepIdValue, null);
+			List<IArtifact> Artifacts = await ArtifactCollection.GetArtifactsAsync(JobId, StepId, null);
 			foreach (IArtifact Artifact in Artifacts)
 			{
 				if (Artifact.Name.Equals("trace.json", StringComparison.OrdinalIgnoreCase))
@@ -1412,11 +1403,11 @@ namespace HordeServer.Controllers
 				IJob? Job = await JobService.GetJobAsync(JobId);
 				if (Job == null)
 				{
-					return NotFound();
+					return NotFound(JobId);
 				}
 				if (!await JobService.AuthorizeAsync(Job, AclAction.CreateSubscription, User, null))
 				{
-					return Forbid();
+					return Forbid(AclAction.CreateSubscription, JobId);
 				}
 
 				ObjectId NewTriggerId;
@@ -1453,7 +1444,7 @@ namespace HordeServer.Controllers
 			IJob? Job = await JobService.GetJobAsync(JobId);
 			if (Job == null)
 			{
-				return NotFound();
+				return NotFound(JobId);
 			}
 
 			INotificationSubscription? Subscription;
diff --git a/Engine/Source/Programs/Horde/Horde.Build/Controllers/LeasesController.cs b/Engine/Source/Programs/Horde/Horde.Build/Controllers/LeasesController.cs
--- a/Engine/Source/Programs/Horde/Horde.Build/Controllers/LeasesController.cs
+++ b/Engine/Source/Programs/Horde/Horde.Build/Controllers/LeasesController.cs
@@ -26,6 +26,7 @@ namespace HordeServer.Controllers
 {
 	using LeaseId = ObjectId<ILease>;
 	using PoolId = StringId<IPool>;
+	using SessionId = ObjectId<ISession>;
 	using AgentSoftwareChannelName = StringId<AgentSoftwareChannels>;
 
 	/// <summary>
@@ -34,7 +35,7 @@ namespace HordeServer.Controllers
 	[ApiController]
 	[Authorize]
 	[Route("[controller]")]
-	public class LeasesController : ControllerBase
+	public class LeasesController : HordeControllerBase
 	{
 		/// <summary>
 		/// Singleton instance of the ACL service
@@ -76,32 +77,29 @@ namespace HordeServer.Controllers
 		/// <returns>Sessions </returns>
 		[HttpGet]
 		[Route("/api/v1/leases")]
-		public async Task<ActionResult<List<GetAgentLeaseResponse>>> FindLeasesAsync([FromQuery] string? AgentId, [FromQuery] string? SessionId, [FromQuery] DateTimeOffset? StartTime, [FromQuery] DateTimeOffset? FinishTime, [FromQuery] int Index = 0, [FromQuery] int Count = 1000)
+		public async Task<ActionResult<List<GetAgentLeaseResponse>>> FindLeasesAsync([FromQuery] AgentId? AgentId, [FromQuery] SessionId? SessionId, [FromQuery] DateTimeOffset? StartTime, [FromQuery] DateTimeOffset? FinishTime, [FromQuery] int Index = 0, [FromQuery] int Count = 1000)
 		{
-			AgentId? AgentIdValue = null;
 			if (AgentId == null)
 			{
 				if (!await AclService.AuthorizeAsync(AclAction.ViewLeases, User))
 				{
-					return Forbid();
+					return Forbid(AclAction.ViewLeases);
 				}
 			}
 			else
 			{
-				AgentIdValue = new AgentId(AgentId);
-
-				IAgent? Agent = await AgentService.GetAgentAsync(AgentIdValue.Value);
+				IAgent? Agent = await AgentService.GetAgentAsync(AgentId.Value);
 				if (Agent == null)
 				{
-					return NotFound();
+					return NotFound(AgentId.Value);
 				}
 				if (!await AgentService.AuthorizeAsync(Agent, AclAction.ViewLeases, User, null))
 				{
-					return Forbid();
+					return Forbid(AclAction.ViewLeases, AgentId.Value);
 				}
 			}
 
-			List<ILease> Leases = await AgentService.FindLeasesAsync(AgentIdValue, SessionId?.ToObjectId(), StartTime?.UtcDateTime, FinishTime?.UtcDateTime, Index, Count);
+			List<ILease> Leases = await AgentService.FindLeasesAsync(AgentId, SessionId, StartTime?.UtcDateTime, FinishTime?.UtcDateTime, Index, Count);
 			return Leases.ConvertAll(x => new GetAgentLeaseResponse(x));
 		}
 
@@ -117,17 +115,17 @@ namespace HordeServer.Controllers
 			ILease? Lease = await AgentService.GetLeaseAsync(LeaseId);
 			if (Lease == null)
 			{
-				return NotFound();
+				return NotFound(LeaseId);
 			}
 
 			IAgent? Agent = await AgentService.GetAgentAsync(Lease.AgentId);
 			if (Agent == null)
 			{
-				return NotFound();
+				return NotFound(Lease.AgentId);
 			}
 			if (!await AgentService.AuthorizeAsync(Agent, AclAction.ViewLeases, User, null))
 			{
-				return Forbid();
+				return Forbid(AclAction.ViewLeases);
 			}
 
 			return new GetAgentLeaseResponse(Lease);
@@ -152,35 +150,33 @@ namespace HordeServer.Controllers
 			ILease? Lease = await AgentService.GetLeaseAsync(LeaseId);
 			if (Lease == null)
 			{
-				return NotFound();
+				return NotFound(LeaseId);
 			}
 
 			IAgent? Agent = await AgentService.GetAgentAsync(Lease.AgentId);
 			if (Agent == null)
 			{
-				return NotFound();
+				return NotFound(Lease.AgentId);
 			}
 
 			if (!await AgentService.AuthorizeAsync(Agent, AclAction.AdminWrite, User, null))
 			{
-				return Forbid();
+				return Forbid(AclAction.AdminWrite, Lease.AgentId);
 			}
 
 			AgentLease? AgentLease = Agent.Leases.FirstOrDefault(x => x.Id == LeaseId);
 
 			if (AgentLease == null)
 			{
-				return NotFound();
+				return NotFound(Agent.Id, LeaseId);
 			}
 
 			if (!AgentLease.IsConformLease())
 			{
-				Logger.LogError("Lease abort only supported on conform leases for now, {LeaseId}", LeaseId);
-				return Ok();
+				return BadRequest("Lease abort only supported on conform leases for now, {LeaseId}", LeaseId);
 			}
 
 			await AgentService.CancelLeaseAsync(Agent, LeaseId);
-
 			return Ok();
 		}
 
diff --git a/Engine/Source/Programs/Horde/Horde.Build/Controllers/PoolsController.cs b/Engine/Source/Programs/Horde/Horde.Build/Controllers/PoolsController.cs
--- a/Engine/Source/Programs/Horde/Horde.Build/Controllers/PoolsController.cs
+++ b/Engine/Source/Programs/Horde/Horde.Build/Controllers/PoolsController.cs
@@ -23,7 +23,7 @@ namespace HordeServer.Controllers
 	[ApiController]
 	[Authorize]
 	[Route("[controller]")]
-	public class PoolsController : ControllerBase
+	public class PoolsController : HordeControllerBase
 	{
 		/// <summary>
 		/// Singleton instance of the ACL service
@@ -57,7 +57,7 @@ namespace HordeServer.Controllers
 		{
 			if(!await AclService.AuthorizeAsync(AclAction.CreatePool, User))
 			{
-				return Forbid();
+				return Forbid(AclAction.CreatePool);
 			}
 
 			IPool NewPool = await PoolService.CreatePoolAsync(Create.Name, Create.Condition, Create.EnableAutoscaling, Create.MinAgents, Create.NumReserveAgents, Create.Properties);
@@ -76,7 +76,7 @@ namespace HordeServer.Controllers
 		{
 			if (!await AclService.AuthorizeAsync(AclAction.ListPools, User))
 			{
-				return Forbid();
+				return Forbid(AclAction.ListPools);
 			}
 
 			List<IPool> Pools = await PoolService.GetPoolsAsync();
@@ -102,7 +102,7 @@ namespace HordeServer.Controllers
 		{
 			if (!await AclService.AuthorizeAsync(AclAction.ViewPool, User))
 			{
-				return Forbid();
+				return Forbid(AclAction.ViewPool);
 			}
 
 			PoolId PoolIdValue = new PoolId(PoolId);
@@ -110,7 +110,7 @@ namespace HordeServer.Controllers
 			IPool? Pool = await PoolService.GetPoolAsync(PoolIdValue);
 			if (Pool == null)
 			{
-				return NotFound();
+				return NotFound(PoolIdValue);
 			}
 
 			return new GetPoolResponse(Pool).ApplyFilter(Filter);
@@ -128,7 +128,7 @@ namespace HordeServer.Controllers
 		{
 			if (!await AclService.AuthorizeAsync(AclAction.UpdatePool, User))
 			{
-				return Forbid();
+				return Forbid(AclAction.UpdatePool);
 			}
 
 			PoolId PoolIdValue = new PoolId(PoolId);
@@ -136,7 +136,7 @@ namespace HordeServer.Controllers
 			IPool? Pool = await PoolService.GetPoolAsync(PoolIdValue);
 			if(Pool == null)
 			{
-				return NotFound();
+				return NotFound(PoolIdValue);
 			}
 
 			await PoolService.UpdatePoolAsync(Pool, Update.Name, Update.Condition, Update.EnableAutoscaling, Update.MinAgents, Update.NumReserveAgents, Update.Properties);
@@ -154,15 +154,15 @@ namespace HordeServer.Controllers
 		{
 			if (!await AclService.AuthorizeAsync(AclAction.DeletePool, User))
 			{
-				return Forbid();
+				return Forbid(AclAction.DeletePool);
 			}
 
 			PoolId PoolIdValue = new PoolId(PoolId);
 			if(!await PoolService.DeletePoolAsync(PoolIdValue))
 			{
-				return NotFound();
+				return NotFound(PoolIdValue);
 			}
-			return new OkResult();
+			return Ok();
 		}
 
 		/// <summary>
@@ -176,7 +176,7 @@ namespace HordeServer.Controllers
 		{
 			if (!await AclService.AuthorizeAsync(AclAction.UpdatePool, User))
 			{
-				return Forbid();
+				return Forbid(AclAction.UpdatePool);
 			}
 
 			foreach (BatchUpdatePoolRequest Update in BatchUpdates)
@@ -186,12 +186,12 @@ namespace HordeServer.Controllers
 				IPool? Pool = await PoolService.GetPoolAsync(PoolIdValue);
 				if (Pool == null)
 				{
-					return NotFound();
+					return NotFound(PoolIdValue);
 				}
 
 				await PoolService.UpdatePoolAsync(Pool, Update.Name, NewProperties: Update.Properties);
 			}
-			return new OkResult();
+			return Ok();
 		}
 	}
 }
diff --git a/Engine/Source/Programs/Horde/Horde.Build/Controllers/ProjectsController.cs b/Engine/Source/Programs/Horde/Horde.Build/Controllers/ProjectsController.cs
--- a/Engine/Source/Programs/Horde/Horde.Build/Controllers/ProjectsController.cs
+++ b/Engine/Source/Programs/Horde/Horde.Build/Controllers/ProjectsController.cs
@@ -23,7 +23,7 @@ namespace HordeServer.Controllers
 	[ApiController]
 	[Authorize]
 	[Route("[controller]")]
-	public class ProjectsController : ControllerBase
+	public class ProjectsController : HordeControllerBase
 	{
 		/// <summary>
 		/// Singleton instance of the project service
@@ -88,20 +88,18 @@ namespace HordeServer.Controllers
 		[HttpGet]
 		[Route("/api/v1/projects/{ProjectId}")]
 		[ProducesResponseType(typeof(List<GetProjectResponse>), 200)]
-		public async Task<ActionResult<object>> GetProjectAsync(string ProjectId, [FromQuery] PropertyFilter? Filter = null)
+		public async Task<ActionResult<object>> GetProjectAsync(ProjectId ProjectId, [FromQuery] PropertyFilter? Filter = null)
 		{
-			ProjectId ProjectIdValue = new ProjectId(ProjectId);
-
-			IProject? Project = await ProjectService.GetProjectAsync(ProjectIdValue);
+			IProject? Project = await ProjectService.GetProjectAsync(ProjectId);
 			if (Project == null)
 			{
-				return NotFound();
+				return NotFound(ProjectId);
 			}
 
 			ProjectPermissionsCache Cache = new ProjectPermissionsCache();
 			if (!await ProjectService.AuthorizeAsync(Project, AclAction.ViewProject, User, Cache))
 			{
-				return Forbid();
+				return Forbid(AclAction.ViewProject, ProjectId);
 			}
 
 			bool bIncludeStreams = PropertyFilter.Includes(Filter, nameof(GetProjectResponse.Streams));
@@ -133,21 +131,19 @@ namespace HordeServer.Controllers
 		/// <returns>Information about the requested project</returns>
 		[HttpGet]
 		[Route("/api/v1/projects/{ProjectId}/logo")]
-		public async Task<ActionResult<object>> GetProjectLogoAsync(string ProjectId)
+		public async Task<ActionResult<object>> GetProjectLogoAsync(ProjectId ProjectId)
 		{
-			ProjectId ProjectIdValue = new ProjectId(ProjectId);
-
-			IProject? Project = await ProjectService.GetProjectAsync(ProjectIdValue);
+			IProject? Project = await ProjectService.GetProjectAsync(ProjectId);
 			if (Project == null)
 			{
-				return NotFound();
+				return NotFound(ProjectId);
 			}
 			if (!await ProjectService.AuthorizeAsync(Project, AclAction.ViewProject, User, null))
 			{
-				return Forbid();
+				return Forbid(AclAction.ViewProject, ProjectId);
 			}
 
-			IProjectLogo? ProjectLogo = await ProjectService.Collection.GetLogoAsync(ProjectIdValue);
+			IProjectLogo? ProjectLogo = await ProjectService.Collection.GetLogoAsync(ProjectId);
 			if (ProjectLogo == null)
 			{
 				return NotFound();
diff --git a/Engine/Source/Programs/Horde/Horde.Build/Controllers/StreamsController.cs b/Engine/Source/Programs/Horde/Horde.Build/Controllers/StreamsController.cs
--- a/Engine/Source/Programs/Horde/Horde.Build/Controllers/StreamsController.cs
+++ b/Engine/Source/Programs/Horde/Horde.Build/Controllers/StreamsController.cs
@@ -26,7 +26,7 @@ namespace HordeServer.Controllers
 	[ApiController]
 	[Authorize]
 	[Route("[controller]")]
-	public class StreamsController : ControllerBase
+	public class StreamsController : HordeControllerBase
 	{
 		private readonly AclService AclService;
 		private readonly ProjectService ProjectService;
@@ -89,23 +89,20 @@ namespace HordeServer.Controllers
 		[HttpGet]
 		[Route("/api/v1/streams/{StreamId}")]
 		[ProducesResponseType(typeof(GetStreamResponse), 200)]
-		public async Task<ActionResult<object>> GetStreamAsync(string StreamId, [FromQuery] PropertyFilter? Filter = null)
+		public async Task<ActionResult<object>> GetStreamAsync(StreamId StreamId, [FromQuery] PropertyFilter? Filter = null)
 		{
-			StreamId StreamIdValue = new StreamId(StreamId);
-
-			IStream? Stream = await StreamService.GetStreamAsync(StreamIdValue);
+			IStream? Stream = await StreamService.GetStreamAsync(StreamId);
 			if (Stream == null)
 			{
-				return NotFound();
+				return NotFound(StreamId);
 			}
 
 			ProjectPermissionsCache PermissionsCache = new ProjectPermissionsCache();
 			if (!await StreamService.AuthorizeAsync(Stream, AclAction.ViewStream, User, PermissionsCache))
 			{
-				return Forbid();
+				return Forbid(AclAction.ViewStream, StreamId);
 			}
 
-			bool bIncludeAcl = await StreamService.AuthorizeAsync(Stream, AclAction.ViewPermissions, User, PermissionsCache);
 			return PropertyFilter.Apply(await CreateGetStreamResponse(Stream, PermissionsCache), Filter);
 		}
 
@@ -148,24 +145,22 @@ namespace HordeServer.Controllers
 		[HttpGet]
 		[Route("/api/v1/streams/{StreamId}/changes")]
 		[ProducesResponseType(typeof(List<GetChangeSummaryResponse>), 200)]
-		public async Task<ActionResult<List<object>>> GetChangesAsync(string StreamId, [FromQuery] int? Min = null, [FromQuery] int? Max = null, [FromQuery] int Results = 50, PropertyFilter? Filter = null)
+		public async Task<ActionResult<List<object>>> GetChangesAsync(StreamId StreamId, [FromQuery] int? Min = null, [FromQuery] int? Max = null, [FromQuery] int Results = 50, PropertyFilter? Filter = null)
 		{
-			StreamId StreamIdValue = new StreamId(StreamId);
-
-			IStream? Stream = await StreamService.GetStreamAsync(StreamIdValue);
+			IStream? Stream = await StreamService.GetStreamAsync(StreamId);
 			if (Stream == null)
 			{
-				return NotFound();
+				return NotFound(StreamId);
 			}
 			if (!await StreamService.AuthorizeAsync(Stream, AclAction.ViewChanges, User, null))
 			{
-				return Forbid();
+				return Forbid(AclAction.ViewChanges, StreamId);
 			}
 
 			string? PerforceUser = User.GetPerforceUser();
 			if(PerforceUser == null)
 			{
-				return Forbid();
+				return BadRequest("Current user does not have an associated Perforce user");
 			}
 
 			List<ChangeSummary> Commits = await PerforceService.GetChangesAsync(Stream.ClusterName, Stream.Name, Min, Max, Results, PerforceUser);
@@ -182,36 +177,34 @@ namespace HordeServer.Controllers
 		/// Gets a list of changes for a stream
 		/// </summary>
 		/// <param name="StreamId">The stream id</param>
-		/// <param name="Number">The changelist number</param>
+		/// <param name="ChangeNumber">The changelist number</param>
 		/// <param name="Filter">The filter to apply to the results</param>
 		/// <returns>Http result code</returns>
 		[HttpGet]
-		[Route("/api/v1/streams/{StreamId}/changes/{Number}")]
+		[Route("/api/v1/streams/{StreamId}/changes/{ChangeNumber}")]
 		[ProducesResponseType(typeof(GetChangeDetailsResponse), 200)]
-		public async Task<ActionResult<object>> GetChangeDetailsAsync(string StreamId, int Number, PropertyFilter? Filter = null)
+		public async Task<ActionResult<object>> GetChangeDetailsAsync(StreamId StreamId, int ChangeNumber, PropertyFilter? Filter = null)
 		{
-			StreamId StreamIdValue = new StreamId(StreamId);
-
-			IStream? Stream = await StreamService.GetStreamAsync(StreamIdValue);
+			IStream? Stream = await StreamService.GetStreamAsync(StreamId);
 			if (Stream == null)
 			{
-				return NotFound();
+				return NotFound(StreamId);
 			}
 			if (!await StreamService.AuthorizeAsync(Stream, AclAction.ViewChanges, User, null))
 			{
-				return Forbid();
+				return Forbid(AclAction.ViewChanges, StreamId);
 			}
 
 			string? PerforceUser = User.GetPerforceUser();
 			if(PerforceUser == null)
 			{
-				return Forbid();
+				return BadRequest("Current user does not have an associated Perforce user");
 			}
 
-			ChangeDetails? ChangeDetails = await PerforceService.GetChangeDetailsAsync(Stream.ClusterName, Stream.Name, Number, PerforceUser);
+			ChangeDetails? ChangeDetails = await PerforceService.GetChangeDetailsAsync(Stream.ClusterName, Stream.Name, ChangeNumber, PerforceUser);
 			if(ChangeDetails == null)
 			{
-				return NotFound();
+				return NotFound("CL {Change} not found in stream {StreamId}", ChangeNumber, StreamId);
 			}
 
 			return PropertyFilter.Apply(new GetChangeDetailsResponse(ChangeDetails), Filter);
@@ -230,23 +223,21 @@ namespace HordeServer.Controllers
 		[HttpGet]
 		[Route("/api/v1/streams/{StreamId}/history")]
 		[ProducesResponseType(typeof(List<GetJobStepRefResponse>), 200)]
-		public async Task<ActionResult<List<object>>> GetStepHistoryAsync(string StreamId, [FromQuery] string TemplateId, [FromQuery] string Step, [FromQuery] int? Change = null, [FromQuery] int Count = 10, [FromQuery] PropertyFilter? Filter = null)
+		public async Task<ActionResult<List<object>>> GetStepHistoryAsync(StreamId StreamId, [FromQuery] string TemplateId, [FromQuery] string Step, [FromQuery] int? Change = null, [FromQuery] int Count = 10, [FromQuery] PropertyFilter? Filter = null)
 		{
-			StreamId StreamIdValue = new StreamId(StreamId);
-
-			IStream? Stream = await StreamService.GetStreamAsync(StreamIdValue);
+			IStream? Stream = await StreamService.GetStreamAsync(StreamId);
 			if (Stream == null)
 			{
-				return NotFound();
+				return NotFound(StreamId);
 			}
 			if (!await StreamService.AuthorizeAsync(Stream, AclAction.ViewJob, User, null))
 			{
-				return Forbid();
+				return Forbid(AclAction.ViewJob, StreamId);
 			}
 
 			TemplateRefId TemplateIdValue = new TemplateRefId(TemplateId);
 
-			List<IJobStepRef> Steps = await JobStepRefCollection.GetStepsForNodeAsync(StreamIdValue, TemplateIdValue, Step, Change, true, Count);
+			List<IJobStepRef> Steps = await JobStepRefCollection.GetStepsForNodeAsync(StreamId, TemplateIdValue, Step, Change, true, Count);
 			return Steps.ConvertAll(x => PropertyFilter.Apply(new GetJobStepRefResponse(x), Filter));
 		}
 
@@ -257,21 +248,19 @@ namespace HordeServer.Controllers
 		/// <returns>Http result code</returns>
 		[HttpDelete]
 		[Route("/api/v1/streams/{StreamId}")]
-		public async Task<ActionResult> DeleteStreamAsync(string StreamId)
+		public async Task<ActionResult> DeleteStreamAsync(StreamId StreamId)
 		{
-			StreamId StreamIdValue = new StreamId(StreamId);
-
-			IStream? Stream = await StreamService.GetStreamAsync(StreamIdValue);
+			IStream? Stream = await StreamService.GetStreamAsync(StreamId);
 			if (Stream == null)
 			{
-				return NotFound();
+				return NotFound(StreamId);
 			}
 			if (!await StreamService.AuthorizeAsync(Stream, AclAction.DeleteStream, User, null))
 			{
-				return Forbid();
+				return Forbid(AclAction.DeleteStream, StreamId);
 			}
 
-			await StreamService.DeleteStreamAsync(StreamIdValue);
+			await StreamService.DeleteStreamAsync(StreamId);
 			return new OkResult();
 		}
 
@@ -368,25 +357,22 @@ namespace HordeServer.Controllers
 		[HttpGet]
 		[Route("/api/v1/streams/{StreamId}/templates/{TemplateRefId}")]
 		[ProducesResponseType(typeof(List<GetTemplateResponse>), 200)]
-		public async Task<ActionResult<object>> GetTemplateAsync(string StreamId, string TemplateRefId, [FromQuery] PropertyFilter? Filter = null)
+		public async Task<ActionResult<object>> GetTemplateAsync(StreamId StreamId, TemplateRefId TemplateRefId, [FromQuery] PropertyFilter? Filter = null)
 		{
-			StreamId StreamIdValue = new StreamId(StreamId);
-			TemplateRefId TemplateRefIdValue = new TemplateRefId(TemplateRefId);
-
-			IStream? Stream = await StreamService.GetStreamAsync(StreamIdValue);
+			IStream? Stream = await StreamService.GetStreamAsync(StreamId);
 			if (Stream == null)
 			{
-				return NotFound();
+				return NotFound(StreamId);
 			}
 
 			TemplateRef? TemplateRef;
-			if (!Stream.Templates.TryGetValue(TemplateRefIdValue, out TemplateRef))
+			if (!Stream.Templates.TryGetValue(TemplateRefId, out TemplateRef))
 			{
-				return NotFound();
+				return NotFound(StreamId, TemplateRefId);
 			}
 			if (!await StreamService.AuthorizeAsync(Stream, TemplateRef, AclAction.ViewTemplate, User, null))
 			{
-				return Forbid();
+				return Forbid(AclAction.ViewTemplate, StreamId);
 			}
 
 			ITemplate? Template = await TemplateCollection.GetAsync(TemplateRef.Hash);
diff --git a/Engine/Source/Programs/Horde/Horde.Build/Controllers/UsersController.cs b/Engine/Source/Programs/Horde/Horde.Build/Controllers/UsersController.cs
--- a/Engine/Source/Programs/Horde/Horde.Build/Controllers/UsersController.cs
+++ b/Engine/Source/Programs/Horde/Horde.Build/Controllers/UsersController.cs
@@ -25,7 +25,7 @@ namespace HordeServer.Controllers
 	[ApiController]
 	[Authorize]
 	[Route("[controller]")]
-	public class UsersController : ControllerBase
+	public class UsersController : HordeControllerBase
 	{
 		/// <summary>
 		/// The user collection instance
@@ -57,10 +57,16 @@ namespace HordeServer.Controllers
 		[ProducesResponseType(typeof(List<GetUserResponse>), 200)]
 		public async Task<ActionResult<object>> GetUserAsync(string Id, [FromQuery] PropertyFilter? Filter = null)
 		{
-			IUser? User = await GetUserInternalAsync(Id);
+			UserId? UserId = ParseUserId(Id);
+			if(UserId == null)
+			{
+				return BadRequest("Invalid user id '{Id}'", Id);
+			}
+
+			IUser? User = await UserCollection.GetUserAsync(UserId.Value);
 			if (User == null)
 			{
-				return NotFound();
+				return NotFound(UserId.Value);
 			}
 
 			IAvatar? Avatar = (AvatarService == null) ? (IAvatar?)null : await AvatarService.GetAvatarAsync(User);
@@ -104,17 +110,6 @@ namespace HordeServer.Controllers
 			return Response;
 		}
 
-
-		async Task<IUser?> GetUserInternalAsync(string Id)
-		{
-			UserId? UserId = ParseUserId(Id);
-			if(UserId == null)
-			{
-				return null;
-			}
-			return await UserCollection.GetUserAsync(UserId.Value);
-		}
-
 		UserId? ParseUserId(string Id)
 		{
 			if (Id.Equals("current", StringComparison.OrdinalIgnoreCase))
diff --git a/Engine/Source/Programs/Horde/Horde.Build/Jobs/IJob.cs b/Engine/Source/Programs/Horde/Horde.Build/Jobs/IJob.cs
--- a/Engine/Source/Programs/Horde/Horde.Build/Jobs/IJob.cs
+++ b/Engine/Source/Programs/Horde/Horde.Build/Jobs/IJob.cs
@@ -23,6 +23,7 @@ namespace HordeServer.Models
 	using JobId = ObjectId<IJob>;
 	using LeaseId = ObjectId<ILease>;
 	using LogId = ObjectId<ILogFile>;
+	using SessionId = ObjectId<ISession>;
 	using StreamId = StringId<IStream>;
 	using TemplateRefId = StringId<TemplateRef>;
 	using PoolId = StringId<IPool>;
@@ -217,7 +218,7 @@ namespace HordeServer.Models
 		/// <summary>
 		/// The agent session that is executing this group
 		/// </summary>
-		public ObjectId? SessionId { get; }
+		public SessionId? SessionId { get; }
 
 		/// <summary>
 		/// The lease that's executing this group
diff --git a/Engine/Source/Programs/Horde/Horde.Build/Jobs/IJobCollection.cs b/Engine/Source/Programs/Horde/Horde.Build/Jobs/IJobCollection.cs
--- a/Engine/Source/Programs/Horde/Horde.Build/Jobs/IJobCollection.cs
+++ b/Engine/Source/Programs/Horde/Horde.Build/Jobs/IJobCollection.cs
@@ -23,6 +23,7 @@ namespace HordeServer.Collections
 	using LeaseId = ObjectId<ILease>;
 	using LogId = ObjectId<ILogFile>;
 	using PoolId = StringId<IPool>;
+	using SessionId = ObjectId<ISession>;
 	using StreamId = StringId<IStream>;
 	using TemplateRefId = StringId<TemplateRef>;
 	using UserId = ObjectId<IUser>;
@@ -240,7 +241,7 @@ namespace HordeServer.Collections
 		/// <param name="LeaseId">The lease unique id</param>
 		/// <param name="LogId">Unique id of the log for the batch</param>
 		/// <returns>True if the batch is updated</returns>
-		Task<IJob?> TryAssignLeaseAsync(IJob Job, int BatchIdx, PoolId PoolId, AgentId AgentId, ObjectId SessionId, LeaseId LeaseId, LogId LogId);
+		Task<IJob?> TryAssignLeaseAsync(IJob Job, int BatchIdx, PoolId PoolId, AgentId AgentId, SessionId SessionId, LeaseId LeaseId, LogId LogId);
 
 		/// <summary>
 		/// Cancel a lease reservation on a batch (before it has started)
diff --git a/Engine/Source/Programs/Horde/Horde.Build/Jobs/Impl/JobCollection.cs b/Engine/Source/Programs/Horde/Horde.Build/Jobs/Impl/JobCollection.cs
--- a/Engine/Source/Programs/Horde/Horde.Build/Jobs/Impl/JobCollection.cs
+++ b/Engine/Source/Programs/Horde/Horde.Build/Jobs/Impl/JobCollection.cs
@@ -26,6 +26,7 @@ namespace HordeServer.Collections.Impl
 	using JobId = ObjectId<IJob>;
 	using LeaseId = ObjectId<ILease>;
 	using LogId = ObjectId<ILogFile>;
+	using SessionId = ObjectId<ISession>;
 	using StreamId = StringId<IStream>;
 	using TemplateRefId = StringId<TemplateRef>;
 	using PoolId = StringId<IPool>;
@@ -134,7 +135,7 @@ namespace HordeServer.Collections.Impl
 			public AgentId? AgentId { get; set; }
 
 			[BsonIgnoreIfNull]
-			public ObjectId? SessionId { get; set; }
+			public SessionId? SessionId { get; set; }
 
 			[BsonIgnoreIfNull]
 			public LeaseId? LeaseId { get; set; }
@@ -1176,7 +1177,7 @@ namespace HordeServer.Collections.Impl
 		}
 
 		/// <inheritdoc/>
-		public async Task<IJob?> TryAssignLeaseAsync(IJob Job, int BatchIdx, PoolId PoolId, AgentId AgentId, ObjectId SessionId, LeaseId LeaseId, LogId LogId)
+		public async Task<IJob?> TryAssignLeaseAsync(IJob Job, int BatchIdx, PoolId PoolId, AgentId AgentId, SessionId SessionId, LeaseId LeaseId, LogId LogId)
 		{
 			// Try to update the job with this agent id
 			UpdateDefinitionBuilder<JobDocument> UpdateBuilder = Builders<JobDocument>.Update;
diff --git a/Engine/Source/Programs/Horde/Horde.Build/Jobs/Impl/JobService.cs b/Engine/Source/Programs/Horde/Horde.Build/Jobs/Impl/JobService.cs
--- a/Engine/Source/Programs/Horde/Horde.Build/Jobs/Impl/JobService.cs
+++ b/Engine/Source/Programs/Horde/Horde.Build/Jobs/Impl/JobService.cs
@@ -7,7 +7,6 @@ using HordeCommon;
 using HordeServer.Models;
 using HordeCommon.Rpc;
 using HordeServer.Utilities;
-using Microsoft.AspNetCore.Http;
 using Microsoft.Extensions.Logging;
 using Microsoft.Extensions.Options;
 using MongoDB.Bson;
@@ -34,6 +33,7 @@ namespace HordeServer.Services
 	using JobId = ObjectId<IJob>;
 	using LeaseId = ObjectId<ILease>;
 	using LogId = ObjectId<ILogFile>;
+	using SessionId = ObjectId<ISession>;
 	using StreamId = StringId<IStream>;
 	using TemplateRefId = StringId<TemplateRef>;
 	using UserId = ObjectId<IUser>;
@@ -1353,8 +1353,8 @@ namespace HordeServer.Services
 				{
 					if(Claim.Type == HordeClaimTypes.AgentSessionId)
 					{
-						ObjectId SessionId;
-						if (ObjectId.TryParse(Claim.Value, out SessionId) && SessionId == Batch.SessionId.Value)
+						SessionId SessionIdValue;
+						if (SessionId.TryParse(Claim.Value, out SessionIdValue) && SessionIdValue == Batch.SessionId.Value)
 						{
 							return true;
 						}
diff --git a/Engine/Source/Programs/Horde/Horde.Build/Jobs/SubResourceId.cs b/Engine/Source/Programs/Horde/Horde.Build/Jobs/SubResourceId.cs
--- a/Engine/Source/Programs/Horde/Horde.Build/Jobs/SubResourceId.cs
+++ b/Engine/Source/Programs/Horde/Horde.Build/Jobs/SubResourceId.cs
@@ -4,6 +4,7 @@ using MongoDB.Bson.Serialization;
 using MongoDB.Bson.Serialization.Attributes;
 using System;
 using System.Collections.Generic;
+using System.ComponentModel;
 using System.Diagnostics;
 using System.Globalization;
 using System.Linq;
@@ -17,6 +18,7 @@ namespace HordeServer.Models
 	/// Subresource identifiers are stored as 16-bit integers formatted as a 4-digit hex code, in order to keep URLs short. Calling Next() will generate a new
 	/// identifier with more entropy than just incrementing the value but an identical period before repeating, in order to make URL fragments more distinctive.
 	/// </summary>
+	[TypeConverter(typeof(SubResourceIdTypeConverter))]
 	[BsonSerializer(typeof(SubResourceIdSerializer))]
 	public struct SubResourceId : IEquatable<SubResourceId>
 	{
@@ -163,4 +165,41 @@ namespace HordeServer.Models
 			return new SubResourceId((ushort)Context.Reader.ReadInt32());
 		}
 	}
+
+	/// <summary>
+	/// Type converter from strings to SubResourceId objects
+	/// </summary>
+	sealed class SubResourceIdTypeConverter : TypeConverter
+	{
+		/// <inheritdoc/>
+		public override bool CanConvertFrom(ITypeDescriptorContext? Context, Type SourceType)
+		{
+			return SourceType == typeof(string);
+		}
+
+		/// <inheritdoc/>
+		public override object ConvertFrom(ITypeDescriptorContext? Context, CultureInfo? Culture, object Value)
+		{
+			return SubResourceId.Parse((string)Value);
+		}
+
+		/// <inheritdoc/>
+		public override bool CanConvertTo(ITypeDescriptorContext? Context, Type? DestinationType)
+		{
+			return DestinationType == typeof(string);
+		}
+
+		/// <inheritdoc/>
+		public override object? ConvertTo(ITypeDescriptorContext? Context, CultureInfo? Culture, object? Value, Type DestinationType)
+		{
+			if (DestinationType == typeof(string))
+			{
+				return Value?.ToString();
+			}
+			else
+			{
+				return null;
+			}
+		}
+	}
 }
diff --git a/Engine/Source/Programs/Horde/Horde.Build/Logs/ILogFile.cs b/Engine/Source/Programs/Horde/Horde.Build/Logs/ILogFile.cs
--- a/Engine/Source/Programs/Horde/Horde.Build/Logs/ILogFile.cs
+++ b/Engine/Source/Programs/Horde/Horde.Build/Logs/ILogFile.cs
@@ -16,6 +16,7 @@ namespace HordeServer.Models
 {
 	using JobId = ObjectId<IJob>;
 	using LogId = ObjectId<ILogFile>;
+	using SessionId = ObjectId<ISession>;
 
 	/// <summary>
 	/// Information about a log file chunk
@@ -61,7 +62,7 @@ namespace HordeServer.Models
 		/// <summary>
 		/// The session allowed to write to this log
 		/// </summary>
-		public ObjectId? SessionId { get; }
+		public SessionId? SessionId { get; }
 
 		/// <summary>
 		/// Maximum line index in the file
diff --git a/Engine/Source/Programs/Horde/Horde.Build/Logs/ILogFileCollection.cs b/Engine/Source/Programs/Horde/Horde.Build/Logs/ILogFileCollection.cs
--- a/Engine/Source/Programs/Horde/Horde.Build/Logs/ILogFileCollection.cs
+++ b/Engine/Source/Programs/Horde/Horde.Build/Logs/ILogFileCollection.cs
@@ -17,6 +17,7 @@ namespace HordeServer.Collections
 {
 	using JobId = ObjectId<IJob>;
 	using LogId = ObjectId<ILogFile>;
+	using SessionId = ObjectId<ISession>;
 
 	/// <summary>
 	/// Updates a log file chunk
@@ -64,7 +65,7 @@ namespace HordeServer.Collections
 		/// <param name="SessionId">Agent session allowed to update the log</param>
 		/// <param name="Type">Type of events to be stored in the log</param>
 		/// <returns>The new log file document</returns>
-		Task<ILogFile> CreateLogFileAsync(JobId JobId, ObjectId? SessionId, LogType Type);
+		Task<ILogFile> CreateLogFileAsync(JobId JobId, SessionId? SessionId, LogType Type);
 
 		/// <summary>
 		/// Adds a new chunk
diff --git a/Engine/Source/Programs/Horde/Horde.Build/Logs/Impl/LogFileCollection.cs b/Engine/Source/Programs/Horde/Horde.Build/Logs/Impl/LogFileCollection.cs
--- a/Engine/Source/Programs/Horde/Horde.Build/Logs/Impl/LogFileCollection.cs
+++ b/Engine/Source/Programs/Horde/Horde.Build/Logs/Impl/LogFileCollection.cs
@@ -17,6 +17,7 @@ namespace HordeServer.Collections.Impl
 {
 	using JobId = ObjectId<IJob>;
 	using LogId = ObjectId<ILogFile>;
+	using SessionId = ObjectId<ISession>;
 
 	/// <summary>
 	/// Wrapper around the jobs collection in a mongo DB
@@ -59,7 +60,7 @@ namespace HordeServer.Collections.Impl
 			[BsonRequired]
 			public JobId JobId { get; set; }
 
-			public ObjectId? SessionId { get; set; }
+			public SessionId? SessionId { get; set; }
 			public LogType Type { get; set; }
 
 			[BsonIgnoreIfNull]
@@ -80,7 +81,7 @@ namespace HordeServer.Collections.Impl
 			{
 			}
 
-			public LogFileDocument(JobId JobId, ObjectId? SessionId, LogType Type)
+			public LogFileDocument(JobId JobId, SessionId? SessionId, LogType Type)
 			{
 				this.Id = LogId.GenerateNewId();
 				this.JobId = JobId;
@@ -118,7 +119,7 @@ namespace HordeServer.Collections.Impl
 		}
 
 		/// <inheritdoc/>
-		public async Task<ILogFile> CreateLogFileAsync(JobId JobId, ObjectId? SessionId, LogType Type)
+		public async Task<ILogFile> CreateLogFileAsync(JobId JobId, SessionId? SessionId, LogType Type)
 		{
 			LogFileDocument NewLogFile = new LogFileDocument(JobId, SessionId, Type);
 			await LogFiles.InsertOneAsync(NewLogFile);
diff --git a/Engine/Source/Programs/Horde/Horde.Build/Logs/Impl/LogFileService.cs b/Engine/Source/Programs/Horde/Horde.Build/Logs/Impl/LogFileService.cs
--- a/Engine/Source/Programs/Horde/Horde.Build/Logs/Impl/LogFileService.cs
+++ b/Engine/Source/Programs/Horde/Horde.Build/Logs/Impl/LogFileService.cs
@@ -37,6 +37,7 @@ namespace HordeServer.Services
 {
 	using JobId = ObjectId<IJob>;
 	using LogId = ObjectId<ILogFile>;
+	using SessionId = ObjectId<ISession>;
 
 	/// <summary>
 	/// Metadata about a log file
@@ -66,7 +67,7 @@ namespace HordeServer.Services
 		/// <param name="SessionId">Agent session allowed to update the log</param>
 		/// <param name="Type">Type of events to be stored in the log</param>
 		/// <returns>The new log file document</returns>
-		Task<ILogFile> CreateLogFileAsync(JobId JobId, ObjectId? SessionId, LogType Type);
+		Task<ILogFile> CreateLogFileAsync(JobId JobId, SessionId? SessionId, LogType Type);
 
 		/// <summary>
 		/// Gets a logfile by ID
@@ -513,7 +514,7 @@ namespace HordeServer.Services
 		}
 
 		/// <inheritdoc/>
-		public Task<ILogFile> CreateLogFileAsync(JobId JobId, ObjectId? SessionId, LogType Type)
+		public Task<ILogFile> CreateLogFileAsync(JobId JobId, SessionId? SessionId, LogType Type)
 		{
 			return LogFiles.CreateLogFileAsync(JobId, SessionId, Type);
 		}
diff --git a/Engine/Source/Programs/Horde/Horde.Build/Models/Globals.cs b/Engine/Source/Programs/Horde/Horde.Build/Models/Globals.cs
--- a/Engine/Source/Programs/Horde/Horde.Build/Models/Globals.cs
+++ b/Engine/Source/Programs/Horde/Horde.Build/Models/Globals.cs
@@ -310,11 +310,6 @@ namespace HordeServer.Models
 		/// </summary>
 		public ObjectId InstanceId { get; set; }
 
-		/// <summary>
-		/// Set to true to drop the database on restart.
-		/// </summary>
-		public bool ForceReset { get; set; }
-
 		/// <summary>
 		/// The config revision
 		/// </summary>
diff --git a/Engine/Source/Programs/Horde/Horde.Build/Server/DatabaseService.cs b/Engine/Source/Programs/Horde/Horde.Build/Server/DatabaseService.cs
--- a/Engine/Source/Programs/Horde/Horde.Build/Server/DatabaseService.cs
+++ b/Engine/Source/Programs/Horde/Horde.Build/Server/DatabaseService.cs
@@ -185,28 +185,6 @@ namespace HordeServer.Services
 				Singletons = GetCollection<BsonDocument>("Singletons");
 
 				Globals Globals = GetGlobalsAsync().Result;
-				if (Globals.ForceReset)
-				{
-					Logger.LogInformation("Performing DB reset...");
-
-					// Drop all the collections
-					List<string> CollectionNames = Database.ListCollectionNames().ToList();
-					foreach (string CollectionName in CollectionNames)
-					{
-						Logger.LogInformation("Dropping {CollectionName}", CollectionName);
-						Database.DropCollection(CollectionName);
-					}
-
-					// Create the new singletons collection
-					Singletons = GetCollection<BsonDocument>("Singletons");
-
-					// Create the new globals object
-					Globals NewGlobals = new Globals();
-					NewGlobals.JwtSigningKey = Globals.JwtSigningKey; // Ensure agents can still register
-					Singletons.InsertOne(NewGlobals.ToBsonDocument());
-					Globals = NewGlobals;
-				}
-
 				while (Globals.JwtSigningKey == null)
 				{
 					Globals.RotateSigningKey();
diff --git a/Engine/Source/Programs/Horde/Horde.Build/Services/ConsistencyService.cs b/Engine/Source/Programs/Horde/Horde.Build/Services/ConsistencyService.cs
--- a/Engine/Source/Programs/Horde/Horde.Build/Services/ConsistencyService.cs
+++ b/Engine/Source/Programs/Horde/Horde.Build/Services/ConsistencyService.cs
@@ -14,6 +14,8 @@ using System.Threading.Tasks;
 
 namespace HordeServer.Services
 {
+	using SessionId = ObjectId<ISession>;
+
 	/// <summary>
 	/// Service which checks the database for consistency and fixes up any errors
 	/// </summary>
@@ -48,7 +50,7 @@ namespace HordeServer.Services
 			DateTime NextTickTime = DateTime.UtcNow + TimeSpan.FromMinutes(20.0);
 
 			List<ISession> Sessions = await SessionCollection.FindActiveSessionsAsync();
-			Dictionary<ObjectId, ISession> SessionIdToInstance = Sessions.ToDictionary(x => x.Id, x => x);
+			Dictionary<SessionId, ISession> SessionIdToInstance = Sessions.ToDictionary(x => x.Id, x => x);
 
 			// Find any leases that are still running when their session has terminated
 			List<ILease> Leases = await LeaseCollection.FindActiveLeasesAsync();
diff --git a/Engine/Source/Programs/Horde/Horde.Build/Services/RpcService.cs b/Engine/Source/Programs/Horde/Horde.Build/Services/RpcService.cs
--- a/Engine/Source/Programs/Horde/Horde.Build/Services/RpcService.cs
+++ b/Engine/Source/Programs/Horde/Horde.Build/Services/RpcService.cs
@@ -32,6 +32,7 @@ namespace HordeServer.Services
 	using IStream = HordeServer.Models.IStream;
 	using JobId = ObjectId<IJob>;
 	using LogId = ObjectId<ILogFile>;
+	using SessionId = ObjectId<ISession>;
 	using StreamId = StringId<IStream>;
 	using RpcAgentCapabilities = HordeCommon.Rpc.Messages.AgentCapabilities;
 	using RpcDeviceCapabilities = HordeCommon.Rpc.Messages.DeviceCapabilities;
@@ -417,7 +418,7 @@ namespace HordeServer.Services
 					}
 
 					// Update the session
-					Agent = await AgentService.UpdateSessionWithWaitAsync(Agent, Request.SessionId.ToObjectId(), Request.Status, Properties, Resources, Request.Leases, CancellationSource.Token);
+					Agent = await AgentService.UpdateSessionWithWaitAsync(Agent, SessionId.Parse(Request.SessionId), Request.Status, Properties, Resources, Request.Leases, CancellationSource.Token);
 				}
 
 				// Handle the invalid agent case
@@ -756,7 +757,7 @@ namespace HordeServer.Services
 			ClaimsPrincipal Principal = Context.GetHttpContext().User;
 			if (!Principal.HasSessionClaim(Batch.SessionId.Value))
 			{
-				throw new StructuredRpcException(StatusCode.PermissionDenied, "Session id {SessionId} not valid for batch {JobId}:{BatchId}. Expected {ExpectedSessionId}.", Principal.GetSessionClaim() ?? ObjectId.Empty, Job.Id, BatchId, Batch.SessionId.Value);
+				throw new StructuredRpcException(StatusCode.PermissionDenied, "Session id {SessionId} not valid for batch {JobId}:{BatchId}. Expected {ExpectedSessionId}.", Principal.GetSessionClaim() ?? SessionId.Empty, Job.Id, BatchId, Batch.SessionId.Value);
 			}
 
 			return Batch;
diff --git a/Engine/Source/Programs/Horde/Horde.Build/Utilities/HordeControllerBase.cs b/Engine/Source/Programs/Horde/Horde.Build/Utilities/HordeControllerBase.cs
new file mode 100644
--- /dev/null
+++ b/Engine/Source/Programs/Horde/Horde.Build/Utilities/HordeControllerBase.cs
@@ -0,0 +1,214 @@
+// Copyright Epic Games, Inc. All Rights Reserved.
+
+using EpicGames.Core;
+using HordeServer.Models;
+using Microsoft.AspNetCore.Http;
+using Microsoft.AspNetCore.Mvc;
+using Microsoft.Extensions.Logging;
+using System;
+using System.Collections.Generic;
+using System.Diagnostics;
+using System.Linq;
+using System.Linq.Expressions;
+using System.Threading;
+using System.Threading.Tasks;
+
+namespace HordeServer.Utilities
+{
+	using JobId = ObjectId<IJob>;
+	using LeaseId = ObjectId<ILease>;
+	using ProjectId = StringId<IProject>;
+	using StreamId = StringId<IStream>;
+	using TemplateRefId = StringId<TemplateRef>;
+
+	/// <summary>
+	/// Base class for Horde controllers
+	/// </summary>
+	public abstract class HordeControllerBase : ControllerBase
+	{
+		/// <summary>
+		/// Create a response indicating a problem with the request
+		/// </summary>
+		/// <param name="Message">Standard structured-logging format string</param>
+		/// <param name="Args">Arguments for the format string</param>
+		[NonAction]
+		protected ActionResult BadRequest(string Message, params object[] Args)
+		{
+			return BadRequest(LogEvent.Create(LogLevel.Error, Message, Args));
+		}
+
+		/// <summary>
+		/// Create a response indicating a problem with the request
+		/// </summary>
+		/// <param name="EventId">Id identifying the error</param>
+		/// <param name="Message">Standard structured-logging format string</param>
+		/// <param name="Args">Arguments for the format string</param>
+		[NonAction]
+		protected ActionResult BadRequest(EventId EventId, string Message, params object[] Args)
+		{
+			return BadRequest(LogEvent.Create(LogLevel.Error, EventId, Message, Args));
+		}
+
+		/// <summary>
+		/// Returns a 403 (forbidden) response with the given action
+		/// </summary>
+		[NonAction]
+		protected ActionResult Forbid(AclAction Action)
+		{
+			return Forbid("User does not have {Action} permission", Action);
+		}
+
+		/// <summary>
+		/// Returns a 403 (forbidden) response with the given action and object
+		/// </summary>
+		[NonAction]
+		protected ActionResult Forbid(AclAction Action, AgentId AgentId)
+		{
+			return Forbid(Action, "agent {AgentId}", AgentId);
+		}
+
+		/// <summary>
+		/// Returns a 403 (forbidden) response with the given action and object
+		/// </summary>
+		[NonAction]
+		protected ActionResult Forbid(AclAction Action, JobId JobId)
+		{
+			return Forbid(Action, "job {JobId}", JobId);
+		}
+
+		/// <summary>
+		/// Returns a 403 (forbidden) response with the given action and object
+		/// </summary>
+		[NonAction]
+		protected ActionResult Forbid(AclAction Action, ProjectId ProjectId)
+		{
+			return Forbid(Action, "project {ProjectId}", ProjectId);
+		}
+
+		/// <summary>
+		/// Returns a 403 (forbidden) response with the given action and object
+		/// </summary>
+		[NonAction]
+		protected ActionResult Forbid(AclAction Action, StreamId StreamId)
+		{
+			return Forbid(Action, "stream {StreamId}", StreamId);
+		}
+
+		/// <summary>
+		/// Returns a 403 (forbidden) response with the given action and object
+		/// </summary>
+		[NonAction]
+		protected ActionResult Forbid(AclAction Action, TemplateRefId TemplateId)
+		{
+			return Forbid(Action, "template {TemplateId}", TemplateId);
+		}
+
+		/// <summary>
+		/// Returns a 403 (forbidden) response with the given action and object
+		/// </summary>
+		[NonAction]
+		private ActionResult Forbid(AclAction Action, string ObjectMessage, object Object)
+		{
+			return Forbid($"User does not have {{Action}} permission for {ObjectMessage}", Action, Object);
+		}
+
+		/// <summary>
+		/// Returns a 403 response with the given log event
+		/// </summary>
+		[NonAction]
+		protected ActionResult Forbid(string Message, params object[] Args)
+		{
+			return StatusCode(StatusCodes.Status403Forbidden, LogEvent.Create(LogLevel.Error, Message, Args));
+		}
+
+		/// <summary>
+		/// Returns a 404 response for the given object
+		/// </summary>
+		[NonAction]
+		protected ActionResult NotFound(AgentId AgentId)
+		{
+			return NotFound("Agent {AgentId} not found", AgentId);
+		}
+
+		/// <summary>
+		/// Returns a 404 response for the given object
+		/// </summary>
+		[NonAction]
+		protected ActionResult NotFound(AgentId AgentId, LeaseId LeaseId)
+		{
+			return NotFound("Lease {LeaseId} not found for agent {AgentId}", LeaseId, AgentId);
+		}
+
+		/// <summary>
+		/// Returns a 404 response for the given object
+		/// </summary>
+		[NonAction]
+		protected ActionResult NotFound(JobId JobId)
+		{
+			return NotFound("Job {JobId} not found", JobId);
+		}
+
+		/// <summary>
+		/// Returns a 404 response for the given object
+		/// </summary>
+		[NonAction]
+		protected ActionResult NotFound(JobId JobId, int GroupIdx)
+		{
+			return NotFound("Group {GroupIdx} on job {JobId} not found", GroupIdx, JobId);
+		}
+
+		/// <summary>
+		/// Returns a 404 response for the given object
+		/// </summary>
+		[NonAction]
+		protected ActionResult NotFound(JobId JobId, int GroupIdx, int NodeIdx)
+		{
+			return NotFound("Node {NodeIdx} not found on job {JobId} group {GroupIdx}", NodeIdx, JobId, GroupIdx);
+		}
+
+		/// <summary>
+		/// Returns a 404 response for the given object
+		/// </summary>
+		[NonAction]
+		protected ActionResult NotFound(JobId JobId, SubResourceId BatchId)
+		{
+			return NotFound("Batch {BatchId} not found on job {JobId}", BatchId, JobId);
+		}
+
+		/// <summary>
+		/// Returns a 404 response for the given object
+		/// </summary>
+		[NonAction]
+		protected ActionResult NotFound(JobId JobId, SubResourceId BatchId, SubResourceId StepId)
+		{
+			return NotFound("Step {StepId} not found on job {JobId} batch {BatchId}", StepId, JobId, BatchId);
+		}
+
+		/// <summary>
+		/// Returns a 404 response for the given object
+		/// </summary>
+		[NonAction]
+		protected ActionResult NotFound(StreamId StreamId)
+		{
+			return NotFound("Stream {StreamId} not found", StreamId);
+		}
+
+		/// <summary>
+		/// Returns a 404 response for the given object
+		/// </summary>
+		[NonAction]
+		protected ActionResult NotFound(StreamId StreamId, TemplateRefId TemplateId)
+		{
+			return NotFound("Template {TemplateId} not found on stream {StreamId}", TemplateId, StreamId);
+		}
+
+		/// <summary>
+		/// Returns a 404 response for the given object
+		/// </summary>
+		[NonAction]
+		protected ActionResult NotFound(string Message, params object[] Args)
+		{
+			return NotFound(LogEvent.Create(LogLevel.Error, Message, Args));
+		}
+	}
+}
diff --git a/Engine/Source/Programs/Shared/EpicGames.Core/LogEvent.cs b/Engine/Source/Programs/Shared/EpicGames.Core/LogEvent.cs
--- a/Engine/Source/Programs/Shared/EpicGames.Core/LogEvent.cs
+++ b/Engine/Source/Programs/Shared/EpicGames.Core/LogEvent.cs
@@ -10,6 +10,7 @@ using System.Globalization;
 using System.Linq;
 using System.Text;
 using System.Text.Json;
+using System.Text.Json.Serialization;
 using System.Text.RegularExpressions;
 
 namespace EpicGames.Core
@@ -17,6 +18,7 @@ namespace EpicGames.Core
 	/// <summary>
 	/// Epic representation of a log event. Can be serialized to/from Json for the Horde dashboard, and passed directly through ILogger interfaces.
 	/// </summary>
+	[JsonConverter(typeof(LogEventConverter))]
 	public class LogEvent : IEnumerable<KeyValuePair<string, object?>>
 	{
 		static class InternalPropertyNames
@@ -416,6 +418,21 @@ namespace EpicGames.Core
 			}
 		}
 
+		public static LogEvent Create(LogLevel Level, string Format, params object[] Args)
+			=> Create(Level, KnownLogEvents.None, null, Format, Args);
+
+		public static LogEvent Create(LogLevel Level, EventId EventId, string Format, params object[] Args)
+			=> Create(Level, EventId, null, Format, Args);
+
+		public static LogEvent Create(LogLevel Level, EventId EventId, Exception? Exception, string Format, params object[] Args)
+		{
+			Dictionary<string, object> Properties = new Dictionary<string, object>();
+			MessageTemplate.ParsePropertyValues(Format, Args, Properties);
+
+			string Message = MessageTemplate.Render(Format, Properties!);
+			return new LogEvent(DateTime.UtcNow, Level, EventId, Message, Format, Properties, LogException.FromException(Exception));
+		}
+
 		/// <summary>
 		/// Creates a log event from an ILogger parameters
 		/// </summary>
@@ -630,4 +647,22 @@ namespace EpicGames.Core
 			return Result;
 		}
 	}
+
+	/// <summary>
+	/// Converter for serialization of <see cref="LogEvent"/> instances to Json streams
+	/// </summary>
+	public class LogEventConverter : JsonConverter<LogEvent>
+	{
+		/// <inheritdoc/>
+		public override LogEvent Read(ref Utf8JsonReader Reader, Type TypeToConvert, JsonSerializerOptions Options)
+		{
+			return LogEvent.Read(ref Reader);
+		}
+
+		/// <inheritdoc/>
+		public override void Write(Utf8JsonWriter Writer, LogEvent Value, JsonSerializerOptions Options)
+		{
+			Value.Write(Writer);
+		}
+	}
 }
