Horde: Validate the preflight change when starting a job, and return a descriptive error describing the problem with it if it fails.

diff --git a/Engine/Source/Programs/Horde/Horde.Build/Controllers/JobsController.cs b/Engine/Source/Programs/Horde/Horde.Build/Controllers/JobsController.cs
--- a/Engine/Source/Programs/Horde/Horde.Build/Controllers/JobsController.cs
+++ b/Engine/Source/Programs/Horde/Horde.Build/Controllers/JobsController.cs
@@ -169,6 +169,27 @@ namespace HordeServer.Controllers
 				}
 			}
 
+			// Check the preflight change is valid
+			if (Create.PreflightChange != null)
+			{
+				CheckShelfResult Result = await Perforce.CheckShelfAsync(Stream.ClusterName, Stream.Name, Create.PreflightChange.Value, null);
+				switch (Result)
+				{
+					case CheckShelfResult.Ok:
+						break;
+					case CheckShelfResult.NoChange:
+						return BadRequest(KnownLogEvents.Horde_InvalidPreflight, "CL {Change} does not exist", Create.PreflightChange);
+					case CheckShelfResult.NoShelvedFiles:
+						return BadRequest(KnownLogEvents.Horde_InvalidPreflight, "CL {Change} does not contain any shelved files", Create.PreflightChange);
+					case CheckShelfResult.WrongStream:
+						return BadRequest(KnownLogEvents.Horde_InvalidPreflight, "CL {Change} does not contain files in {Stream}", Create.PreflightChange, Stream.Name);
+					case CheckShelfResult.MixedStream:
+						return BadRequest(KnownLogEvents.Horde_InvalidPreflight, "CL {Change} contains files from multiple streams", Create.PreflightChange);
+					default:
+						return BadRequest(KnownLogEvents.Horde_InvalidPreflight, "CL {Change} cannot be preflighted ({Result})", Create.PreflightChange, Result);
+				}
+			}
+
 			bool? UpdateIssues = null;
 			if (Template.UpdateIssues)
 			{
@@ -700,7 +721,7 @@ namespace HordeServer.Controllers
 			}
 			return Responses;
 		}
-		
+
 		/// <summary>
 		/// Find jobs for a stream with given templates, sorted by creation date
 		/// </summary>
diff --git a/Engine/Source/Programs/Horde/Horde.Build/Perforce/IPerforceService.cs b/Engine/Source/Programs/Horde/Horde.Build/Perforce/IPerforceService.cs
--- a/Engine/Source/Programs/Horde/Horde.Build/Perforce/IPerforceService.cs
+++ b/Engine/Source/Programs/Horde/Horde.Build/Perforce/IPerforceService.cs
@@ -44,6 +44,37 @@ namespace HordeServer.Services
 		public string? Email { get; set; }
 	}
 
+	/// <summary>
+	/// Result from checking a shelved change status
+	/// </summary>
+	public enum CheckShelfResult
+	{
+		/// <summary>
+		/// The shelf is ok
+		/// </summary>
+		Ok,
+
+		/// <summary>
+		/// Changelist does not exist
+		/// </summary>
+		NoChange,
+
+		/// <summary>
+		/// The change does not contain any shelved files
+		/// </summary>
+		NoShelvedFiles,
+
+		/// <summary>
+		/// The shelf contains *some* files from a different stream
+		/// </summary>
+		MixedStream,
+
+		/// <summary>
+		/// The shelf contains only files from a different stream
+		/// </summary>
+		WrongStream,
+	}
+
 	/// <summary>
 	/// Interface for a Perforce stream view
 	/// </summary>
@@ -143,6 +174,16 @@ namespace HordeServer.Services
 		/// <returns>Latest changelist number</returns>
 		public Task<List<ChangeSummary>> GetChangesAsync(string ClusterName, string StreamName, int? MinChange, int? MaxChange, int Results, string? ImpersonateUser);
 
+		/// <summary>
+		/// Checks a shelf is valid for the given stream
+		/// </summary>
+		/// <param name="ClusterName">Name of the Perforce cluster</param>
+		/// <param name="StreamName">The stream to query</param>
+		/// <param name="ChangeNumber">Shelved changelist number</param>
+		/// <param name="ImpersonateUser">Name of the user to impersonate</param>
+		/// <returns></returns>
+		public Task<CheckShelfResult> CheckShelfAsync(string ClusterName, string StreamName, int ChangeNumber, string? ImpersonateUser);
+
 		/// <summary>
 		/// Gets the latest change for a particular stream
 		/// </summary>
diff --git a/Engine/Source/Programs/Horde/Horde.Build/Perforce/Impl/PerforceService.cs b/Engine/Source/Programs/Horde/Horde.Build/Perforce/Impl/PerforceService.cs
--- a/Engine/Source/Programs/Horde/Horde.Build/Perforce/Impl/PerforceService.cs
+++ b/Engine/Source/Programs/Horde/Horde.Build/Perforce/Impl/PerforceService.cs
@@ -885,6 +885,53 @@ namespace HordeServer.Services
 			}
 		}
 
+		/// <inheritdoc/>
+		public async Task<CheckShelfResult> CheckShelfAsync(string ClusterName, string StreamName, int ChangeNumber, string? ImpersonateUser)
+		{
+			using IScope Scope = GlobalTracer.Instance.BuildSpan("PerforceService.CheckPreflightAsync").StartActive();
+			Scope.Span.SetTag("ClusterName", ClusterName);
+			Scope.Span.SetTag("StreamName", StreamName);
+			Scope.Span.SetTag("ChangeNumber", ChangeNumber);
+			Scope.Span.SetTag("ImpersonateUser", ImpersonateUser);
+
+			PerforceCluster Cluster = await GetClusterAsync(ClusterName);
+			using (P4.Repository Repository = await GetConnection(Cluster, Stream: StreamName, Username: ImpersonateUser))
+			{
+				P4.Changelist Change = Repository.GetChangelist(ChangeNumber, new P4.DescribeCmdOptions(P4.DescribeChangelistCmdFlags.Omit | P4.DescribeChangelistCmdFlags.Shelved, 0, 0));
+				if(Change == null)
+				{
+					return CheckShelfResult.NoChange;
+				}
+				if (Change.ShelvedFiles.Count == 0)
+				{
+					return CheckShelfResult.NoShelvedFiles;
+				}
+
+				P4.Stream Stream = Repository.GetStream(StreamName, new P4.StreamCmdOptions(P4.StreamCmdFlags.View, null, null));
+				ViewMap Map = new ViewMap(Stream.View);
+
+				bool bHasMappedFile = false;
+				bool bHasUnmappedFile = false;
+				foreach (P4.ShelvedFile ShelvedFile in Change.ShelvedFiles)
+				{
+					if (Map.TryMapFile(ShelvedFile.Path.Path, Utf8StringComparer.OrdinalIgnoreCase, out Utf8String _))
+					{
+						bHasMappedFile = true;
+					}
+					else
+					{
+						bHasUnmappedFile = true;
+					}
+				}
+
+				if (bHasUnmappedFile)
+				{
+					return bHasMappedFile ? CheckShelfResult.MixedStream : CheckShelfResult.WrongStream;
+				}
+			}
+			return CheckShelfResult.Ok;
+		}
+
 		/// <inheritdoc/>
 		public async Task<List<ChangeDetails>> GetChangeDetailsAsync(string ClusterName, string StreamName, IReadOnlyList<int> ChangeNumbers, string? ImpersonateUser)
 		{
diff --git a/Engine/Source/Programs/Shared/EpicGames.Core/KnownLogEvents.cs b/Engine/Source/Programs/Shared/EpicGames.Core/KnownLogEvents.cs
--- a/Engine/Source/Programs/Shared/EpicGames.Core/KnownLogEvents.cs
+++ b/Engine/Source/Programs/Shared/EpicGames.Core/KnownLogEvents.cs
@@ -163,6 +163,16 @@ namespace EpicGames.Core
 		/// A systemic event from MSBuild
 		/// </summary>
 		public static EventId Systemic_MSBuild { get; } = new EventId(750);
+
+		/// <summary>
+		/// Horde error codes
+		/// </summary>
+		public static EventId Horde { get; } = new EventId(1000);
+
+		/// <summary>
+		/// Invalid preflight change
+		/// </summary>
+		public static EventId Horde_InvalidPreflight { get; } = new EventId(1001);
 	}
 #pragma warning restore CA1707 // Identifiers should not contain underscores
 }
