Horde Storage
Fixed issue with finalize sometimes operating on a stale version of the object. Avoiding IF in the update (which turns it into a transaction) coupled with avoiding reading the object when we just inserted resolves these issues.

Fundementally the problem is that we set a very aggresive consistency level to make writes faster, which means you need to be careful when doing a read directly after a write.

diff --git a/Engine/Source/Programs/Horde/HordeStorage/Horde.Storage/Implementation/ObjectService.cs b/Engine/Source/Programs/Horde/HordeStorage/Horde.Storage/Implementation/ObjectService.cs
--- a/Engine/Source/Programs/Horde/HordeStorage/Horde.Storage/Implementation/ObjectService.cs
+++ b/Engine/Source/Programs/Horde/HordeStorage/Horde.Storage/Implementation/ObjectService.cs
@@ -94,7 +94,7 @@ namespace Horde.Storage.Implementation
 
             if (missingReferences.Length == 0)
             {
-                await _referencesStore.Finalize(ns, bucket, key, blobHash);
+                await _referencesStore.Finalize(ns, bucket, key);
                 await _replicationLog.InsertAddEvent(ns, bucket, key, blobHash);
             }
 
@@ -107,6 +107,9 @@ namespace Horde.Storage.Implementation
             byte[] blobContents = await blob.Stream.ToByteArray();
             CompactBinaryObject payload = CompactBinaryObject.Load(blobContents);
 
+            if (!o.BlobIdentifier.Equals(blobHash))
+                throw new ObjectHashMismatchException(ns, bucket, key, blobHash, o.BlobIdentifier);
+
             bool hasReferences = payload.GetAllFields().Any(field => field.IsAttachment());
 
             BlobIdentifier[] missingReferences = Array.Empty<BlobIdentifier>();
@@ -126,7 +129,7 @@ namespace Horde.Storage.Implementation
 
             if (missingReferences.Length == 0)
             {
-                await _referencesStore.Finalize(ns, bucket, key, blobHash);
+                await _referencesStore.Finalize(ns, bucket, key);
                 await _replicationLog.InsertAddEvent(ns, bucket, key, blobHash);
             }
 
diff --git a/Engine/Source/Programs/Horde/HordeStorage/Horde.Storage/Implementation/Objects/IReferencesStore.cs b/Engine/Source/Programs/Horde/HordeStorage/Horde.Storage/Implementation/Objects/IReferencesStore.cs
--- a/Engine/Source/Programs/Horde/HordeStorage/Horde.Storage/Implementation/Objects/IReferencesStore.cs
+++ b/Engine/Source/Programs/Horde/HordeStorage/Horde.Storage/Implementation/Objects/IReferencesStore.cs
@@ -11,7 +11,7 @@ namespace Horde.Storage.Implementation
     {
         Task<ObjectRecord> Get(NamespaceId ns, BucketId bucket, IoHashKey key);
         Task Put(NamespaceId ns, BucketId bucket, IoHashKey key, BlobIdentifier blobHash, byte[] blob, bool isFinalized);
-        Task Finalize(NamespaceId ns, BucketId bucket, IoHashKey key, BlobIdentifier blobHash);
+        Task Finalize(NamespaceId ns, BucketId bucket, IoHashKey key);
 
 
         Task UpdateLastAccessTime(NamespaceId ns, BucketId bucket, IoHashKey key, DateTime newLastAccessTime);
diff --git a/Engine/Source/Programs/Horde/HordeStorage/Horde.Storage/Implementation/Objects/MemoryReferencesStore.cs b/Engine/Source/Programs/Horde/HordeStorage/Horde.Storage/Implementation/Objects/MemoryReferencesStore.cs
--- a/Engine/Source/Programs/Horde/HordeStorage/Horde.Storage/Implementation/Objects/MemoryReferencesStore.cs
+++ b/Engine/Source/Programs/Horde/HordeStorage/Horde.Storage/Implementation/Objects/MemoryReferencesStore.cs
@@ -44,7 +44,7 @@ namespace Horde.Storage.Implementation
             return Task.FromResult(o);
         }
 
-        public Task Finalize(NamespaceId ns, BucketId bucket, IoHashKey key, BlobIdentifier blobHash)
+        public Task Finalize(NamespaceId ns, BucketId bucket, IoHashKey key)
         {
             if (!_objects.TryGetValue(BuildKey(ns, bucket, key), out MemoryStoreObject? o))
             {
diff --git a/Engine/Source/Programs/Horde/HordeStorage/Horde.Storage/Implementation/Objects/ScyllaReferencesStore.cs b/Engine/Source/Programs/Horde/HordeStorage/Horde.Storage/Implementation/Objects/ScyllaReferencesStore.cs
--- a/Engine/Source/Programs/Horde/HordeStorage/Horde.Storage/Implementation/Objects/ScyllaReferencesStore.cs
+++ b/Engine/Source/Programs/Horde/HordeStorage/Horde.Storage/Implementation/Objects/ScyllaReferencesStore.cs
@@ -89,13 +89,13 @@ namespace Horde.Storage.Implementation
         }
 
 
-        public async Task Finalize(NamespaceId ns, BucketId bucket, IoHashKey name, BlobIdentifier blobHash)
+        public async Task Finalize(NamespaceId ns, BucketId bucket, IoHashKey name)
         {
             using Scope _ = Tracer.Instance.StartActive("scylla.finalize");
 
-            AppliedInfo<ScyllaObject> info = await _mapper.UpdateIfAsync<ScyllaObject>("SET is_finalized=true WHERE namespace=? AND bucket=? AND name=? IF payload_hash=?", ns.ToString(), bucket.ToString(), name.ToString(), blobHash.ToString());
+            AppliedInfo<ScyllaObject> info = await _mapper.UpdateIfAsync<ScyllaObject>("SET is_finalized=true WHERE namespace=? AND bucket=? AND name=?", ns.ToString(), bucket.ToString(), name.ToString());
             if (!info.Applied)
-                throw new ObjectHashMismatchException(ns, bucket, name, blobHash, info.Existing.PayloadHash!.AsBlobIdentifier());
+                throw new Exception($"Failed to finalize object in namespace {ns} {bucket} {name}");
         }
 
         private async Task CreateTTLRecord(NamespaceId ns, BucketId bucket, IoHashKey name, sbyte partitionIndex)
diff --git a/Engine/Source/Programs/Horde/HordeStorage/Horde.Storage/Tests/Functional/Objects/ReferencesTests.cs b/Engine/Source/Programs/Horde/HordeStorage/Horde.Storage/Tests/Functional/Objects/ReferencesTests.cs
--- a/Engine/Source/Programs/Horde/HordeStorage/Horde.Storage/Tests/Functional/Objects/ReferencesTests.cs
+++ b/Engine/Source/Programs/Horde/HordeStorage/Horde.Storage/Tests/Functional/Objects/ReferencesTests.cs
@@ -79,7 +79,6 @@ namespace Horde.Storage.FunctionalTests.References
                 localKeyspace.ExecuteAsync(new SimpleStatement("DROP TABLE IF EXISTS replication_namespace;"))
             );
 
-            await Task.CompletedTask;
         }
     }
 
