Horde: Fix compilation errors in Horde.Build.Tests

diff --git a/Engine/Source/Programs/Horde/Horde.Build.Tests/AgentControllerTest.cs b/Engine/Source/Programs/Horde/Horde.Build.Tests/AgentControllerTest.cs
--- a/Engine/Source/Programs/Horde/Horde.Build.Tests/AgentControllerTest.cs
+++ b/Engine/Source/Programs/Horde/Horde.Build.Tests/AgentControllerTest.cs
@@ -18,16 +18,16 @@ namespace HordeServerTests
 			Fixture Fixture = await CreateFixtureAsync();
 	        IAgent FixtureAgent = Fixture.Agent1;
 
-	        ActionResult<object> Obj = await AgentsController.GetAgentAsync(FixtureAgent.Id.ToString());
+	        ActionResult<object> Obj = await AgentsController.GetAgentAsync(FixtureAgent.Id);
 	        GetAgentResponse GetRes = (Obj.Value as GetAgentResponse)!;
 	        Assert.AreEqual(Fixture!.Agent1Name.ToUpper(), GetRes.Name);
 	        Assert.IsNull(GetRes.Comment);
 	        
 	        UpdateAgentRequest UpdateReq = new UpdateAgentRequest();
 	        UpdateReq.Comment = "foo bar baz";
-	        await AgentsController.UpdateAgentAsync(FixtureAgent.Id.ToString(), UpdateReq);
+	        await AgentsController.UpdateAgentAsync(FixtureAgent.Id, UpdateReq);
 	        
-	        Obj = await AgentsController.GetAgentAsync(FixtureAgent.Id.ToString());
+	        Obj = await AgentsController.GetAgentAsync(FixtureAgent.Id);
 	        GetRes = (Obj.Value as GetAgentResponse)!;
 	        Assert.AreEqual("foo bar baz", GetRes.Comment);
         }
diff --git a/Engine/Source/Programs/Horde/Horde.Build.Tests/JobCollectionTests.cs b/Engine/Source/Programs/Horde/Horde.Build.Tests/JobCollectionTests.cs
--- a/Engine/Source/Programs/Horde/Horde.Build.Tests/JobCollectionTests.cs
+++ b/Engine/Source/Programs/Horde/Horde.Build.Tests/JobCollectionTests.cs
@@ -122,12 +122,14 @@ namespace HordeServerTests
 		{
 			Fixture Fixture = await CreateFixtureAsync();
 
+			ObjectId<ISession> SessionId1 = new (ObjectId.GenerateNewId());
 			await JobCollection.TryAssignLeaseAsync(Fixture.Job1, 0, new PoolId("foo"), Fixture.Agent1.Id,
-				ObjectId.GenerateNewId(), LeaseId.GenerateNewId(), LogId.GenerateNewId());
+				SessionId1, LeaseId.GenerateNewId(), LogId.GenerateNewId());
 			
+			ObjectId<ISession> SessionId2 = new (ObjectId.GenerateNewId());
 			IJob Job = (await JobCollection.GetAsync(Fixture.Job1.Id))!;
 			await JobCollection.TryAssignLeaseAsync(Job, 0, new PoolId("foo"), Fixture.Agent1.Id,
-				ObjectId.GenerateNewId(), LeaseId.GenerateNewId(), LogId.GenerateNewId());
+				SessionId2, LeaseId.GenerateNewId(), LogId.GenerateNewId());
 			
 			// Manually verify the log output
 		}
diff --git a/Engine/Source/Programs/Horde/Horde.Build.Tests/JobsControllerDbTest.cs b/Engine/Source/Programs/Horde/Horde.Build.Tests/JobsControllerDbTest.cs
--- a/Engine/Source/Programs/Horde/Horde.Build.Tests/JobsControllerDbTest.cs
+++ b/Engine/Source/Programs/Horde/Horde.Build.Tests/JobsControllerDbTest.cs
@@ -45,16 +45,16 @@ namespace HordeServerTests
 	        SubResourceId BatchId = Job.Batches[0].Id;
 	        SubResourceId StepId = Job.Batches[0].Steps[0].Id;
 
-	        object Obj = (await JobsController.GetStepAsync(Job.Id, BatchId.ToString(), StepId.ToString())).Value!;
+	        object Obj = (await JobsController.GetStepAsync(Job.Id, BatchId, StepId)).Value!;
 	        GetStepResponse StepRes = (Obj as GetStepResponse)!;
 	        Assert.IsFalse(StepRes.AbortRequested);
 	        
 	        UpdateStepRequest UpdateReq = new UpdateStepRequest();
 	        UpdateReq.AbortRequested = true;
-	        Obj = (await JobsController.UpdateStepAsync(Job.Id, BatchId.ToString(), StepId.ToString(), UpdateReq)).Value!;
+	        Obj = (await JobsController.UpdateStepAsync(Job.Id, BatchId, StepId, UpdateReq)).Value!;
 	        UpdateStepResponse UpdateRes = (Obj as UpdateStepResponse)!;
 	        
-	        Obj = (await JobsController.GetStepAsync(Job.Id, BatchId.ToString(), StepId.ToString())).Value!;
+	        Obj = (await JobsController.GetStepAsync(Job.Id, BatchId, StepId)).Value!;
 	        StepRes = (Obj as GetStepResponse)!;
 	        Assert.IsTrue(StepRes.AbortRequested);
 //	        Assert.AreEqual("Anonymous", StepRes.AbortByUser);
diff --git a/Engine/Source/Programs/Horde/Horde.Build.Tests/LogFileServiceTests.cs b/Engine/Source/Programs/Horde/Horde.Build.Tests/LogFileServiceTests.cs
--- a/Engine/Source/Programs/Horde/Horde.Build.Tests/LogFileServiceTests.cs
+++ b/Engine/Source/Programs/Horde/Horde.Build.Tests/LogFileServiceTests.cs
@@ -231,7 +231,7 @@ namespace HordeServerTests
 			// Will implicitly test GetLogFileAsync(), AddCachedLogFile()
 			JobId JobId = JobId.GenerateNewId();
             ObjectId SessionId = ObjectId.GenerateNewId();
-            ILogFile A = await LogFileService.CreateLogFileAsync(JobId, SessionId, LogType.Text);
+            ILogFile A = await LogFileService.CreateLogFileAsync(JobId, new ObjectId<ISession>(SessionId), LogType.Text);
             ILogFile B = (await LogFileService.GetCachedLogFileAsync(A.Id))!;
             Assert.AreEqual(A.JobId, B.JobId);
             Assert.AreEqual(A.SessionId, B.SessionId);
@@ -240,7 +240,7 @@ namespace HordeServerTests
             ILogFile? NotFound = await LogFileService.GetCachedLogFileAsync(LogId.GenerateNewId());
             Assert.IsNull(NotFound);
 
-            await LogFileService.CreateLogFileAsync(JobId.GenerateNewId(), ObjectId.GenerateNewId(), LogType.Text);
+            await LogFileService.CreateLogFileAsync(JobId.GenerateNewId(), new ObjectId<ISession>(ObjectId.GenerateNewId()), LogType.Text);
             Assert.AreEqual(2, (await LogFileService.GetLogFilesAsync()).Count);
         }
 
@@ -249,7 +249,7 @@ namespace HordeServerTests
         {
 			JobId JobId = JobId.GenerateNewId();
             ObjectId SessionId = ObjectId.GenerateNewId();
-            ILogFile LogFile = await LogFileService.CreateLogFileAsync(JobId, SessionId, LogType.Text);
+            ILogFile LogFile = await LogFileService.CreateLogFileAsync(JobId, new ObjectId<ISession>(SessionId), LogType.Text);
             ILogFile LogFileNoSession = await LogFileService.CreateLogFileAsync(JobId, null, LogType.Text);
 
             var HasClaim = new ClaimsPrincipal(new ClaimsIdentity(new List<Claim>
diff --git a/Engine/Source/Programs/Horde/Horde.Build.Tests/RpcServiceTest.cs b/Engine/Source/Programs/Horde/Horde.Build.Tests/RpcServiceTest.cs
--- a/Engine/Source/Programs/Horde/Horde.Build.Tests/RpcServiceTest.cs
+++ b/Engine/Source/Programs/Horde/Horde.Build.Tests/RpcServiceTest.cs
@@ -407,7 +407,7 @@ namespace HordeServerTests
 
 			// Set the session ID on the job batch to pass auth later
 			Deref(await JobCollection.TryAssignLeaseAsync(Fixture.Job1, 0, new PoolId("foo"),
-				new AgentId("test"), SessionId,
+				new AgentId("test"), new ObjectId<HordeServer.Models.ISession>(SessionId),
 				LeaseId.GenerateNewId(), LogId.GenerateNewId()));
 /*
 			TestAsyncStreamReader<UploadArtifactRequest> RequestStream = new TestAsyncStreamReader<UploadArtifactRequest>(Context);
diff --git a/Engine/Source/Programs/Horde/Horde.Build.Tests/Stubs/Services/PerforceServiceStub.cs b/Engine/Source/Programs/Horde/Horde.Build.Tests/Stubs/Services/PerforceServiceStub.cs
--- a/Engine/Source/Programs/Horde/Horde.Build.Tests/Stubs/Services/PerforceServiceStub.cs
+++ b/Engine/Source/Programs/Horde/Horde.Build.Tests/Stubs/Services/PerforceServiceStub.cs
@@ -115,6 +115,11 @@ namespace HordeServerTests.Stubs.Services
 			return Task.FromResult(Results);
 		}
 
+		public Task<CheckShelfResult> CheckShelfAsync(string ClusterName, string StreamName, int ChangeNumber, string? ImpersonateUser)
+		{
+			throw new NotImplementedException();
+		}
+
 		public Task<PerforceUserInfo?> GetUserInfoAsync(string ClusterName, string UserName)
 		{
 			return Task.FromResult<PerforceUserInfo?>(new PerforceUserInfo { Login = UserName, FullName = UserName, Email = $"[REDACTED_EMAIL]" });
