Horde: Add index hint and read consistency parameters to LeaseCollection

A new method override for FindLeasesAsync is also provided with hand-picked index and read consistency. Should result in more performant lease lookups in the autoscale service (which are currently super slow and I/O demanding due to a bad index). A database index class was also factored out to avoid duplicating code.

diff --git a/Engine/Source/Programs/Horde/Horde.Build/Agents/ILeaseCollection.cs b/Engine/Source/Programs/Horde/Horde.Build/Agents/ILeaseCollection.cs
--- a/Engine/Source/Programs/Horde/Horde.Build/Agents/ILeaseCollection.cs
+++ b/Engine/Source/Programs/Horde/Horde.Build/Agents/ILeaseCollection.cs
@@ -63,8 +63,19 @@ namespace HordeServer.Collections
 		/// <param name="MaxTime">End of the window to include leases</param>
 		/// <param name="Index">Index of the first result to return</param>
 		/// <param name="Count">Number of results to return</param>
+		/// <param name="IndexHint">Name of index to be specified as a hint to the database query planner</param>
+		/// <param name="ConsistentRead">If the database read should be made to the replica server</param>
 		/// <returns>List of leases matching the given criteria</returns>
-		Task<List<ILease>> FindLeasesAsync(AgentId? AgentId = null, SessionId? SessionId = null, DateTime? MinTime = null, DateTime? MaxTime = null, int? Index = null, int? Count = null);
+		Task<List<ILease>> FindLeasesAsync(AgentId? AgentId = null, SessionId? SessionId = null, DateTime? MinTime = null, DateTime? MaxTime = null, int? Index = null, int? Count = null, string? IndexHint = null, bool ConsistentRead = true);
+
+		/// <summary>
+		/// Finds all leases between min and/or max time
+		/// Preferred over the more generic FindLeasesAsync as this one use a better index and relaxed read consistency
+		/// </summary>
+		/// <param name="MinTime">Start of the window to include leases</param>
+		/// <param name="MaxTime">End of the window to include leases</param>
+		/// <returns>List of leases matching the given criteria</returns>
+		Task<List<ILease>> FindLeasesAsync(DateTime? MinTime, DateTime? MaxTime);
 
 		/// <summary>
 		/// Finds all active leases
diff --git a/Engine/Source/Programs/Horde/Horde.Build/Agents/Impl/LeaseCollection.cs b/Engine/Source/Programs/Horde/Horde.Build/Agents/Impl/LeaseCollection.cs
--- a/Engine/Source/Programs/Horde/Horde.Build/Agents/Impl/LeaseCollection.cs
+++ b/Engine/Source/Programs/Horde/Horde.Build/Agents/Impl/LeaseCollection.cs
@@ -68,11 +68,34 @@ namespace HordeServer.Collections.Impl
 				this.Payload = Payload;
 			}
 		}
+		
+		class DatabaseIndexes : BaseDatabaseIndexes
+		{
+			public readonly string AgentId;
+			public readonly string SessionId;
+			public readonly string StartTime;
+			public readonly string FinishTime;
+			public readonly string FinishTimeStartTimeCompound;
+
+			public DatabaseIndexes(IMongoCollection<LeaseDocument> Leases, bool DatabaseReadOnlyMode) : base(DatabaseReadOnlyMode)
+			{
+				AgentId = CreateOrGetIndex(Leases, "StreamId_1", Builders<LeaseDocument>.IndexKeys.Ascending(x => x.AgentId));
+				SessionId = CreateOrGetIndex(Leases, "SessionId_1", Builders<LeaseDocument>.IndexKeys.Ascending(x => x.SessionId));
+				StartTime = CreateOrGetIndex(Leases, "StartTime_1", Builders<LeaseDocument>.IndexKeys.Ascending(x => x.StartTime));
+				FinishTime = CreateOrGetIndex(Leases, "FinishTime_1", Builders<LeaseDocument>.IndexKeys.Ascending(x => x.FinishTime));
+				FinishTimeStartTimeCompound = CreateOrGetIndex(Leases, "FinishTime_1_StartTime_-1", Builders<LeaseDocument>.IndexKeys.Ascending(x => x.FinishTime).Descending(x => x.StartTime));
+			}
+		}
 
 		/// <summary>
 		/// Collection of lease documents
 		/// </summary>
 		readonly IMongoCollection<LeaseDocument> Leases;
+		
+		/// <summary>
+		/// Creation and lookup of indexes
+		/// </summary>
+		DatabaseIndexes Indexes;
 
 		/// <summary>
 		/// Constructor
@@ -81,14 +104,7 @@ namespace HordeServer.Collections.Impl
 		public LeaseCollection(DatabaseService DatabaseService)
 		{
 			Leases = DatabaseService.GetCollection<LeaseDocument>("Leases");
-
-			if (!DatabaseService.ReadOnlyMode)
-			{
-				Leases.Indexes.CreateOne(new CreateIndexModel<LeaseDocument>(Builders<LeaseDocument>.IndexKeys.Ascending(x => x.AgentId)));
-				Leases.Indexes.CreateOne(new CreateIndexModel<LeaseDocument>(Builders<LeaseDocument>.IndexKeys.Ascending(x => x.SessionId)));
-				Leases.Indexes.CreateOne(new CreateIndexModel<LeaseDocument>(Builders<LeaseDocument>.IndexKeys.Ascending(x => x.StartTime)));
-				Leases.Indexes.CreateOne(new CreateIndexModel<LeaseDocument>(Builders<LeaseDocument>.IndexKeys.Ascending(x => x.FinishTime)));
-			}
+			Indexes = new DatabaseIndexes(Leases, DatabaseService.ReadOnlyMode);
 		}
 
 		/// <inheritdoc/>
@@ -112,10 +128,11 @@ namespace HordeServer.Collections.Impl
 		}
 
 		/// <inheritdoc/>
-		public async Task<List<ILease>> FindLeasesAsync(AgentId? AgentId, SessionId? SessionId, DateTime? MinTime, DateTime? MaxTime, int? Index, int? Count)
+		public async Task<List<ILease>> FindLeasesAsync(AgentId? AgentId, SessionId? SessionId, DateTime? MinTime, DateTime? MaxTime, int? Index, int? Count, string? IndexHint = null, bool ConsistentRead = true)
 		{
+			IMongoCollection<LeaseDocument> Collection = ConsistentRead ? Leases : Leases.WithReadPreference(ReadPreference.SecondaryPreferred);
+			
 			FilterDefinitionBuilder<LeaseDocument> FilterBuilder = Builders<LeaseDocument>.Filter;
-
 			FilterDefinition<LeaseDocument> Filter = FilterDefinition<LeaseDocument>.Empty;
 			if (AgentId != null)
 			{
@@ -133,11 +150,22 @@ namespace HordeServer.Collections.Impl
 			{
 				Filter &= FilterBuilder.Lt(x => x.StartTime, MaxTime.Value);
 			}
+			FindOptions? FindOptions = null;
+			if (IndexHint != null)
+			{
+				FindOptions = new FindOptions { Hint = new BsonString(IndexHint) };
+			}
 
-			List<LeaseDocument> Results = await Leases.Find(Filter).SortByDescending(x => x.StartTime).Range(Index, Count).ToListAsync();
+			List<LeaseDocument> Results = await Collection.Find(Filter, FindOptions).SortByDescending(x => x.StartTime).Range(Index, Count).ToListAsync();
 			return Results.ConvertAll<ILease>(x => x);
 		}
 
+		/// <inheritdoc/>
+		public async Task<List<ILease>> FindLeasesAsync(DateTime? MinTime, DateTime? MaxTime)
+		{
+			return await FindLeasesAsync(null, null, MinTime, MaxTime, null, null, Indexes.FinishTimeStartTimeCompound, false);
+		}
+
 		/// <inheritdoc/>
 		public async Task<List<ILease>> FindActiveLeasesAsync(int? Index = null, int? Count = null)
 		{
diff --git a/Engine/Source/Programs/Horde/Horde.Build/Fleet/AutoscaleService.cs b/Engine/Source/Programs/Horde/Horde.Build/Fleet/AutoscaleService.cs
--- a/Engine/Source/Programs/Horde/Horde.Build/Fleet/AutoscaleService.cs
+++ b/Engine/Source/Programs/Horde/Horde.Build/Fleet/AutoscaleService.cs
@@ -131,7 +131,7 @@ namespace HordeServer.Services
 				// Query leases in last interval
 				DateTime MaxTime = DateTime.UtcNow;
 				DateTime MinTime = MaxTime - (SampleTime * NumSamples);
-				List<ILease> Leases = await LeaseCollection.FindLeasesAsync(null, null, MinTime, MaxTime, null, null);
+				List<ILease> Leases = await LeaseCollection.FindLeasesAsync(MinTime, MaxTime);
 
 				// Add all the leases to a data object for each agent
 				Dictionary<AgentId, AgentData> AgentIdToData = Agents.ToDictionary(x => x.Id, x => new AgentData(x));
diff --git a/Engine/Source/Programs/Horde/Horde.Build/Jobs/Impl/JobCollection.cs b/Engine/Source/Programs/Horde/Horde.Build/Jobs/Impl/JobCollection.cs
--- a/Engine/Source/Programs/Horde/Horde.Build/Jobs/Impl/JobCollection.cs
+++ b/Engine/Source/Programs/Horde/Horde.Build/Jobs/Impl/JobCollection.cs
@@ -313,7 +313,7 @@ namespace HordeServer.Collections.Impl
 			}
 		}
 
-		class DatabaseIndexes
+		class DatabaseIndexes : BaseDatabaseIndexes
 		{
 			public readonly string StreamId;
 			public readonly string Change;
@@ -324,12 +324,9 @@ namespace HordeServer.Collections.Impl
 			public readonly string StartedByUserId;
 			public readonly string TemplateId;
 			public readonly string SchedulePriority;
-			private bool DatabaseReadOnlyMode;
 
-			public DatabaseIndexes(IMongoCollection<JobDocument> Jobs, bool DatabaseReadOnlyMode)
+			public DatabaseIndexes(IMongoCollection<JobDocument> Jobs, bool DatabaseReadOnlyMode) : base(DatabaseReadOnlyMode)
 			{
-				this.DatabaseReadOnlyMode = DatabaseReadOnlyMode;
-				
 				StreamId = CreateOrGetIndex(Jobs, "StreamId_1", Builders<JobDocument>.IndexKeys.Ascending(x => x.StreamId));
 				Change = CreateOrGetIndex(Jobs, "Change_1", Builders<JobDocument>.IndexKeys.Ascending(x => x.Change));
 				PreflightChange = CreateOrGetIndex(Jobs, "PreflightChange_1", Builders<JobDocument>.IndexKeys.Ascending(x => x.PreflightChange));
@@ -340,22 +337,6 @@ namespace HordeServer.Collections.Impl
 				TemplateId = CreateOrGetIndex(Jobs, "TemplateId_1", Builders<JobDocument>.IndexKeys.Ascending(x => x.TemplateId));
 				SchedulePriority = CreateOrGetIndex(Jobs, "SchedulePriority_-1", Builders<JobDocument>.IndexKeys.Descending(x => x.SchedulePriority));
 			}
-
-			private string CreateOrGetIndex<T>(IMongoCollection<T> Collection, string Name, IndexKeysDefinition<T> IndexDef)
-			{
-				if (DatabaseReadOnlyMode)
-				{
-					if (Collection.Indexes.List().ToList().Find(x => x.TryGetString("name", out string? n) && n == Name) == null)
-					{
-						throw new ArgumentException("Index with name {Name} not found (running in database read-only mode)");
-					}
-					return Name;
-				}
-
-				CreateIndexOptions IndexOptions = new() { Name = Name, Background = true };
-				Collection.Indexes.CreateOne(new CreateIndexModel<T>(IndexDef, IndexOptions));
-				return Name;
-			}
 		}
 
 		/// <summary>
diff --git a/Engine/Source/Programs/Horde/Horde.Build/Utilities/BaseDatabaseIndexes.cs b/Engine/Source/Programs/Horde/Horde.Build/Utilities/BaseDatabaseIndexes.cs
new file mode 100644
--- /dev/null
+++ b/Engine/Source/Programs/Horde/Horde.Build/Utilities/BaseDatabaseIndexes.cs
@@ -0,0 +1,53 @@
+// Copyright Epic Games, Inc. All Rights Reserved.
+
+using System;
+using MongoDB.Driver;
+
+namespace HordeServer.Utilities
+{
+	/// <summary>
+	/// Base class for keeping track of indexes in MongoDB
+	/// </summary>
+	public abstract class BaseDatabaseIndexes
+	{
+		/// <summary>
+		/// If database is in read-only mode (i.e no indexes can be created/modified)
+		/// </summary>
+		private bool DatabaseReadOnlyMode;
+
+		/// <summary>
+		/// Constructor
+		/// </summary>
+		/// <param name="DatabaseReadOnlyMode"></param>
+		protected BaseDatabaseIndexes(bool DatabaseReadOnlyMode)
+		{
+			this.DatabaseReadOnlyMode = DatabaseReadOnlyMode;
+		}
+
+		/// <summary>
+		/// Create or get an index
+		/// </summary>
+		/// <param name="Collection">Name of collection</param>
+		/// <param name="Name">Name of index (used for checking if index exists)</param>
+		/// <param name="IndexDef">Type-specific index definition</param>
+		/// <typeparam name="T">Document type</typeparam>
+		/// <returns></returns>
+		/// <exception cref="ArgumentException">If index cannot be found and database is in read-only mode</exception>
+		protected string CreateOrGetIndex<T>(IMongoCollection<T> Collection, string Name, IndexKeysDefinition<T> IndexDef)
+		{
+			if (DatabaseReadOnlyMode)
+			{
+				if (Collection.Indexes.List().ToList().Find(x => x.TryGetString("name", out string? n) && n == Name) == null)
+				{
+					throw new ArgumentException("Index with name {Name} not found (running in database read-only mode)");
+				}
+				return Name;
+			}
+
+			CreateIndexOptions IndexOptions = new() { Name = Name, Background = true };
+			Collection.Indexes.CreateOne(new CreateIndexModel<T>(IndexDef, IndexOptions));
+			return Name;
+		}
+	}
+
+}
