Horde Storage
Commit a full object to the the blob index if it doesnt exist, otherwise update the regions.
Should help resolve errors in the tests due to partial objects (references field missing)

diff --git a/Engine/Source/Programs/Horde/HordeStorage/Horde.Storage/Implementation/Blob/ScyllaBlobIndex.cs b/Engine/Source/Programs/Horde/HordeStorage/Horde.Storage/Implementation/Blob/ScyllaBlobIndex.cs
--- a/Engine/Source/Programs/Horde/HordeStorage/Horde.Storage/Implementation/Blob/ScyllaBlobIndex.cs
+++ b/Engine/Source/Programs/Horde/HordeStorage/Horde.Storage/Implementation/Blob/ScyllaBlobIndex.cs
@@ -37,9 +37,17 @@ public class ScyllaBlobIndex : IBlobIndex
     public async Task AddBlobToIndex(NamespaceId ns, BlobIdentifier id)
     {
         using Scope _ = Tracer.Instance.StartActive("scylla.insert_blob_index");
+        ScyllaBlobIndexTable tableRow = new ScyllaBlobIndexTable(ns.ToString(), id,
+            new HashSet<string> { _jupiterSettings.CurrentValue.CurrentSite }, new List<ScyllaObjectReference>());
 
-        await _mapper.UpdateAsync<ScyllaBlobIndexTable>("SET regions = regions + ? WHERE namespace = ? AND blob_id = ?",
-            new string[] { _jupiterSettings.CurrentValue.CurrentSite }, ns.ToString(), new ScyllaBlobIdentifier(id));
+        AppliedInfo<ScyllaBlobIndexTable>? result = await _mapper.InsertIfNotExistsAsync(tableRow);
+
+        // if it already existed we just update the regions the blob existed in
+        if (!result.Applied)
+        {
+            await _mapper.UpdateAsync<ScyllaBlobIndexTable>("SET regions = regions + ? WHERE namespace = ? AND blob_id = ?",
+                new string[] { _jupiterSettings.CurrentValue.CurrentSite }, ns.ToString(), new ScyllaBlobIdentifier(id));
+        }
     }
 
     public async Task<IBlobIndex.BlobInfo?> GetBlobInfo(NamespaceId ns, BlobIdentifier id)
