Horde: Allow overriding the description for changes submitted to create new CL numbers on scheduled jobs, and fix issues when retrying failed submits.

diff --git a/Engine/Source/Programs/Horde/Horde.Build/Api/Templates.cs b/Engine/Source/Programs/Horde/Horde.Build/Api/Templates.cs
--- a/Engine/Source/Programs/Horde/Horde.Build/Api/Templates.cs
+++ b/Engine/Source/Programs/Horde/Horde.Build/Api/Templates.cs
@@ -425,6 +425,11 @@ namespace HordeServer.Api
 		/// </summary>
 		public string? SubmitNewChange { get; set; }
 
+		/// <summary>
+		/// Description for new changelists
+		/// </summary>
+		public string? SubmitDescription { get; set; }
+
 		/// <summary>
 		/// Fixed arguments for the new job
 		/// </summary>
diff --git a/Engine/Source/Programs/Horde/Horde.Build/Controllers/StreamsController.cs b/Engine/Source/Programs/Horde/Horde.Build/Controllers/StreamsController.cs
--- a/Engine/Source/Programs/Horde/Horde.Build/Controllers/StreamsController.cs
+++ b/Engine/Source/Programs/Horde/Horde.Build/Controllers/StreamsController.cs
@@ -264,89 +264,6 @@ namespace HordeServer.Controllers
 			return new OkResult();
 		}
 
-		/// <summary>
-		/// Validates that the user is allowed to specify the given template refs
-		/// </summary>
-		/// <param name="Requests">List of template refs</param>
-		/// <param name="OutResult">On failure, receives information about the missing claim</param>
-		/// <returns>True if the refs are valid, false otherwise</returns>
-		bool TryAuthorizeTemplateClaims(List<CreateTemplateRefRequest>? Requests, [NotNullWhen(false)] out ActionResult? OutResult)
-		{
-			if (Requests != null)
-			{
-				foreach (CreateTemplateRefRequest Request in Requests)
-				{
-					if (Request.Schedule != null && Request.Schedule.Claims != null)
-					{
-						foreach (CreateAclClaimRequest Claim in Request.Schedule.Claims)
-						{
-							if (!User.HasClaim(Claim.Type, Claim.Value))
-							{
-								OutResult = BadRequest("User does not have the {Claim.Type}: {Claim.Value} claim");
-								return false;
-							}
-						}
-					}
-				}
-			}
-
-			OutResult = null;
-			return true;
-		}
-
-		/// <summary>
-		/// Creates a list of template refs from a set of request objects
-		/// </summary>
-		/// <param name="Requests">Request objects</param>
-		/// <param name="Stream">The current stream state</param>
-		/// <param name="TemplateCollection">The template service</param>
-		/// <returns>List of new template references</returns>
-		static async Task<Dictionary<TemplateRefId, TemplateRef>> CreateTemplateRefs(List<CreateTemplateRefRequest> Requests, IStream? Stream, ITemplateCollection TemplateCollection)
-		{
-			Dictionary<TemplateRefId, TemplateRef> NewTemplateRefs = new Dictionary<TemplateRefId, TemplateRef>();
-			foreach (CreateTemplateRefRequest Request in Requests)
-			{
-				// Create the template
-				ITemplate NewTemplate = await TemplateCollection.AddAsync(Request.Name, Request.Priority, Request.AllowPreflights, Request.UpdateIssues, Request.InitialAgentType, Request.SubmitNewChange, Request.Arguments, Request.Parameters.ConvertAll(x => x.ToModel()));
-
-				// Get an identifier for the new template ref
-				TemplateRefId NewTemplateRefId;
-				if (Request.Id != null)
-				{
-					NewTemplateRefId = new TemplateRefId(Request.Id);
-				}
-				else
-				{
-					NewTemplateRefId = TemplateRefId.Sanitize(Request.Name);
-				}
-
-				// Add it to the list
-				TemplateRef NewTemplateRef = new TemplateRef(NewTemplate, Request.ShowUgsBadges, Request.ShowUgsAlerts, Request.NotificationChannel, Request.NotificationChannelFilter, Request.TriageChannel, Request.Schedule?.ToModel(), Request.ChainedJobs?.ConvertAll(x => new ChainedJobTemplate(x)), Acl.Merge(null, Request.Acl));
-				if (Stream != null && Stream.Templates.TryGetValue(NewTemplateRefId, out TemplateRef? OldTemplateRef))
-				{
-					if (OldTemplateRef.Schedule != null && NewTemplateRef.Schedule != null)
-					{
-						NewTemplateRef.Schedule.CopyState(OldTemplateRef.Schedule);
-					}
-				}
-				NewTemplateRefs.Add(NewTemplateRefId, NewTemplateRef);
-			}
-			foreach (TemplateRef TemplateRef in NewTemplateRefs.Values)
-			{
-				if (TemplateRef.ChainedJobs != null)
-				{
-					foreach (ChainedJobTemplate ChainedJob in TemplateRef.ChainedJobs)
-					{
-						if (!NewTemplateRefs.ContainsKey(ChainedJob.TemplateRefId))
-						{
-							throw new InvalidDataException($"Invalid template ref id '{ChainedJob.TemplateRefId}");
-						}
-					}
-				}
-			}
-			return NewTemplateRefs;
-		}
-
 		/// <summary>
 		/// Gets all the templates for a stream
 		/// </summary>
diff --git a/Engine/Source/Programs/Horde/Horde.Build/Jobs/ITemplate.cs b/Engine/Source/Programs/Horde/Horde.Build/Jobs/ITemplate.cs
--- a/Engine/Source/Programs/Horde/Horde.Build/Jobs/ITemplate.cs
+++ b/Engine/Source/Programs/Horde/Horde.Build/Jobs/ITemplate.cs
@@ -456,6 +456,11 @@ namespace HordeServer.Models
 		/// </summary>
 		public string? SubmitNewChange { get; }
 
+		/// <summary>
+		/// Description for new changelists
+		/// </summary>
+		public string? SubmitDescription { get; }
+
 		/// <summary>
 		/// Optional predefined user-defined properties for this job
 		/// </summary>
diff --git a/Engine/Source/Programs/Horde/Horde.Build/Jobs/ITemplateCollection.cs b/Engine/Source/Programs/Horde/Horde.Build/Jobs/ITemplateCollection.cs
--- a/Engine/Source/Programs/Horde/Horde.Build/Jobs/ITemplateCollection.cs
+++ b/Engine/Source/Programs/Horde/Horde.Build/Jobs/ITemplateCollection.cs
@@ -31,9 +31,10 @@ namespace HordeServer.Collections
 		/// <param name="bUpdateIssues"> Whether to update issues for all jobs using this template</param>
 		/// <param name="InitialAgentType">The agent type to parse the buildgraph script</param>
 		/// <param name="SubmitNewChange">Path to a file within the stream to submit to generate a new changelist for jobs</param>
+		/// <param name="SubmitDescription">Description for new changes submitted to the stream</param>
 		/// <param name="Arguments">List of arguments which are always specified</param>
 		/// <param name="Parameters">List of template parameters</param>
-		Task<ITemplate> AddAsync(string Name, Priority? Priority = null, bool bAllowPreflights = true, bool bUpdateIssues = false, string ? InitialAgentType = null, string? SubmitNewChange = null, List<string>? Arguments = null, List<Parameter>? Parameters = null);
+		Task<ITemplate> AddAsync(string Name, Priority? Priority = null, bool bAllowPreflights = true, bool bUpdateIssues = false, string? InitialAgentType = null, string? SubmitNewChange = null, string? SubmitDescription = null, List<string>? Arguments = null, List<Parameter>? Parameters = null);
 
 		/// <summary>
 		/// Gets all the available templates
diff --git a/Engine/Source/Programs/Horde/Horde.Build/Jobs/Impl/TemplateCollection.cs b/Engine/Source/Programs/Horde/Horde.Build/Jobs/Impl/TemplateCollection.cs
--- a/Engine/Source/Programs/Horde/Horde.Build/Jobs/Impl/TemplateCollection.cs
+++ b/Engine/Source/Programs/Horde/Horde.Build/Jobs/Impl/TemplateCollection.cs
@@ -42,6 +42,9 @@ namespace HordeServer.Collections.Impl
 			[BsonIgnoreIfNull]
 			public string? SubmitNewChange { get; set; }
 
+			[BsonIgnoreIfNull]
+			public string? SubmitDescription { get; set; }
+
 			public List<string> Arguments { get; private set; } = new List<string>();
 			public List<Parameter> Parameters { get; private set; } = new List<Parameter>();
 
@@ -54,7 +57,7 @@ namespace HordeServer.Collections.Impl
 				Name = null!;
 			}
 
-			public TemplateDocument(string Name, Priority? Priority, bool bAllowPreflights, bool UpdateIssues, string? InitialAgentType, string? SubmitNewChange, List<string>? Arguments, List<Parameter>? Parameters)
+			public TemplateDocument(string Name, Priority? Priority, bool bAllowPreflights, bool UpdateIssues, string? InitialAgentType, string? SubmitNewChange, string? SubmitDescription, List<string>? Arguments, List<Parameter>? Parameters)
 			{
 				this.Name = Name;
 				this.Priority = Priority;
@@ -62,6 +65,7 @@ namespace HordeServer.Collections.Impl
 				this.UpdateIssues = UpdateIssues;
 				this.InitialAgentType = InitialAgentType;
 				this.SubmitNewChange = SubmitNewChange;
+				this.SubmitDescription = SubmitDescription;
 				this.Arguments = Arguments ?? new List<string>();
 				this.Parameters = Parameters ?? new List<Parameter>();
 
@@ -105,9 +109,9 @@ namespace HordeServer.Collections.Impl
 		}
 
 		/// <inheritdoc/>
-		public async Task<ITemplate> AddAsync(string Name, Priority? Priority, bool bAllowPreflights, bool bUpdateIssues, string? InitialAgentType, string? SubmitNewChange, List<string>? Arguments, List<Parameter>? Parameters)
+		public async Task<ITemplate> AddAsync(string Name, Priority? Priority, bool bAllowPreflights, bool bUpdateIssues, string? InitialAgentType, string? SubmitNewChange, string? SubmitDescription, List<string>? Arguments, List<Parameter>? Parameters)
 		{
-			TemplateDocument Template = new TemplateDocument(Name, Priority, bAllowPreflights, bUpdateIssues, InitialAgentType, SubmitNewChange, Arguments, Parameters);
+			TemplateDocument Template = new TemplateDocument(Name, Priority, bAllowPreflights, bUpdateIssues, InitialAgentType, SubmitNewChange, SubmitDescription, Arguments, Parameters);
 			if (await GetAsync(Template.Id) == null)
 			{
 				await Templates.ReplaceOneAsync(x => x.Id == Template.Id, Template, new ReplaceOptions { IsUpsert = true });
diff --git a/Engine/Source/Programs/Horde/Horde.Build/Perforce/IPerforceService.cs b/Engine/Source/Programs/Horde/Horde.Build/Perforce/IPerforceService.cs
--- a/Engine/Source/Programs/Horde/Horde.Build/Perforce/IPerforceService.cs
+++ b/Engine/Source/Programs/Horde/Horde.Build/Perforce/IPerforceService.cs
@@ -132,8 +132,9 @@ namespace HordeServer.Services
 		/// <param name="ClusterName">Name of the Perforce cluster</param>
 		/// <param name="StreamName">The stream to query</param>
 		/// <param name="Path">Path for the file to submit</param>
+		/// <param name="Description">Description for the changelist</param>
 		/// <returns>New changelist number</returns>
-		public Task<int> CreateNewChangeAsync(string ClusterName, string StreamName, string Path);
+		public Task<int> CreateNewChangeAsync(string ClusterName, string StreamName, string Path, string Description);
 
 		/// <summary>
 		/// Gets the code change corresponding to an actual change submitted to a stream
@@ -276,14 +277,16 @@ namespace HordeServer.Services
 		/// <returns>New changelist number</returns>
 		public static Task<int> CreateNewChangeForTemplateAsync(this IPerforceService Perforce, IStream Stream, ITemplate Template)
 		{
+			string Description = (Template.SubmitDescription ?? "[Horde] New change for $(TemplateName)").Replace("$(TemplateName)", Template.Name, StringComparison.OrdinalIgnoreCase);
+
 			Match Match = Regex.Match(Template.SubmitNewChange!, @"^(//[^/]+/[^/]+)/(.+)$");
 			if (Match.Success)
 			{
-				return Perforce.CreateNewChangeAsync(Stream.ClusterName, Match.Groups[1].Value, Match.Groups[2].Value);
+				return Perforce.CreateNewChangeAsync(Stream.ClusterName, Match.Groups[1].Value, Match.Groups[2].Value, Description);
 			}
 			else
 			{
-				return Perforce.CreateNewChangeAsync(Stream.ClusterName, Stream.Name, Template.SubmitNewChange!);
+				return Perforce.CreateNewChangeAsync(Stream.ClusterName, Stream.Name, Template.SubmitNewChange!, Description);
 			}
 		}
 
diff --git a/Engine/Source/Programs/Horde/Horde.Build/Perforce/Impl/PerforceService.cs b/Engine/Source/Programs/Horde/Horde.Build/Perforce/Impl/PerforceService.cs
--- a/Engine/Source/Programs/Horde/Horde.Build/Perforce/Impl/PerforceService.cs
+++ b/Engine/Source/Programs/Horde/Horde.Build/Perforce/Impl/PerforceService.cs
@@ -461,6 +461,17 @@ namespace HordeServer.Services
 
 		}
 
+		static void ResetClient(P4.Repository Repository)
+		{
+			Repository.Connection.Client.RevertFiles(new P4.Options(), new P4.DepotPath("//..."));
+
+			IList<P4.Changelist> Changes = Repository.GetChangelists(new P4.ChangesCmdOptions(P4.ChangesCmdFlags.None, Repository.Connection.Client.Name, 100, P4.ChangeListStatus.Pending, null));
+			foreach (P4.Changelist Change in Changes)
+			{
+				Repository.DeleteChangelist(Change, new P4.Options());
+			}
+		}
+
 		static P4.Client GetOrCreateClient(string? ServiceUserName, P4.Repository Repository, string? Stream, string? Username = null, bool ReadOnly = true, bool CreateChange = false, int? ClientFromChange = null, bool UseClientFromChange = false, bool UsePortFromChange = false)
 		{
 			using IScope Scope = GlobalTracer.Instance.BuildSpan("PerforceService.GetOrCreateClient").StartActive();
@@ -1223,7 +1234,7 @@ namespace HordeServer.Services
 		}
 
 		/// <inheritdoc/>
-		public async Task<int> CreateNewChangeAsync(string ClusterName, string StreamName, string FilePath)
+		public async Task<int> CreateNewChangeAsync(string ClusterName, string StreamName, string FilePath, string Description)
 		{
 			using IScope Scope = GlobalTracer.Instance.BuildSpan("PerforceService.CreateNewChangeAsync").StartActive();
 			Scope.Span.SetTag("ClusterName", ClusterName);
@@ -1247,11 +1258,13 @@ namespace HordeServer.Services
 				P4.Changelist? SubmitChangelist = null;
 				P4.DepotPath? DepotPath = null;
 
-				const int MaxRetries = 10;
+				const int MaxRetries = 5;
 				int Retry = 0;
 
 				for (; ; )
 				{
+					ResetClient(Repository);
+
 					DepotPath = null;
 
 					if (Retry == MaxRetries)
@@ -1333,17 +1346,14 @@ namespace HordeServer.Services
 						}
 
 						// create a new change
+						P4.Changelist Changelist = new P4.Changelist();
+						Changelist.Description = Description;
+						Changelist.Files.Add(new P4.FileMetaData(new P4.FileSpec(DepotPath)));
+						SubmitChangelist = Repository.CreateChangelist(Changelist);
+
 						if (SubmitChangelist == null)
 						{
-							P4.Changelist Changelist = new P4.Changelist();
-							Changelist.Description = "New change for Horde job";
-							Changelist.Files.Add(new P4.FileMetaData(new P4.FileSpec(DepotPath)));
-							SubmitChangelist = Repository.CreateChangelist(Changelist);
-
-							if (SubmitChangelist == null)
-							{
-								throw new Exception($"Unable to create a changelist for: {DepotPath}");
-							}
+							throw new Exception($"Unable to create a changelist for: {DepotPath}");
 						}
 
 						SubmitResults = SubmitChangelist.Submit(null);
@@ -1358,9 +1368,6 @@ namespace HordeServer.Services
 					catch
 					{
 						Retry++;
-
-						Client.RevertFiles(new P4.Options(), WorkspaceFileSpec);
-
 						continue;
 					}
 
diff --git a/Engine/Source/Programs/Horde/Horde.Build/Projects/Impl/StreamCollection.cs b/Engine/Source/Programs/Horde/Horde.Build/Projects/Impl/StreamCollection.cs
--- a/Engine/Source/Programs/Horde/Horde.Build/Projects/Impl/StreamCollection.cs
+++ b/Engine/Source/Programs/Horde/Horde.Build/Projects/Impl/StreamCollection.cs
@@ -172,7 +172,7 @@ namespace HordeServer.Collections.Impl
 			foreach (CreateTemplateRefRequest Request in Requests)
 			{
 				// Create the template
-				ITemplate NewTemplate = await TemplateCollection.AddAsync(Request.Name, Request.Priority, Request.AllowPreflights, Request.UpdateIssues, Request.InitialAgentType, Request.SubmitNewChange, Request.Arguments, Request.Parameters.ConvertAll(x => x.ToModel()));
+				ITemplate NewTemplate = await TemplateCollection.AddAsync(Request.Name, Request.Priority, Request.AllowPreflights, Request.UpdateIssues, Request.InitialAgentType, Request.SubmitNewChange, Request.SubmitDescription, Request.Arguments, Request.Parameters.ConvertAll(x => x.ToModel()));
 
 				// Get an identifier for the new template ref
 				TemplateRefId NewTemplateRefId;
