Horde: Format model binding errors in the same format as other controller errors.

diff --git a/Engine/Source/Programs/Horde/Horde.Build/Startup.cs b/Engine/Source/Programs/Horde/Horde.Build/Startup.cs
--- a/Engine/Source/Programs/Horde/Horde.Build/Startup.cs
+++ b/Engine/Source/Programs/Horde/Horde.Build/Startup.cs
@@ -83,6 +83,8 @@ using HordeServer.Compute;
 using System.Net.Http.Headers;
 using Serilog.Events;
 using HordeServer.Jobs;
+using Microsoft.AspNetCore.Mvc;
+using Microsoft.AspNetCore.Mvc.ModelBinding;
 
 namespace HordeServer
 {
@@ -644,6 +646,27 @@ namespace HordeServer
 				Config.IncludeXmlComments(Path.Combine(AppContext.BaseDirectory, $"{Assembly.GetExecutingAssembly().GetName().Name}.xml"));
 			});
 
+			Services.Configure<ApiBehaviorOptions>(Options =>
+			{
+				Options.InvalidModelStateResponseFactory = Context =>
+				{
+					foreach(KeyValuePair<string, ModelStateEntry> Pair in Context.ModelState)
+					{
+						ModelError? Error = Pair.Value.Errors.FirstOrDefault();
+						if (Error != null)
+						{
+							string Message = Error.ErrorMessage;
+							if (String.IsNullOrEmpty(Message))
+							{
+								Message = Error.Exception?.Message ?? "Invalid error object";
+							}
+							return new BadRequestObjectResult(EpicGames.Core.LogEvent.Create(LogLevel.Error, KnownLogEvents.None, Error.Exception, "Invalid value for {Name}: {Message}", Pair.Key, Message));
+						}
+					}
+					return new BadRequestObjectResult(Context.ModelState);
+				};
+			});
+
 			DirectoryReference DashboardDir = DirectoryReference.Combine(Program.AppDir, "DashboardApp");
 			if (DirectoryReference.Exists(DashboardDir)) 
 			{
