Horde Storage
Use cleanup script in Horde to always shutdown docker-compose and attach logs to the build log.
Added sleep after running docker-compose to let scylla finish starting up, fixes intermittent connection issues to scylla.

diff --git a/Engine/Source/Programs/Horde/BuildHorde.xml b/Engine/Source/Programs/Horde/BuildHorde.xml
--- a/Engine/Source/Programs/Horde/BuildHorde.xml
+++ b/Engine/Source/Programs/Horde/BuildHorde.xml
@@ -30,12 +30,19 @@
 	</Agent>
 
 	<Agent Name="HordeStorage" Type="Linux">
-		<Property Name="ComposeFile" Value="Engine/Source/Programs/Horde/HordeStorage/Composes/docker-compose-tests.yml"/>
+		<Property Name="ComposeFile" Value="$(RootDir)/Engine/Source/Programs/Horde/HordeStorage/Composes/docker-compose-tests.yml"/>
 		<Node Name="Test HordeStorage">
+			<!-- Make sure to run docker compose down after the build no matter what, also copy the logs to AutomationTool logs so we get the logs in Horde -->
+			<WriteTextFile File="$(RootDir)/Cleanup.sh" Text="docker-compose -f $(ComposeFile) logs --no-color > $(RootDir)/Engine/Programs/AutomationTool/Saved/Logs/docker-compose-logs.txt"/>
+			<WriteTextFile File="$(RootDir)/Cleanup.sh" Append="true" Text="docker-compose -f $(ComposeFile) down"/>
+
 			<Docker-Compose-Up File="$(ComposeFile)"/>
+
+			<!-- Sleep a while to let scylla finish starting up before we execute the tests -->
+			<Spawn Exe="/bin/bash" Arguments="-c 'sleep 30'" />
+
 			<!-- Run the unit tests, and override the dotnet sdk used to the global one for the machine rather then what is in the sync, to allow us to use newer versions then the rest of the engine has migrated to -->
 			<DotNet Arguments="test Engine/Source/Programs/Horde/HordeStorage/Horde.Storage.sln" DotNetPath="/usr/share/dotnet/dotnet"/>
-			<Docker-Compose-Down File="$(ComposeFile)"/>
 		</Node>
 
 		<Node Name="Build HordeStorage" Requires="Test HordeStorage">
