Horde: Fix server error querying for negative changelist numbers in issue service.

diff --git a/Engine/Source/Programs/Horde/Horde.Build/Issues/Impl/IssueService.cs b/Engine/Source/Programs/Horde/Horde.Build/Issues/Impl/IssueService.cs
--- a/Engine/Source/Programs/Horde/Horde.Build/Issues/Impl/IssueService.cs
+++ b/Engine/Source/Programs/Horde/Horde.Build/Issues/Impl/IssueService.cs
@@ -1034,7 +1034,7 @@ namespace HordeServer.Services.Impl
 		async ValueTask<bool> ContainsFixChange(StreamId StreamId, int FixChange, Dictionary<(StreamId, int), bool> CachedContainsFixChange)
 		{
 			bool bContainsFixChange;
-			if (!CachedContainsFixChange.TryGetValue((StreamId, FixChange), out bContainsFixChange))
+			if (!CachedContainsFixChange.TryGetValue((StreamId, FixChange), out bContainsFixChange) && FixChange > 0)
 			{
 				IStream? Stream = await Streams.GetCachedStream(StreamId);
 				if (Stream != null)
