Fix test compile error.

diff --git a/Engine/Source/Programs/Horde/Horde.Build.Tests/Stubs/Services/PerforceServiceStub.cs b/Engine/Source/Programs/Horde/Horde.Build.Tests/Stubs/Services/PerforceServiceStub.cs
--- a/Engine/Source/Programs/Horde/Horde.Build.Tests/Stubs/Services/PerforceServiceStub.cs
+++ b/Engine/Source/Programs/Horde/Horde.Build.Tests/Stubs/Services/PerforceServiceStub.cs
@@ -160,9 +160,9 @@ namespace HordeServerTests.Stubs.Services
 			return Task.FromResult(CodeChange);
 		}
 
-		public Task<int> CreateNewChangeAsync(string ClusterName, string StreamName, string Path)
+		public Task<int> CreateNewChangeAsync(string ClusterName, string StreamName, string Path, string Description)
 		{
-			ChangeDetails NewChange = new ChangeDetails(Changes[StreamName].First().Key + 1, TestUser, null!, "", new List<ChangeFile> { PerforceExtensions.CreateChangeFile(Path) }, DateTime.Now);
+			ChangeDetails NewChange = new ChangeDetails(Changes[StreamName].First().Key + 1, TestUser, null!, Description, new List<ChangeFile> { PerforceExtensions.CreateChangeFile(Path) }, DateTime.Now);
 			Changes[StreamName].Add(NewChange.Number, NewChange);
 			return Task.FromResult(NewChange.Number);
 		}
