Horde: New build targets and controller json log errors
#rnx

diff --git a/Engine/Source/Programs/Horde/HordeDashboard/src/backend/Api.ts b/Engine/Source/Programs/Horde/HordeDashboard/src/backend/Api.ts
--- a/Engine/Source/Programs/Horde/HordeDashboard/src/backend/Api.ts
+++ b/Engine/Source/Programs/Horde/HordeDashboard/src/backend/Api.ts
@@ -1867,6 +1867,9 @@ export enum ParameterType {
 /**Base class for template parameters */
 export type ParameterData = {
 	type: ParameterType;
+
+	// client side key, with group encoding
+	parameterKey?: string;
 }
 
 /**Used to group a number of other parameters */
diff --git a/Engine/Source/Programs/Horde/HordeDashboard/src/backend/Fetch.ts b/Engine/Source/Programs/Horde/HordeDashboard/src/backend/Fetch.ts
--- a/Engine/Source/Programs/Horde/HordeDashboard/src/backend/Fetch.ts
+++ b/Engine/Source/Programs/Horde/HordeDashboard/src/backend/Fetch.ts
@@ -21,7 +21,7 @@ export type FetchResponse = {
     status: number;
 }
 
-function handleError(errorIn: ErrorInfo, force?:boolean) {
+function handleError(errorIn: ErrorInfo, force?: boolean) {
 
     const e = { ...errorIn, hordeError: true } as ErrorInfo;
 
@@ -56,7 +56,7 @@ export class Fetch {
                 const challenge = await this.challenge();
 
                 if (challenge === ChallengeStatus.Unauthorized) {
-                    
+
                     this.login(window.location.toString());
                     return;
                 }
@@ -181,6 +181,8 @@ export class Fetch {
 
         if (!response.ok) {
 
+            let message = response.statusText;
+
             if (response.url?.indexOf("AccessDenied") !== -1) {
 
                 handleError({
@@ -193,15 +195,34 @@ export class Fetch {
 
                 if (response.status !== 404 || !config?.suppress404) {
 
-                    handleError({
+                    let errorInfo: ErrorInfo = {
                         response: response,
                         url: url
-                    });
+                    }
+
+                    let json: any = undefined;
+
+                    try {
+                        json = await response.json();
+                    } catch {
+
+                    }
+
+                    // dynamic detection of horde formatted error
+                    if (json && json.time && json.message && json.level) {
+                        errorInfo.format = json;
+                        message = json.message;
+                        if (json.id) {
+                            message = `(Error ${json.id}) - ${message}`;
+                        }
+                    }
+
+                    handleError(errorInfo);
 
                 }
             }
 
-            reject(response.statusText);
+            reject(message);
             return;
         }
 
@@ -354,7 +375,7 @@ export class Fetch {
                         url: url
                     });
                 }
-        
+
                 return ChallengeStatus.Unauthorized;
             }
         } catch (reason) {
@@ -375,7 +396,7 @@ export class Fetch {
         if (!token) {
 
             this.debugToken = undefined;
-            this.setAuthorization(undefined);    
+            this.setAuthorization(undefined);
             return;
         }
 
@@ -391,10 +412,10 @@ export class Fetch {
             return;
         }
 
-        this.baseUrl = url;        
+        this.baseUrl = url;
     }
 
-    logout:boolean = false;
+    logout: boolean = false;
 
     private baseUrl = "";
     private debugToken?: string;
diff --git a/Engine/Source/Programs/Horde/HordeDashboard/src/backend/JobDetails.ts b/Engine/Source/Programs/Horde/HordeDashboard/src/backend/JobDetails.ts
--- a/Engine/Source/Programs/Horde/HordeDashboard/src/backend/JobDetails.ts
+++ b/Engine/Source/Programs/Horde/HordeDashboard/src/backend/JobDetails.ts
@@ -33,6 +33,32 @@ export class JobDetails {
         return this.jobdata?.name;
     }
 
+    get targets(): string[] {
+
+        const detailArgs = this?.jobdata?.arguments;
+
+        if (!detailArgs) {
+            return [];
+        }
+
+        const targets: string[] = [];
+
+        if (detailArgs && detailArgs.length) {
+
+            const targetArgs = detailArgs.map(arg => arg.trim()).filter(arg => arg.toLowerCase().startsWith("-target="));
+
+            targetArgs.forEach(t => {
+                const target = t.split("=")[1]?.trim();
+                if (target) {
+                    targets.push(target);
+                }
+            });
+        }
+
+        return targets;
+
+    }
+
     async set(id: string, logId: string | undefined = undefined, stepId: string | undefined = undefined, labelIdx: number | undefined, updateMS = defaultUpdateMS, callback?: (details: JobDetails) => void, eventsCallback?: (details: JobDetails) => void) {
 
         if (this.id === id) {
@@ -330,7 +356,7 @@ export class JobDetails {
 
     }
 
-    getStepName(stepId: string | undefined, includeRetry:boolean = true): string {
+    getStepName(stepId: string | undefined, includeRetry: boolean = true): string {
 
         if (!stepId) {
             return "";
@@ -734,7 +760,7 @@ export class JobDetails {
                 this.updateTimingInfo++;
                 this.updateTimingInfo %= 3;
             }
-            
+
             if (doTiming) {
                 requests.push(backend.getJobTiming(this.id));
             }
