Horde: Fix test build errors.

diff --git a/Engine/Source/Programs/Horde/Horde.Build.Tests/Fixture.cs b/Engine/Source/Programs/Horde/Horde.Build.Tests/Fixture.cs
--- a/Engine/Source/Programs/Horde/Horde.Build.Tests/Fixture.cs
+++ b/Engine/Source/Programs/Horde/Horde.Build.Tests/Fixture.cs
@@ -56,8 +56,7 @@ namespace HordeServerTests
 			Fg.Aggregates = new List<IAggregate>();
 			Fg.Labels = new List<ILabel>();
 
-			Template = await TemplateCollection.AddAsync("Test template", null, false, false, null, null,
-				new List<string>(), new List<Parameter>());
+			Template = await TemplateCollection.AddAsync("Test template");
 			Graph = await GraphCollection.AddAsync(Template);
 
 			TemplateRefId1 = new TemplateRefId("template1");
diff --git a/Engine/Source/Programs/Horde/Horde.Build.Tests/JobServiceTests.cs b/Engine/Source/Programs/Horde/Horde.Build.Tests/JobServiceTests.cs
--- a/Engine/Source/Programs/Horde/Horde.Build.Tests/JobServiceTests.cs
+++ b/Engine/Source/Programs/Horde/Horde.Build.Tests/JobServiceTests.cs
@@ -30,7 +30,7 @@ namespace HordeServerTests
 			IProject? Project = await ProjectService.Collection.AddOrUpdateAsync(ProjectId, "", "", 0, new ProjectConfig { Name = "UE5" });
 			Assert.IsNotNull(Project);
 
-			ITemplate Template = await TemplateCollection.AddAsync("Test template", null, false, false, null, null, new List<string>(), new List<Parameter>());
+			ITemplate Template = await TemplateCollection.AddAsync("Test template");
 			IGraph Graph = await GraphCollection.AddAsync(Template);
 
 			TemplateRefId TemplateRefId1 = new TemplateRefId("template1");
@@ -153,7 +153,7 @@ namespace HordeServerTests
 			IStream? Stream = await StreamCollection.GetAsync(StreamId);
 			Stream = await StreamCollection.TryCreateOrReplaceAsync(StreamId, Stream, "", "", Project!.Id, new StreamConfig { Name = "//UE5/Main" });
 
-			ITemplate Template = await TemplateCollection.AddAsync("Test template", null, false, false, null, null, new List<string>(), new List<Parameter>());
+			ITemplate Template = await TemplateCollection.AddAsync("Test template");
 			IGraph Graph = await GraphCollection.AddAsync(Template);
 
 			NewGroup GroupA = new NewGroup("win", new List<NewNode>());
diff --git a/Engine/Source/Programs/Horde/Horde.Build.Tests/SchedulerTests.cs b/Engine/Source/Programs/Horde/Horde.Build.Tests/SchedulerTests.cs
--- a/Engine/Source/Programs/Horde/Horde.Build.Tests/SchedulerTests.cs
+++ b/Engine/Source/Programs/Horde/Horde.Build.Tests/SchedulerTests.cs
@@ -42,7 +42,7 @@ namespace HordeServerTests
 			IProject ? Project = ProjectService.Collection.AddOrUpdateAsync(ProjectId, "", "", 0, new ProjectConfig { Name = "UE5" }).Result;
 			Assert.IsNotNull(Project);
 
-			Template = TemplateCollection.AddAsync("Test template", null, false, false, null, null, new List<string>(), new List<Parameter>()).Result;
+			Template = TemplateCollection.AddAsync("Test template").Result;
 
 			InitialJobIds = new HashSet<JobId>(JobCollection.FindAsync().Result.Select(x => x.Id));
 
