Horde: Increase max request size to 100mb.

diff --git a/Engine/Source/Programs/Horde/Horde.Build/Program.cs b/Engine/Source/Programs/Horde/Horde.Build/Program.cs
--- a/Engine/Source/Programs/Horde/Horde.Build/Program.cs
+++ b/Engine/Source/Programs/Horde/Horde.Build/Program.cs
@@ -240,6 +240,8 @@ namespace HordeServer
 				{
 					WebBuilder.ConfigureKestrel(Options =>
 					{
+						Options.Limits.MaxRequestBodySize = 100 * 1024 * 1024;
+
 						if (ServerSettings.HttpPort != 0)
 						{
 							Options.ListenAnyIP(ServerSettings.HttpPort, Configure => { Configure.Protocols = HttpProtocols.Http1AndHttp2; });
