Updated Namespace auth handler to use the new namespace settings
Fixes issues were dev crashes on null dereference
(this file was supposed to have been submitted in my earlier changes)

diff --git a/Engine/Source/Programs/Horde/HordeStorage/Jupiter.Common/Authentication/NamespaceAuthorizationHandler.cs b/Engine/Source/Programs/Horde/HordeStorage/Jupiter.Common/Authentication/NamespaceAuthorizationHandler.cs
--- a/Engine/Source/Programs/Horde/HordeStorage/Jupiter.Common/Authentication/NamespaceAuthorizationHandler.cs
+++ b/Engine/Source/Programs/Horde/HordeStorage/Jupiter.Common/Authentication/NamespaceAuthorizationHandler.cs
@@ -1,5 +1,6 @@
 ﻿// Copyright Epic Games, Inc. All Rights Reserved.
 
+using System;
 using System.Collections.Generic;
 using System.ComponentModel.DataAnnotations;
 using System.Threading.Tasks;
@@ -12,11 +13,11 @@ namespace Jupiter
     // verifies that you have access to a namespace by checking if you have a corresponding claim to that namespace
     public class NamespaceAuthorizationHandler : AuthorizationHandler<NamespaceAccessRequirement, NamespaceId>
     {
-        private readonly IOptionsMonitor<AuthorizationSettings> _authorizationSettings;
+        private readonly IOptionsMonitor<NamespaceSettings> _namespaceSettings;
 
-        public NamespaceAuthorizationHandler(IOptionsMonitor<AuthorizationSettings> authorizationSettings)
+        public NamespaceAuthorizationHandler(IOptionsMonitor<NamespaceSettings> namespaceSettings)
         {
-            _authorizationSettings = authorizationSettings;
+            _namespaceSettings = namespaceSettings;
         }
 
         protected override Task HandleRequirementAsync(AuthorizationHandlerContext context, NamespaceAccessRequirement requirement,
@@ -28,7 +29,9 @@ namespace Jupiter
                 return Task.CompletedTask;
             }
 
-            if (_authorizationSettings.CurrentValue.NamespaceToClaim.TryGetValue(namespaceName.ToString(), out string? expectedClaim))
+            NamespaceSettings.PerNamespaceSettings settings = _namespaceSettings.CurrentValue.GetPoliciesForNs(namespaceName);
+            // These are ANDed, e.g. all claims needs to be present
+            foreach (string expectedClaim in settings.Claims)
             {
                 // if expected claim is * then everyone is allowed to use the namespace
                 if (expectedClaim == "*")
@@ -42,16 +45,10 @@ namespace Jupiter
                 }
             }
 
-
             return Task.CompletedTask;
         }
     }
 
-    public class AuthorizationSettings
-    {
-        [Required] public Dictionary<string, string> NamespaceToClaim { get; set; } = null!;
-    }
-
     public class NamespaceAccessRequirement : IAuthorizationRequirement
     {
         public const string Name = "NamespaceAccess";
