Horde Storage
Added filter for which namespaces to clean in the orphan gc step. Lets us gc namespaces that are used with the old api.

diff --git a/Engine/Source/Programs/Horde/HordeStorage/Horde.Storage/Implementation/GC/OrphanBlobCleanup.cs b/Engine/Source/Programs/Horde/HordeStorage/Horde.Storage/Implementation/GC/OrphanBlobCleanup.cs
--- a/Engine/Source/Programs/Horde/HordeStorage/Horde.Storage/Implementation/GC/OrphanBlobCleanup.cs
+++ b/Engine/Source/Programs/Horde/HordeStorage/Horde.Storage/Implementation/GC/OrphanBlobCleanup.cs
@@ -36,14 +36,16 @@ namespace Horde.Storage.Implementation
     {
         private readonly IBlobService _blobService;
         private readonly ILeaderElection _leaderElection;
+        private readonly IOptionsMonitor<GCSettings> _gcSettings;
         private readonly IRestClient _client;
         private readonly ILogger _logger = Log.ForContext<OrphanBlobCleanup>();
 
         // ReSharper disable once UnusedMember.Global
-        public OrphanBlobCleanup(IBlobService blobService, ILeaderElection leaderElection, IOptionsMonitor<CallistoTransactionLogSettings> callistoSettings, IServiceCredentials serviceCredentials)
+        public OrphanBlobCleanup(IBlobService blobService, ILeaderElection leaderElection, IOptionsMonitor<CallistoTransactionLogSettings> callistoSettings, IOptionsMonitor<GCSettings> gcSettings, IServiceCredentials serviceCredentials)
         {
             _blobService = blobService;
             _leaderElection = leaderElection;
+            _gcSettings = gcSettings;
 
             _client = new RestClient(callistoSettings.CurrentValue.ConnectionString)
             {
@@ -51,11 +53,12 @@ namespace Horde.Storage.Implementation
             }.UseSerializer(() => new JsonNetSerializer());
         }
 
-        internal OrphanBlobCleanup(IBlobService blobService, ILeaderElection leaderElection, IRestClient callistoClient)
+        internal OrphanBlobCleanup(IBlobService blobService, ILeaderElection leaderElection, IRestClient callistoClient, IOptionsMonitor<GCSettings> gcSettings)
         {
             _blobService = blobService;
             _leaderElection = leaderElection;
             _client = callistoClient;
+            _gcSettings = gcSettings;
         }
 
         private struct GCRootState
@@ -81,7 +84,7 @@ namespace Horde.Storage.Implementation
                 return new List<RemovedBlobs>();
             }
 
-            List<NamespaceId> namespaces = await ListNamespaces().ToListAsync();
+            List<NamespaceId> namespaces = await ListNamespaces().Where(NamespaceShouldBeCleaned).ToListAsync();
             ConcurrentDictionary<NamespaceId, GCRootState> perNamespaceRoots = new();
             await namespaces.ParallelForEachAsync(async ns =>
             {
@@ -160,6 +163,11 @@ namespace Horde.Storage.Implementation
             return removedBlobs;
         }
 
+        private bool NamespaceShouldBeCleaned(NamespaceId ns)
+        {
+            return _gcSettings.CurrentValue.CleanNamespacesV1.Contains(ns.ToString());
+        }
+
         public virtual async IAsyncEnumerable<NamespaceId> ListNamespaces()
         {
             RestRequest request = new RestRequest("api/v1/t/");
diff --git a/Engine/Source/Programs/Horde/HordeStorage/Horde.Storage/Settings.cs b/Engine/Source/Programs/Horde/HordeStorage/Horde.Storage/Settings.cs
--- a/Engine/Source/Programs/Horde/HordeStorage/Horde.Storage/Settings.cs
+++ b/Engine/Source/Programs/Horde/HordeStorage/Horde.Storage/Settings.cs
@@ -246,5 +246,6 @@ namespace Horde.Storage
 
         public TimeSpan BlobCleanupPollFrequency { get; set; } = TimeSpan.FromMinutes(60);
         public TimeSpan RefCleanupPollFrequency { get; set; } = TimeSpan.FromMinutes(60);
+        public List<string> CleanNamespacesV1 { get; set; } = new List<string>();
     }
 }
diff --git a/Engine/Source/Programs/Horde/HordeStorage/Horde.Storage/Tests/Functional/GC/GCBlobTests.cs b/Engine/Source/Programs/Horde/HordeStorage/Horde.Storage/Tests/Functional/GC/GCBlobTests.cs
--- a/Engine/Source/Programs/Horde/HordeStorage/Horde.Storage/Tests/Functional/GC/GCBlobTests.cs
+++ b/Engine/Source/Programs/Horde/HordeStorage/Horde.Storage/Tests/Functional/GC/GCBlobTests.cs
@@ -14,6 +14,7 @@ using Microsoft.AspNetCore.TestHost;
 using Microsoft.AspNetCore.Hosting;
 using Microsoft.Extensions.Configuration;
 using Microsoft.Extensions.DependencyInjection;
+using Microsoft.Extensions.Options;
 using Microsoft.VisualStudio.TestTools.UnitTesting;
 using RestSharp;
 using Serilog;
@@ -73,8 +74,14 @@ namespace Horde.Storage.FunctionalTests.GC
                 {
                     collection.AddSingleton<IBlobCleanup>(provider =>
                     {
+                        GCSettings gcSettings = new GCSettings()
+                        {
+                            CleanOldBlobs = true
+                        };
+                        IOptionsMonitor<GCSettings> gcSettingsMon = Mock.Of<IOptionsMonitor<GCSettings>>(_ => _.CurrentValue == gcSettings);
+
                         IBlobService blobService = provider.GetService<IBlobService>()!;
-                        return new OrphanBlobCleanup(blobService!, new LeaderElectionStub(true), _callistoBlobMock.Object);
+                        return new OrphanBlobCleanup(blobService!, new LeaderElectionStub(true), _callistoBlobMock.Object, gcSettingsMon);
                     });
 
                 })
@@ -156,7 +163,14 @@ namespace Horde.Storage.FunctionalTests.GC
         [TestMethod]
         public async Task RunBlobCleanup()
         {
-            OrphanBlobCleanup cleanup = new OrphanBlobCleanup(_blobService!, new LeaderElectionStub(true), _callistoBlobMock.Object);
+            GCSettings gcSettings = new GCSettings()
+            {
+                CleanOldBlobs = true,
+                CleanNamespacesV1 = new List<string> { TestNamespace.ToString() }
+            };
+            IOptionsMonitor<GCSettings> gcSettingsMon = Mock.Of<IOptionsMonitor<GCSettings>>(_ => _.CurrentValue == gcSettings);
+
+            OrphanBlobCleanup cleanup = new OrphanBlobCleanup(_blobService!, new LeaderElectionStub(true), _callistoBlobMock.Object, gcSettingsMon);
             List<NamespaceId> namespaces = await cleanup.ListNamespaces().ToListAsync();
             Assert.AreEqual(1, namespaces.Count);
 
