Horde: Add better time info to steps
#rnx

diff --git a/Engine/Source/Programs/Horde/HordeDashboard/src/base/utilities/timeUtils.ts b/Engine/Source/Programs/Horde/HordeDashboard/src/base/utilities/timeUtils.ts
--- a/Engine/Source/Programs/Horde/HordeDashboard/src/base/utilities/timeUtils.ts
+++ b/Engine/Source/Programs/Horde/HordeDashboard/src/base/utilities/timeUtils.ts
@@ -148,6 +148,25 @@ export const getStepFinishTime = (step: StepData): HordeTime => {
     };
 };
 
+export const getStepStartTime = (step: StepData): HordeTime => {
+
+    let displayStart = "";    
+    let serverStart = "";    
+
+    const format = dashboard.display24HourClock ? "HH:mm:ss z" : "LT z";
+
+    if (step.startTime) {
+        const end = moment(step.startTime);
+        displayStart = moment.utc(end).tz(displayTimeZone()).format(format);
+        serverStart = moment.utc(end).tz(serverTimeZone).format(format);
+    }
+
+    return {
+        display: displayStart,
+        server: serverStart
+    };
+};
+
 
 export const getLabelFinishTime = (label: JobLabel, job: JobData): HordeTime => {
 
