Horde Storage

Implemented orphan blob gc based on the blob index. This will check every blob in the blob store for a corresponding blob index, if it doesnt exist we delete the blob. If it does exist we check the object references it has, if none of them exist anymore we delete the blob.

diff --git a/Engine/Source/Programs/Horde/HordeStorage/Horde.Storage/Controllers/AdminController.cs b/Engine/Source/Programs/Horde/HordeStorage/Horde.Storage/Controllers/AdminController.cs
--- a/Engine/Source/Programs/Horde/HordeStorage/Horde.Storage/Controllers/AdminController.cs
+++ b/Engine/Source/Programs/Horde/HordeStorage/Horde.Storage/Controllers/AdminController.cs
@@ -73,23 +73,7 @@ namespace Horde.Storage.Controllers
                 records.Select(r => new RemovedRefRecordsResponse.RemovedRecord(r.RefName, r.Bucket))
             ));
         }
-
-        /// <summary>
-        /// Manually run the blob cleanup
-        /// </summary>
-        /// <remarks>
-        /// Manually triggers a cleanup of unused blobs not referenced in the transaction log. This is done automatically so the only reason to use this endpoint is for debugging purposes.
-        /// </remarks>
-        /// <returns></returns>
-        [HttpPost("blobCleanup/{ns}")]
-        public async Task<IActionResult> BlobCleanup([FromRoute] [Required] NamespaceId ns)
-        {
-            List<RemovedBlobs> records = await _blobCleanupService.Cleanup(_blobCleanupService.State, CancellationToken.None);
-            return Ok(new RemovedBlobRecords(
-                records.Select(r => r.BlobIdentifier)
-            ));
-        }
-
+        
         /// <summary>
         /// Dumps all settings currently in use
         /// </summary>
diff --git a/Engine/Source/Programs/Horde/HordeStorage/Horde.Storage/Horde.Storage.Startup.cs b/Engine/Source/Programs/Horde/HordeStorage/Horde.Storage/Horde.Storage.Startup.cs
--- a/Engine/Source/Programs/Horde/HordeStorage/Horde.Storage/Horde.Storage.Startup.cs
+++ b/Engine/Source/Programs/Horde/HordeStorage/Horde.Storage/Horde.Storage.Startup.cs
@@ -216,6 +216,7 @@ namespace Horde.Storage
             services.AddSingleton<MemoryBlobStore>();
 
             services.AddSingleton<OrphanBlobCleanup>();
+            services.AddSingleton<OrphanBlobCleanupRefs>();
 
             services.AddSingleton(typeof(IBlobService), typeof(BlobService));
 
diff --git a/Engine/Source/Programs/Horde/HordeStorage/Horde.Storage/Implementation/Blob/FileSystemStore.cs b/Engine/Source/Programs/Horde/HordeStorage/Horde.Storage/Implementation/Blob/FileSystemStore.cs
--- a/Engine/Source/Programs/Horde/HordeStorage/Horde.Storage/Implementation/Blob/FileSystemStore.cs
+++ b/Engine/Source/Programs/Horde/HordeStorage/Horde.Storage/Implementation/Blob/FileSystemStore.cs
@@ -185,7 +185,7 @@ namespace Horde.Storage.Implementation
             }
         }
         
-        public async Task<List<RemovedBlobs>> Cleanup(CancellationToken cancellationToken)
+        public async Task<List<BlobIdentifier>> Cleanup(CancellationToken cancellationToken)
         {
             return await CleanupInternal(cancellationToken);
         }
@@ -198,13 +198,13 @@ namespace Horde.Storage.Implementation
         /// <param name="cancellationToken">Cancellation token</param>
         /// <param name="batchSize">Number of files to scan for clean up. A higher number is recommended since blob store can contain many small blobs</param>
         /// <returns></returns>
-        public async Task<List<RemovedBlobs>> CleanupInternal(CancellationToken cancellationToken, int batchSize = 100000)
+        public async Task<List<BlobIdentifier>> CleanupInternal(CancellationToken cancellationToken, int batchSize = 100000)
         {
             long size = await CalculateDiskSpaceUsed();
             ulong maxSizeBytes = _settings.CurrentValue.MaxSizeBytes;
             long triggerSize = (long) (maxSizeBytes * _settings.CurrentValue.TriggerThresholdPercentage);
             long targetSize = (long) (maxSizeBytes * _settings.CurrentValue.TargetThresholdPercentage); // Target to shrink to if triggered
-            List<RemovedBlobs> blobsRemoved = new List<RemovedBlobs>(batchSize);
+            List<BlobIdentifier> blobsRemoved = new List<BlobIdentifier>(batchSize);
 
             if (size < triggerSize)
             {
@@ -231,7 +231,7 @@ namespace Horde.Storage.Implementation
                     {
                         totalBytesDeleted += fi.Length;
                         fi.Delete();
-                        blobsRemoved.Add(new RemovedBlobs(new BlobIdentifier(fi.Name)));
+                        blobsRemoved.Add(new BlobIdentifier(fi.Name));
 
                         long currentSize = size - totalBytesDeleted;
                         if (currentSize <= targetSize || cancellationToken.IsCancellationRequested)
diff --git a/Engine/Source/Programs/Horde/HordeStorage/Horde.Storage/Implementation/GC/BlobCleanupService.cs b/Engine/Source/Programs/Horde/HordeStorage/Horde.Storage/Implementation/GC/BlobCleanupService.cs
--- a/Engine/Source/Programs/Horde/HordeStorage/Horde.Storage/Implementation/GC/BlobCleanupService.cs
+++ b/Engine/Source/Programs/Horde/HordeStorage/Horde.Storage/Implementation/GC/BlobCleanupService.cs
@@ -7,6 +7,7 @@ using System.Threading.Tasks;
 using Jupiter;
 using Microsoft.Extensions.Options;
 using Datadog.Trace;
+using Jupiter.Implementation;
 using Microsoft.Extensions.DependencyInjection;
 using Serilog;
 
@@ -39,6 +40,14 @@ namespace Horde.Storage.Implementation
                 OrphanBlobCleanup orphanBlobCleanup = provider.GetService<OrphanBlobCleanup>()!;
                 RegisterCleanup(orphanBlobCleanup);
             }
+
+            
+            if (settings.CurrentValue.CleanOldBlobs)
+            {
+                OrphanBlobCleanupRefs orphanBlobCleanupRefs = provider.GetService<OrphanBlobCleanupRefs>()!;
+                RegisterCleanup(orphanBlobCleanupRefs);
+            }
+            
         }
 
         public void RegisterCleanup(IBlobCleanup cleanup)
@@ -63,9 +72,9 @@ namespace Horde.Storage.Implementation
             }
         }
 
-        public async Task<List<RemovedBlobs>> Cleanup(BlobCleanupState state, CancellationToken cancellationToken)
+        public async Task<List<BlobIdentifier>> Cleanup(BlobCleanupState state, CancellationToken cancellationToken)
         {
-            List<RemovedBlobs> removedBlobs = new List<RemovedBlobs>();
+            List<BlobIdentifier> removedBlobs = new List<BlobIdentifier>();
             
             foreach (IBlobCleanup blobCleanup in state.BlobCleanups)
             {
@@ -77,7 +86,7 @@ namespace Horde.Storage.Implementation
                 _logger.Information("Attempting to run Blob Cleanup {BlobCleanup}. ", type);
                 try
                 {
-                    List<RemovedBlobs> tempBlobs = await blobCleanup.Cleanup(cancellationToken);
+                    List<BlobIdentifier> tempBlobs = await blobCleanup.Cleanup(cancellationToken);
                     _logger.Information("Ran blob cleanup {BlobCleanup}. Deleted {CountBlobRecords}", type, tempBlobs.Count);
                     removedBlobs.AddRange(tempBlobs);
                 }
diff --git a/Engine/Source/Programs/Horde/HordeStorage/Horde.Storage/Implementation/GC/OrphanBlobCleanup.cs b/Engine/Source/Programs/Horde/HordeStorage/Horde.Storage/Implementation/GC/OrphanBlobCleanup.cs
--- a/Engine/Source/Programs/Horde/HordeStorage/Horde.Storage/Implementation/GC/OrphanBlobCleanup.cs
+++ b/Engine/Source/Programs/Horde/HordeStorage/Horde.Storage/Implementation/GC/OrphanBlobCleanup.cs
@@ -17,19 +17,9 @@ using Serilog;
 
 namespace Horde.Storage.Implementation
 {
-    public class RemovedBlobs
-    {
-        public RemovedBlobs(BlobIdentifier blobIdentifier)
-        {
-            BlobIdentifier = blobIdentifier;
-        }
-
-        public BlobIdentifier BlobIdentifier { get; set; }
-    }
-
     public interface IBlobCleanup
     {
-        Task<List<RemovedBlobs>> Cleanup(CancellationToken none);
+        Task<List<BlobIdentifier>> Cleanup(CancellationToken none);
     }
 
     public class OrphanBlobCleanup : IBlobCleanup
@@ -76,12 +66,12 @@ namespace Horde.Storage.Implementation
 
         }
 
-        public async Task<List<RemovedBlobs>> Cleanup(CancellationToken cancellationToken)
+        public async Task<List<BlobIdentifier>> Cleanup(CancellationToken cancellationToken)
         {
             if (!_leaderElection.IsThisInstanceLeader())
             {
                 _logger.Information("Skipped orphan blob cleanup run as this instance is not the leader");
-                return new List<RemovedBlobs>();
+                return new List<BlobIdentifier>();
             }
 
             List<NamespaceId> namespaces = await ListNamespaces().Where(NamespaceShouldBeCleaned).ToListAsync();
@@ -93,7 +83,7 @@ namespace Horde.Storage.Implementation
             }, cancellationToken);
 
             List<BlobIdentifier> blobsToRemove = new List<BlobIdentifier>();
-            List<RemovedBlobs> removedBlobs = new List<RemovedBlobs>();
+            List<BlobIdentifier> removedBlobs = new List<BlobIdentifier>();
 
             // enumerate all namespaces, and check if the old blob is valid in any of them to allow for a blob store to just store them in a single pile if it wants to
             foreach (NamespaceId @namespace in namespaces)
@@ -156,7 +146,7 @@ namespace Horde.Storage.Implementation
                 }
                 if (deleted)
                 {
-                    removedBlobs.Add(new RemovedBlobs(blob));
+                    removedBlobs.Add(blob);
                 }
             }
 
diff --git a/Engine/Source/Programs/Horde/HordeStorage/Horde.Storage/Implementation/GC/OrphanBlobCleanupRefs.cs b/Engine/Source/Programs/Horde/HordeStorage/Horde.Storage/Implementation/GC/OrphanBlobCleanupRefs.cs
new file mode 100644
--- /dev/null
+++ b/Engine/Source/Programs/Horde/HordeStorage/Horde.Storage/Implementation/GC/OrphanBlobCleanupRefs.cs
@@ -0,0 +1,129 @@
+// Copyright Epic Games, Inc. All Rights Reserved.
+
+using System;
+using System.Collections.Generic;
+using System.Threading;
+using System.Threading.Tasks;
+using async_enumerable_dotnet;
+using Dasync.Collections;
+using Horde.Storage.Controllers;
+using Horde.Storage.Implementation.Blob;
+using Jupiter.Implementation;
+using Microsoft.Extensions.Options;
+using Serilog;
+
+namespace Horde.Storage.Implementation
+{
+    public class OrphanBlobCleanupRefs : IBlobCleanup
+    {
+        private readonly IBlobService _blobService;
+        private readonly IObjectService _objectService;
+        private readonly IBlobIndex _blobIndex;
+        private readonly ILeaderElection _leaderElection;
+        private readonly IOptionsMonitor<GCSettings> _gcSettings;
+        private readonly ILogger _logger = Log.ForContext<OrphanBlobCleanupRefs>();
+
+        // ReSharper disable once UnusedMember.Global
+        public OrphanBlobCleanupRefs(IBlobService blobService, IObjectService objectService, IBlobIndex blobIndex, ILeaderElection leaderElection, IOptionsMonitor<GCSettings> gcSettings)
+        {
+            _blobService = blobService;
+            _objectService = objectService;
+            _blobIndex = blobIndex;
+            _leaderElection = leaderElection;
+            _gcSettings = gcSettings;
+        }
+
+        public async Task<List<BlobIdentifier>> Cleanup(CancellationToken cancellationToken)
+        {
+            if (!_leaderElection.IsThisInstanceLeader())
+            {
+                _logger.Information("Skipped orphan blob (refs) cleanup run as this instance is not the leader");
+                return new List<BlobIdentifier>();
+            }
+
+            List<NamespaceId> namespaces = await ListNamespaces().Where(NamespaceShouldBeCleaned).ToListAsync();
+          
+
+            // enumerate all namespaces, and check if the old blob is valid in any of them to allow for a blob store to just store them in a single pile if it wants to
+            List<BlobIdentifier> removedBlobs = new List<BlobIdentifier>();
+            foreach (NamespaceId @namespace in namespaces)
+            {
+                if (cancellationToken.IsCancellationRequested)
+                    break;
+
+                // only consider blobs that have been around for 60 minutes
+                // this due to cases were blobs are uploaded first
+                await foreach ((BlobIdentifier blob, DateTime lastModified) in _blobService.ListObjects(@namespace).WithCancellation(cancellationToken))
+                {
+                    if (lastModified > DateTime.Now.AddMinutes(-60))
+                        continue;
+                    
+                    if (cancellationToken.IsCancellationRequested)
+                        break;
+
+                    IBlobIndex.BlobInfo? blobIndex = await _blobIndex.GetBlobInfo(@namespace, blob);
+
+                    bool shouldDelete = blobIndex == null;
+                    if (blobIndex != null)
+                    {
+                        bool found = false;
+                        foreach ((BucketId bucket, IoHashKey key) in blobIndex.References)
+                        {
+                            if (cancellationToken.IsCancellationRequested)
+                                break;
+
+                            try
+                            {
+                                (ObjectRecord, BlobContents) _ = await _objectService.Get(@namespace, bucket, key, Array.Empty<string>());
+                                found = true;
+                            }
+                            catch (ObjectNotFoundException)
+                            {
+                                // this is not a valid reference so we should delete
+                            }
+                            catch (MissingBlobsException)
+                            {
+                                // we do not care if there are missing blobs, as long as the record is valid we keep this blob around
+                                found = true;
+                            }
+                        }
+
+                        if (!found)
+                            shouldDelete = true;
+                    }
+                    
+                    if (shouldDelete)
+                    {
+                        await RemoveBlob(@namespace, blob);
+                        removedBlobs.Add(blob);
+                    }
+                }
+            }
+
+            return removedBlobs;
+        }
+
+        private async Task RemoveBlob(NamespaceId ns, BlobIdentifier blob)
+        {
+            _logger.Information("Attempting to GC Orphan blob {Blob} from {Namespace}", blob, ns);
+            try
+            {
+                await _blobService.DeleteObject(ns, blob);
+            }
+            catch (Exception e)
+            {
+                _logger.Warning("Failed to delete blob {Blob} from {Namespace} due to {Error}", blob, ns, e.Message);
+            }
+        }
+
+        private bool NamespaceShouldBeCleaned(NamespaceId ns)
+        {
+            return _gcSettings.CurrentValue.CleanNamespaces.Contains(ns.ToString());
+        }
+
+        private IAsyncEnumerable<NamespaceId> ListNamespaces()
+        {
+            return _objectService.GetNamespaces();
+        }
+    }
+}
diff --git a/Engine/Source/Programs/Horde/HordeStorage/Horde.Storage/Implementation/ObjectService.cs b/Engine/Source/Programs/Horde/HordeStorage/Horde.Storage/Implementation/ObjectService.cs
--- a/Engine/Source/Programs/Horde/HordeStorage/Horde.Storage/Implementation/ObjectService.cs
+++ b/Engine/Source/Programs/Horde/HordeStorage/Horde.Storage/Implementation/ObjectService.cs
@@ -19,7 +19,7 @@ namespace Horde.Storage.Implementation
         Task<PutObjectResult> Put(NamespaceId ns, BucketId bucket, IoHashKey key, BlobIdentifier blobHash, CompactBinaryObject payload);
         Task<BlobIdentifier[]> Finalize(NamespaceId ns, BucketId bucket, IoHashKey key, BlobIdentifier blobHash);
 
-        IAsyncEnumerator<NamespaceId> GetNamespaces();
+        IAsyncEnumerable<NamespaceId> GetNamespaces();
 
         Task<long> Delete(NamespaceId ns, BucketId bucket, IoHashKey key);
         Task<long> DropNamespace(NamespaceId ns);
@@ -136,7 +136,7 @@ namespace Horde.Storage.Implementation
             return missingReferences;
         }
 
-        public IAsyncEnumerator<NamespaceId> GetNamespaces()
+        public IAsyncEnumerable<NamespaceId> GetNamespaces()
         {
             return _referencesStore.GetNamespaces();
         }
diff --git a/Engine/Source/Programs/Horde/HordeStorage/Horde.Storage/Implementation/Objects/IReferencesStore.cs b/Engine/Source/Programs/Horde/HordeStorage/Horde.Storage/Implementation/Objects/IReferencesStore.cs
--- a/Engine/Source/Programs/Horde/HordeStorage/Horde.Storage/Implementation/Objects/IReferencesStore.cs
+++ b/Engine/Source/Programs/Horde/HordeStorage/Horde.Storage/Implementation/Objects/IReferencesStore.cs
@@ -15,9 +15,9 @@ namespace Horde.Storage.Implementation
 
 
         Task UpdateLastAccessTime(NamespaceId ns, BucketId bucket, IoHashKey key, DateTime newLastAccessTime);
-        IAsyncEnumerator<ObjectRecord> GetOldestRecords(NamespaceId ns);
+        IAsyncEnumerable<ObjectRecord> GetOldestRecords(NamespaceId ns);
 
-        IAsyncEnumerator<NamespaceId> GetNamespaces();
+        IAsyncEnumerable<NamespaceId> GetNamespaces();
         Task<long> Delete(NamespaceId ns, BucketId bucket, IoHashKey key);
         Task<long> DropNamespace(NamespaceId ns);
         Task<long> DeleteBucket(NamespaceId ns, BucketId bucket);
diff --git a/Engine/Source/Programs/Horde/HordeStorage/Horde.Storage/Implementation/Objects/MemoryReferencesStore.cs b/Engine/Source/Programs/Horde/HordeStorage/Horde.Storage/Implementation/Objects/MemoryReferencesStore.cs
--- a/Engine/Source/Programs/Horde/HordeStorage/Horde.Storage/Implementation/Objects/MemoryReferencesStore.cs
+++ b/Engine/Source/Programs/Horde/HordeStorage/Horde.Storage/Implementation/Objects/MemoryReferencesStore.cs
@@ -66,7 +66,7 @@ namespace Horde.Storage.Implementation
             return Task.CompletedTask;
         }
 
-        public async IAsyncEnumerator<ObjectRecord> GetOldestRecords(NamespaceId ns)
+        public async IAsyncEnumerable<ObjectRecord> GetOldestRecords(NamespaceId ns)
         {
             foreach (MemoryStoreObject o in _objects.Values.Where(o => o.Namespace == ns).OrderBy(o => o.LastAccessTime))
             {
@@ -75,9 +75,9 @@ namespace Horde.Storage.Implementation
             }
         }
 
-        public IAsyncEnumerator<NamespaceId> GetNamespaces()
+        public IAsyncEnumerable<NamespaceId> GetNamespaces()
         {
-            return _namespaces.GetAsyncEnumerator();
+            return _namespaces.ToAsyncEnumerable();
         }
 
         public Task<long> Delete(NamespaceId ns, BucketId bucket, IoHashKey key)
diff --git a/Engine/Source/Programs/Horde/HordeStorage/Horde.Storage/Implementation/Objects/ScyllaReferencesStore.cs b/Engine/Source/Programs/Horde/HordeStorage/Horde.Storage/Implementation/Objects/ScyllaReferencesStore.cs
--- a/Engine/Source/Programs/Horde/HordeStorage/Horde.Storage/Implementation/Objects/ScyllaReferencesStore.cs
+++ b/Engine/Source/Programs/Horde/HordeStorage/Horde.Storage/Implementation/Objects/ScyllaReferencesStore.cs
@@ -141,7 +141,7 @@ namespace Horde.Storage.Implementation
             await updateObjectTracking;
         }
 
-        public async IAsyncEnumerator<ObjectRecord> GetOldestRecords(NamespaceId ns)
+        public async IAsyncEnumerable<ObjectRecord> GetOldestRecords(NamespaceId ns)
         {
             using Scope _ = Tracer.Instance.StartActive("scylla.get_records");
             for (sbyte partitionIndex = sbyte.MinValue; partitionIndex < sbyte.MaxValue; partitionIndex++)
@@ -187,7 +187,7 @@ namespace Horde.Storage.Implementation
             }
         }
 
-        public async IAsyncEnumerator<NamespaceId> GetNamespaces()
+        public async IAsyncEnumerable<NamespaceId> GetNamespaces()
         {
             using Scope _ = Tracer.Instance.StartActive("scylla.get_namespaces");
             RowSet rowSet = await _session.ExecuteAsync(new SimpleStatement("SELECT DISTINCT namespace FROM buckets"));
diff --git a/Engine/Source/Programs/Horde/HordeStorage/Horde.Storage/Settings.cs b/Engine/Source/Programs/Horde/HordeStorage/Horde.Storage/Settings.cs
--- a/Engine/Source/Programs/Horde/HordeStorage/Horde.Storage/Settings.cs
+++ b/Engine/Source/Programs/Horde/HordeStorage/Horde.Storage/Settings.cs
@@ -247,5 +247,6 @@ namespace Horde.Storage
         public TimeSpan BlobCleanupPollFrequency { get; set; } = TimeSpan.FromMinutes(60);
         public TimeSpan RefCleanupPollFrequency { get; set; } = TimeSpan.FromMinutes(60);
         public List<string> CleanNamespacesV1 { get; set; } = new List<string>();
+        public List<string> CleanNamespaces { get; set; }  = new List<string>();
     }
 }
diff --git a/Engine/Source/Programs/Horde/HordeStorage/Horde.Storage/Tests/Functional/GC/GCBlobTests.cs b/Engine/Source/Programs/Horde/HordeStorage/Horde.Storage/Tests/Functional/GC/GCBlobTests.cs
--- a/Engine/Source/Programs/Horde/HordeStorage/Horde.Storage/Tests/Functional/GC/GCBlobTests.cs
+++ b/Engine/Source/Programs/Horde/HordeStorage/Horde.Storage/Tests/Functional/GC/GCBlobTests.cs
@@ -10,6 +10,7 @@ using System.Threading.Tasks;
 using Dasync.Collections;
 using Jupiter.Implementation;
 using Horde.Storage.Implementation;
+using Horde.Storage.Implementation.Blob;
 using Microsoft.AspNetCore.TestHost;
 using Microsoft.AspNetCore.Hosting;
 using Microsoft.Extensions.Configuration;
@@ -175,17 +176,160 @@ namespace Horde.Storage.FunctionalTests.GC
             Assert.AreEqual(1, namespaces.Count);
 
             CancellationTokenSource cts = new CancellationTokenSource();
-            List<RemovedBlobs> removedBlobs = await cleanup.Cleanup(cts.Token);
+            List<BlobIdentifier> removedBlobs = await cleanup.Cleanup(cts.Token);
             Assert.AreEqual(4, removedBlobs.Count);
             
             BlobIdentifier[] expectedIdentifiers = {object1id, object2id, object4id, object5id};
-            foreach (RemovedBlobs removedBlob in removedBlobs)
+            foreach (BlobIdentifier removedBlob in removedBlobs)
             {
-                Assert.IsTrue(expectedIdentifiers.Contains(removedBlob.BlobIdentifier));
+                Assert.IsTrue(expectedIdentifiers.Contains(removedBlob));
             }
             
             _callistoBlobMock.Verify();
             _callistoBlobMock.VerifyNoOtherCalls();
         }
     }
+
+    [TestClass]
+    public class MemoryGCBlobTestsRefs : GCBlobTestsRefs
+    {
+        protected override string GetImplementation()
+        {
+            return "Memory";
+        }
+    }
+
+    [TestClass]
+    public class ScyllaGCBlobTestsRefs : GCBlobTestsRefs
+    {
+        protected override string GetImplementation()
+        {
+            return "Scylla";
+        }
+    }
+
+    public abstract class GCBlobTestsRefs
+    {
+        private TestServer? _server;
+        private HttpClient? _httpClient;
+
+        private readonly BlobIdentifier object0id = new BlobIdentifier("0000000000000000000000000000000000000000");
+        private readonly BlobIdentifier object1id = new BlobIdentifier("1111111111111111111111111111111111111111");
+        private readonly BlobIdentifier object2id = new BlobIdentifier("2222222222222222222222222222222222222222");
+        private readonly BlobIdentifier object3id = new BlobIdentifier("3333333333333333333333333333333333333333");
+        private readonly BlobIdentifier object4id = new BlobIdentifier("4444444444444444444444444444444444444444");
+        private readonly BlobIdentifier object5id = new BlobIdentifier("5555555555555555555555555555555555555555");
+        private readonly BlobIdentifier object6id = new BlobIdentifier("6666666666666666666666666666666666666666");
+        private IBlobService? _blobService;
+
+        private readonly NamespaceId TestNamespace = new NamespaceId("test-namespace");
+
+        [TestInitialize]
+        public async Task Setup()
+        {
+            IConfigurationRoot configuration = new ConfigurationBuilder()
+                // we are not reading the base appSettings here as we want exact control over what runs in the tests
+                .AddJsonFile("appsettings.Testing.json", false)
+                .AddEnvironmentVariables()
+                .AddInMemoryCollection(new List<KeyValuePair<string, string>>()
+                {
+                    new KeyValuePair<string, string>("Horde.Storage:StorageImplementations:0", "MemoryBlobStore"),
+                    new KeyValuePair<string, string>("GC:CleanNamespaces:0", TestNamespace.ToString()),
+                    new KeyValuePair<string, string>("GC:CleanOldBlobs", true.ToString()),
+                    new KeyValuePair<string, string>("Horde.Storage:BlobIndexImplementation", GetImplementation()),
+                    
+                })
+                .Build();
+
+            Logger logger = new LoggerConfiguration()
+                .ReadFrom.Configuration(configuration)
+                .CreateLogger();
+
+            TestServer server = new TestServer(new WebHostBuilder()
+                .UseConfiguration(configuration)
+                .UseEnvironment("Testing")
+                .UseSerilog(logger)
+                .UseStartup<HordeStorageStartup>()
+            );
+
+            _httpClient = server.CreateClient();
+            _server = server;
+
+            _blobService = server.Services.GetService<IBlobService>()!;
+
+            MemoryBlobStore memoryBlobStore = (MemoryBlobStore) ((BlobService)_blobService).BlobStore.First();
+            byte[] emptyContents = new byte[0];
+            await memoryBlobStore.PutObject(TestNamespace, emptyContents, object0id);
+            await memoryBlobStore.PutObject(TestNamespace, emptyContents, object1id);// this is not in the index
+            await memoryBlobStore.PutObject(TestNamespace, emptyContents, object2id);
+            await memoryBlobStore.PutObject(TestNamespace, emptyContents, object3id);
+            await memoryBlobStore.PutObject(TestNamespace, emptyContents, object4id); // this is not in the index
+            await memoryBlobStore.PutObject(TestNamespace, emptyContents, object5id); // this is not in the index
+            await memoryBlobStore.PutObject(TestNamespace, emptyContents, object6id);
+
+            // set all objects to be old, only the orphaned blobs should be deleted
+            memoryBlobStore.SetLastModifiedTime(TestNamespace, object0id, DateTime.Now.AddDays(-2));
+            memoryBlobStore.SetLastModifiedTime(TestNamespace, object1id, DateTime.Now.AddDays(-2));
+            memoryBlobStore.SetLastModifiedTime(TestNamespace, object2id, DateTime.Now.AddDays(-2));
+            memoryBlobStore.SetLastModifiedTime(TestNamespace, object3id, DateTime.Now.AddDays(-2));
+            memoryBlobStore.SetLastModifiedTime(TestNamespace, object4id, DateTime.Now.AddDays(-2));
+            memoryBlobStore.SetLastModifiedTime(TestNamespace, object5id, DateTime.Now.AddDays(-2));
+            memoryBlobStore.SetLastModifiedTime(TestNamespace, object6id, DateTime.Now.AddDays(-2));
+
+            BucketId testBucket = new BucketId("test");
+            IObjectService? objectService = server.Services.GetService<IObjectService>()!;
+            Assert.IsNotNull(objectService);
+            (BlobIdentifier ob0_hash, CompactBinaryObject ob0_cb) = GetCBWithAttachment(object0id);
+            await objectService.Put(TestNamespace, testBucket, IoHashKey.FromName("object0"), ob0_hash, ob0_cb);
+
+            (BlobIdentifier ob2_hash, CompactBinaryObject ob2_cb) = GetCBWithAttachment(object2id);
+            await objectService.Put(TestNamespace, testBucket, IoHashKey.FromName("object2"), ob2_hash, ob2_cb);
+
+            (BlobIdentifier ob3_hash, CompactBinaryObject ob3_cb) = GetCBWithAttachment(object3id);
+            await objectService.Put(TestNamespace, testBucket, IoHashKey.FromName("object3"), ob3_hash, ob3_cb);
+
+            (BlobIdentifier ob6_hash, CompactBinaryObject ob6_cb) = GetCBWithAttachment(object6id);
+            await objectService.Put(TestNamespace, testBucket, IoHashKey.FromName("object6"), ob6_hash, ob6_cb);
+
+            IBlobIndex? blobIndex = server.Services.GetService<IBlobIndex>()!;
+            Assert.IsNotNull(blobIndex);
+            await blobIndex.AddBlobToIndex(TestNamespace, object0id);
+            await blobIndex.AddRefToBlobs(TestNamespace, testBucket, IoHashKey.FromName("object0"), new [] {object0id });
+            await blobIndex.AddBlobToIndex(TestNamespace, object2id);
+            await blobIndex.AddRefToBlobs(TestNamespace, testBucket, IoHashKey.FromName("object2"), new [] {object2id });
+            await blobIndex.AddBlobToIndex(TestNamespace, object3id);
+            await blobIndex.AddRefToBlobs(TestNamespace, testBucket, IoHashKey.FromName("object3"), new [] {object3id });
+            await blobIndex.AddBlobToIndex(TestNamespace, object6id);
+            await blobIndex.AddRefToBlobs(TestNamespace, testBucket, IoHashKey.FromName("object6"), new [] {object6id });
+
+
+        }
+
+        protected abstract string GetImplementation();
+
+        private (BlobIdentifier, CompactBinaryObject) GetCBWithAttachment(BlobIdentifier blobIdentifier)
+        {
+            CompactBinaryWriter compactBinaryWriter = new CompactBinaryWriter();
+            compactBinaryWriter.BeginObject();
+            compactBinaryWriter.AddBinaryAttachment(blobIdentifier);
+            compactBinaryWriter.EndObject();
+
+            byte[] b = compactBinaryWriter.Save();
+            return (BlobIdentifier.FromBlob(b), CompactBinaryObject.Load(b));
+        }
+
+        [TestMethod]
+        public async Task RunBlobCleanupRefs()
+        {
+            OrphanBlobCleanupRefs? cleanup = _server!.Services.GetService<OrphanBlobCleanupRefs>();
+            Assert.IsNotNull(cleanup);
+
+            CancellationTokenSource cts = new CancellationTokenSource();
+            List<BlobIdentifier> removedBlobs = await cleanup.Cleanup(cts.Token);
+            Assert.AreEqual(3, removedBlobs.Count);
+            
+            BlobIdentifier[] expectedIdentifiers = {object1id, object4id, object5id};
+            CollectionAssert.AreEquivalent(removedBlobs, expectedIdentifiers);
+        }
+    }
 }
