Horde: Add null check against returned list of changes.

diff --git a/Engine/Source/Programs/Horde/Horde.Build/Perforce/Impl/PerforceService.cs b/Engine/Source/Programs/Horde/Horde.Build/Perforce/Impl/PerforceService.cs
--- a/Engine/Source/Programs/Horde/Horde.Build/Perforce/Impl/PerforceService.cs
+++ b/Engine/Source/Programs/Horde/Horde.Build/Perforce/Impl/PerforceService.cs
@@ -466,9 +466,12 @@ namespace HordeServer.Services
 			Repository.Connection.Client.RevertFiles(new P4.Options(), new P4.DepotPath("//..."));
 
 			IList<P4.Changelist> Changes = Repository.GetChangelists(new P4.ChangesCmdOptions(P4.ChangesCmdFlags.None, Repository.Connection.Client.Name, 100, P4.ChangeListStatus.Pending, null));
-			foreach (P4.Changelist Change in Changes)
+			if (Changes != null)
 			{
-				Repository.DeleteChangelist(Change, new P4.Options());
+				foreach (P4.Changelist Change in Changes)
+				{
+					Repository.DeleteChangelist(Change, new P4.Options());
+				}
 			}
 		}
 
