Horde: Exclude disabled and offline agents from telemetry data.

diff --git a/Engine/Source/Programs/Horde/Horde.Build/Fleet/TelemetryService.cs b/Engine/Source/Programs/Horde/Horde.Build/Fleet/TelemetryService.cs
--- a/Engine/Source/Programs/Horde/Horde.Build/Fleet/TelemetryService.cs
+++ b/Engine/Source/Programs/Horde/Horde.Build/Fleet/TelemetryService.cs
@@ -86,6 +86,9 @@ namespace HordeServer.Services
 			List<IPool> Pools = await PoolCollection.GetAsync();
 			List<ILease> Leases = await LeaseCollection.FindLeasesAsync(MinTime: MinTime);
 
+			// Remove any agents which are offline
+			Agents.RemoveAll(x => !x.Enabled || !x.IsSessionValid(CurrentTime));
+
 			// Find all the agents
 			Dictionary<AgentId, List<PoolId>> AgentToPoolIds = Agents.ToDictionary(x => x.Id, x => x.GetPools().ToList());
 
