Horde: Change schedule service to allow multiple updates to be handled at once. Triggers are now scheduled in a sorted set stored in Redis, and served in order.

diff --git a/Engine/Source/Programs/Horde/Horde.Build.Tests/SchedulerTests.cs b/Engine/Source/Programs/Horde/Horde.Build.Tests/SchedulerTests.cs
--- a/Engine/Source/Programs/Horde/Horde.Build.Tests/SchedulerTests.cs
+++ b/Engine/Source/Programs/Horde/Horde.Build.Tests/SchedulerTests.cs
@@ -19,6 +19,7 @@ using System.Linq;
 using HordeServerTests.Stubs.Services;
 using HordeServer.Collections;
 using HordeServer.Utilities;
+using System.Threading;
 
 namespace HordeServerTests
 {
@@ -54,6 +55,8 @@ namespace HordeServerTests
 
 		async Task<IStream> SetScheduleAsync(CreateScheduleRequest Schedule)
 		{
+			await ScheduleService.ResetAsync();
+
 			IStream? Stream = await StreamService.GetStreamAsync(StreamId);
 
 			StreamConfig Config = new StreamConfig();
@@ -88,7 +91,7 @@ namespace HordeServerTests
 			await SetScheduleAsync(Schedule);
 
 			Clock.Advance(TimeSpan.FromHours(1.25));
-			await ScheduleService.TickSharedOnlyForTestingAsync();
+			await ScheduleService.TickForTestingAsync();
 
 			return await GetNewJobs();
 		}
@@ -128,19 +131,19 @@ namespace HordeServerTests
 			DateTime StartTime = new DateTime(2021, 1, 1, 12, 0, 0, DateTimeKind.Utc); // Friday Jan 1, 2021 
 			Clock.UtcNow = StartTime;
 
-			Schedule Schedule = new Schedule(RequireSubmittedChange: false);
+			Schedule Schedule = new Schedule(Clock.UtcNow, RequireSubmittedChange: false);
 			Schedule.Patterns.Add(new SchedulePattern(new List<DayOfWeek> { DayOfWeek.Friday, DayOfWeek.Sunday }, 13 * 60, null, null));
 
-			DateTimeOffset? NextTime = Schedule.GetNextTriggerTime(StartTime, TimeZoneInfo.Utc);
+			DateTime? NextTime = Schedule.GetNextTriggerTimeUtc(StartTime, TimeZoneInfo.Utc);
 			Assert.AreEqual(StartTime + TimeSpan.FromHours(1.0), NextTime!.Value);
 
-			NextTime = Schedule.GetNextTriggerTime(NextTime.Value, TimeZoneInfo.Utc);
+			NextTime = Schedule.GetNextTriggerTimeUtc(NextTime.Value, TimeZoneInfo.Utc);
 			Assert.AreEqual(StartTime + TimeSpan.FromHours(1.0 + 24.0 * 2.0), NextTime!.Value);
 
-			NextTime = Schedule.GetNextTriggerTime(NextTime.Value, TimeZoneInfo.Utc);
+			NextTime = Schedule.GetNextTriggerTimeUtc(NextTime.Value, TimeZoneInfo.Utc);
 			Assert.AreEqual(StartTime + TimeSpan.FromHours(1.0 + 24.0 * 7.0), NextTime!.Value);
 
-			NextTime = Schedule.GetNextTriggerTime(NextTime.Value, TimeZoneInfo.Utc);
+			NextTime = Schedule.GetNextTriggerTimeUtc(NextTime.Value, TimeZoneInfo.Utc);
 			Assert.AreEqual(StartTime + TimeSpan.FromHours(1.0 + 24.0 * 9.0), NextTime!.Value);
 		}
 
@@ -150,25 +153,25 @@ namespace HordeServerTests
 			DateTime StartTime = new DateTime(2021, 1, 1, 12, 0, 0, DateTimeKind.Utc); // Friday Jan 1, 2021 
 			Clock.UtcNow = StartTime;
 
-			Schedule Schedule = new Schedule(RequireSubmittedChange: false);
+			Schedule Schedule = new Schedule(Clock.UtcNow, RequireSubmittedChange: false);
 			Schedule.Patterns.Add(new SchedulePattern(null, 13 * 60, 14 * 60, 15));
 
-			DateTimeOffset? NextTime = Schedule.GetNextTriggerTime(StartTime, TimeZoneInfo.Utc);
+			DateTime? NextTime = Schedule.GetNextTriggerTimeUtc(StartTime, TimeZoneInfo.Utc);
 			Assert.AreEqual(StartTime + TimeSpan.FromHours(1.0), NextTime!.Value);
 
-			NextTime = Schedule.GetNextTriggerTime(NextTime.Value, TimeZoneInfo.Utc);
+			NextTime = Schedule.GetNextTriggerTimeUtc(NextTime.Value, TimeZoneInfo.Utc);
 			Assert.AreEqual(StartTime + TimeSpan.FromHours(1.25), NextTime!.Value);
 
-			NextTime = Schedule.GetNextTriggerTime(NextTime.Value, TimeZoneInfo.Utc);
+			NextTime = Schedule.GetNextTriggerTimeUtc(NextTime.Value, TimeZoneInfo.Utc);
 			Assert.AreEqual(StartTime + TimeSpan.FromHours(1.5), NextTime!.Value);
 
-			NextTime = Schedule.GetNextTriggerTime(NextTime.Value, TimeZoneInfo.Utc);
+			NextTime = Schedule.GetNextTriggerTimeUtc(NextTime.Value, TimeZoneInfo.Utc);
 			Assert.AreEqual(StartTime + TimeSpan.FromHours(1.75), NextTime!.Value);
 
-			NextTime = Schedule.GetNextTriggerTime(NextTime.Value, TimeZoneInfo.Utc);
+			NextTime = Schedule.GetNextTriggerTimeUtc(NextTime.Value, TimeZoneInfo.Utc);
 			Assert.AreEqual(StartTime + TimeSpan.FromHours(2.0), NextTime!.Value);
 
-			NextTime = Schedule.GetNextTriggerTime(NextTime.Value, TimeZoneInfo.Utc);
+			NextTime = Schedule.GetNextTriggerTimeUtc(NextTime.Value, TimeZoneInfo.Utc);
 			Assert.AreEqual(StartTime + TimeSpan.FromHours(1.0 + 24.0), NextTime!.Value);
 		}
 
@@ -185,14 +188,14 @@ namespace HordeServerTests
 			await SetScheduleAsync(Schedule);
 
 			// Initial tick
-			await ScheduleService.TickSharedOnlyForTestingAsync();
+			await ScheduleService.TickForTestingAsync();
 
 			List<IJob> Jobs1 = await GetNewJobs();
 			Assert.AreEqual(0, Jobs1.Count);
 
 			// Trigger a job
 			Clock.Advance(TimeSpan.FromHours(1.25));
-			await ScheduleService.TickSharedOnlyForTestingAsync();
+			await ScheduleService.TickForTestingAsync();
 
 			List<IJob> Jobs2 = await GetNewJobs();
 			Assert.AreEqual(1, Jobs2.Count);
@@ -212,14 +215,14 @@ namespace HordeServerTests
 			IStream Stream = await SetScheduleAsync(Schedule);
 
 			// Initial tick
-			await ScheduleService.TickSharedOnlyForTestingAsync();
+			await ScheduleService.TickForTestingAsync();
 
 			List<IJob> Jobs1 = await GetNewJobs();
 			Assert.AreEqual(0, Jobs1.Count);
 
 			// Trigger a job
 			Clock.Advance(TimeSpan.FromHours(1.25));
-			await ScheduleService.TickSharedOnlyForTestingAsync();
+			await ScheduleService.TickForTestingAsync();
 
 			List<IJob> Jobs2 = await GetNewJobs();
 			Assert.AreEqual(1, Jobs2.Count);
@@ -233,7 +236,7 @@ namespace HordeServerTests
 
 			// Trigger another job
 			Clock.Advance(TimeSpan.FromHours(0.5));
-			await ScheduleService.TickSharedOnlyForTestingAsync();
+			await ScheduleService.TickForTestingAsync();
 
 			List<IJob> Jobs3 = await GetNewJobs();
 			Assert.AreEqual(0, Jobs3.Count);
@@ -253,7 +256,7 @@ namespace HordeServerTests
 			await SetScheduleAsync(Schedule);
 
 			// Initial tick
-			await ScheduleService.TickSharedOnlyForTestingAsync();
+			await ScheduleService.TickForTestingAsync();
 
 			List<IJob> Jobs1 = await GetNewJobs();
 			Assert.AreEqual(0, Jobs1.Count);
@@ -266,7 +269,7 @@ namespace HordeServerTests
 			PerforceService.AddChange("//UE5/Main", 106, Bob, "", new string[] { "foo.cpp" });
 
 			Clock.Advance(TimeSpan.FromHours(1.25));
-			await ScheduleService.TickSharedOnlyForTestingAsync();
+			await ScheduleService.TickForTestingAsync();
 
 			List<IJob> Jobs2 = await GetNewJobs();
 			Assert.AreEqual(2, Jobs2.Count);
@@ -291,7 +294,7 @@ namespace HordeServerTests
 
 			// Trigger a job
 			Clock.Advance(TimeSpan.FromHours(1.25));
-			await ScheduleService.TickSharedOnlyForTestingAsync();
+			await ScheduleService.TickForTestingAsync();
 
 			List<IJob> Jobs2 = await GetNewJobs();
 			Assert.AreEqual(1, Jobs2.Count);
@@ -300,7 +303,7 @@ namespace HordeServerTests
 
 			// Test that another job does not trigger
 			Clock.Advance(TimeSpan.FromHours(0.5));
-			await ScheduleService.TickSharedOnlyForTestingAsync();
+			await ScheduleService.TickForTestingAsync();
 
 			List<IJob> Jobs3 = await GetNewJobs();
 			Assert.AreEqual(0, Jobs3.Count);
@@ -310,7 +313,7 @@ namespace HordeServerTests
 
 			// Test that another job does not trigger
 			Clock.Advance(TimeSpan.FromHours(0.5));
-			await ScheduleService.TickSharedOnlyForTestingAsync();
+			await ScheduleService.TickForTestingAsync();
 
 			List<IJob> Jobs4 = await GetNewJobs();
 			Assert.AreEqual(1, Jobs4.Count);
@@ -331,7 +334,7 @@ namespace HordeServerTests
 
 			// Trigger a job
 			Clock.Advance(TimeSpan.FromHours(1.25));
-			await ScheduleService.TickSharedOnlyForTestingAsync();
+			await ScheduleService.TickForTestingAsync();
 
 			List<IJob> Jobs2 = await GetNewJobs();
 			Assert.AreEqual(1, Jobs2.Count);
@@ -340,7 +343,7 @@ namespace HordeServerTests
 
 			// Check another job does not trigger due to the change above
 			Clock.Advance(TimeSpan.FromHours(1.25));
-			await ScheduleService.TickSharedOnlyForTestingAsync();
+			await ScheduleService.TickForTestingAsync();
 
 			List<IJob> Jobs3 = await GetNewJobs();
 			Assert.AreEqual(0, Jobs3.Count);
@@ -359,7 +362,7 @@ namespace HordeServerTests
 
 			ITemplate? NewTemplate2 = await TemplateCollection.AddAsync("Test template 2");
 			TemplateRef NewTemplateRef2 = new TemplateRef(NewTemplate2);
-			NewTemplateRef2.Schedule = new Schedule();
+			NewTemplateRef2.Schedule = new Schedule(Clock.UtcNow);
 			NewTemplateRef2.Schedule.Gate = new ScheduleGate(NewTemplateRefId1, "TriggerNext");
 			NewTemplateRef2.Schedule.Patterns.Add(new SchedulePattern(null, 0, null, 10));
 			NewTemplateRef2.Schedule.LastTriggerTime = StartTime;
@@ -379,7 +382,7 @@ namespace HordeServerTests
 			GraphA = await GraphCollection.AppendAsync(GraphA, new List<NewGroup> { GroupA });
 
 			// Tick the schedule and make sure it doesn't trigger
-			await ScheduleService.TickSharedOnlyForTestingAsync();
+			await ScheduleService.TickForTestingAsync();
 			List<IJob> Jobs2 = await GetNewJobs();
 			Assert.AreEqual(0, Jobs2.Count);
 
@@ -392,7 +395,7 @@ namespace HordeServerTests
 
 			// Tick the schedule and make sure it doesn't trigger
 			Clock.Advance(TimeSpan.FromMinutes(30.0));
-			await ScheduleService.TickSharedOnlyForTestingAsync();
+			await ScheduleService.TickForTestingAsync();
 			List<IJob> Jobs3 = await GetNewJobs();
 			Assert.AreEqual(0, Jobs3.Count);
 
@@ -402,7 +405,7 @@ namespace HordeServerTests
 			Assert.IsNotNull(await JobService.UpdateStepAsync(Job2, Job2.Batches[0].Id, Job2.Batches[0].Steps[0].Id, JobStepState.Completed, JobStepOutcome.Success));
 
 			// Tick the schedule and make sure it does trigger
-			await ScheduleService.TickSharedOnlyForTestingAsync();
+			await ScheduleService.TickForTestingAsync();
 			List<IJob> Jobs4 = await GetNewJobs();
 			Assert.AreEqual(1, Jobs4.Count);
 			Assert.AreEqual(1234, Jobs4[0].Change);
@@ -423,7 +426,7 @@ namespace HordeServerTests
 			await SetScheduleAsync(Schedule);
 
 			Clock.Advance(TimeSpan.FromHours(1.25));
-			await ScheduleService.TickSharedOnlyForTestingAsync();
+			await ScheduleService.TickForTestingAsync();
 
 			List<IJob> Jobs1 = await GetNewJobs();
 			Assert.AreEqual(1, Jobs1.Count);
@@ -462,13 +465,13 @@ namespace HordeServerTests
 			
 			// Try trigger a job. No job should be scheduled as the stream is paused
 			Clock.Advance(TimeSpan.FromHours(1.25));
-			await ScheduleService.TickSharedOnlyForTestingAsync();
+			await ScheduleService.TickForTestingAsync();
 			List<IJob> Jobs2 = await GetNewJobs();
 			Assert.AreEqual(0, Jobs2.Count);
 
 			// Advance time beyond the pause period. A build should now trigger
 			Clock.Advance(TimeSpan.FromHours(5.25));
-			await ScheduleService.TickSharedOnlyForTestingAsync();
+			await ScheduleService.TickForTestingAsync();
 
 			List<IJob> Jobs3 = await GetNewJobs();
 			Assert.AreEqual(1, Jobs3.Count);
diff --git a/Engine/Source/Programs/Horde/Horde.Build/Api/Schedules.cs b/Engine/Source/Programs/Horde/Horde.Build/Api/Schedules.cs
--- a/Engine/Source/Programs/Horde/Horde.Build/Api/Schedules.cs
+++ b/Engine/Source/Programs/Horde/Horde.Build/Api/Schedules.cs
@@ -136,10 +136,11 @@ namespace HordeServer.Api
 		/// <summary>
 		/// Constructs a model object
 		/// </summary>
+		/// <param name="CurrentTimeUtc">The current time</param>
 		/// <returns>New model object</returns>
-		public Schedule ToModel()
+		public Schedule ToModel(DateTime CurrentTimeUtc)
 		{
-			return new Schedule(Enabled, MaxActive, MaxChanges, RequireSubmittedChange, Gate?.ToModel(), Filter, Files, TemplateParameters, Patterns.ConvertAll(x => x.ToModel()));
+			return new Schedule(CurrentTimeUtc, Enabled, MaxActive, MaxChanges, RequireSubmittedChange, Gate?.ToModel(), Filter, Files, TemplateParameters, Patterns.ConvertAll(x => x.ToModel()));
 		}
 	}
 
@@ -285,7 +286,7 @@ namespace HordeServer.Api
 			this.TemplateParameters = Schedule.TemplateParameters;
 			this.Patterns = Schedule.Patterns.ConvertAll(x => new GetSchedulePatternResponse(x));
 			this.LastTriggerChange = Schedule.LastTriggerChange;
-			this.LastTriggerTime = Schedule.LastTriggerTime;
+			this.LastTriggerTime = Schedule.GetLastTriggerTimeUtc();
 			this.ActiveJobs = Schedule.ActiveJobs.ConvertAll(x => x.ToString());
 		}
 	}
diff --git a/Engine/Source/Programs/Horde/Horde.Build/Projects/ISchedule.cs b/Engine/Source/Programs/Horde/Horde.Build/Projects/ISchedule.cs
--- a/Engine/Source/Programs/Horde/Horde.Build/Projects/ISchedule.cs
+++ b/Engine/Source/Programs/Horde/Horde.Build/Projects/ISchedule.cs
@@ -65,13 +65,13 @@ namespace HordeServer.Models
 		/// <summary>
 		/// Calculates the trigger index based on the given time in minutes
 		/// </summary>
-		/// <param name="LastTime">Time in minutes for the last trigger</param>
+		/// <param name="LastTimeUtc">Time for the last trigger</param>
 		/// <param name="TimeZone">The timezone for running the schedule</param>
 		/// <returns>Index of the trigger</returns>
-		public DateTimeOffset GetNextTriggerTime(DateTimeOffset LastTime, TimeZoneInfo TimeZone)
+		public DateTime GetNextTriggerTimeUtc(DateTime LastTimeUtc, TimeZoneInfo TimeZone)
 		{
 			// Convert last time into the correct timezone for running the scheule
-			LastTime = TimeZoneInfo.ConvertTime(LastTime, TimeZone);
+			DateTimeOffset LastTime = TimeZoneInfo.ConvertTime((DateTimeOffset)LastTimeUtc, TimeZone);
 
 			// Get the base time (ie. the start of this day) for anchoring the schedule
 			DateTimeOffset BaseTime = new DateTimeOffset(LastTime.Year, LastTime.Month, LastTime.Day, 0, 0, 0, LastTime.Offset);
@@ -85,7 +85,7 @@ namespace HordeServer.Models
 					// Get the time of the first trigger of this day. If the last time is less than this, this is the next trigger.
 					if (LastTimeMinutes < MinTime)
 					{
-						return BaseTime.AddMinutes(MinTime);
+						return BaseTime.AddMinutes(MinTime).UtcDateTime;
 					}
 
 					// Otherwise, get the time for the last trigger in the day.
@@ -100,7 +100,7 @@ namespace HordeServer.Models
 							int NextTimeMinutes = MinTime + (NextIndex * Interval.Value);
 							if (NextTimeMinutes <= ActualMaxTime)
 							{
-								return BaseTime.AddMinutes(NextTimeMinutes);
+								return BaseTime.AddMinutes(NextTimeMinutes).UtcDateTime;
 							}
 						}
 					}
@@ -205,7 +205,14 @@ namespace HordeServer.Models
 		/// <summary>
 		/// Last time that the schedule was triggered
 		/// </summary>
-		public DateTimeOffset LastTriggerTime { get; set; }
+		[BsonIgnoreIfNull]
+		public DateTimeOffset? LastTriggerTime { get; set; }
+
+		/// <summary>
+		/// Gets the last trigger time, in UTC
+		/// </summary>
+		[System.Diagnostics.CodeAnalysis.SuppressMessage("Naming", "CA1721:Property names should not match get methods", Justification = "<Pending>")]
+		public DateTime LastTriggerTimeUtc { get; set; }
 
 		/// <summary>
 		/// List of jobs that are currently active
@@ -225,6 +232,7 @@ namespace HordeServer.Models
 		/// <summary>
 		/// Constructor
 		/// </summary>
+		/// <param name="CurrentTimeUtc">The current time. This will be used to seed the last trigger time.</param>
 		/// <param name="Enabled">Whether the schedule is currently enabled</param>
 		/// <param name="MaxActive">Maximum number of builds that may be active at once</param>
 		/// <param name="MaxChanges">Maximum number of changes the schedule can fall behind head revision</param>
@@ -234,7 +242,7 @@ namespace HordeServer.Models
 		/// <param name="Files">Files that should trigger the schedule</param>
 		/// <param name="TemplateParameters">Parameters for the template to run</param>
 		/// <param name="Patterns">List of patterns for the schedule</param>
-		public Schedule(bool Enabled = true, int MaxActive = 0, int MaxChanges = 0, bool RequireSubmittedChange = true, ScheduleGate? Gate = null, List<ChangeContentFlags>? Filter = null, List<string>? Files = null, Dictionary<string, string>? TemplateParameters = null, List<SchedulePattern>? Patterns = null)
+		public Schedule(DateTime CurrentTimeUtc, bool Enabled = true, int MaxActive = 0, int MaxChanges = 0, bool RequireSubmittedChange = true, ScheduleGate? Gate = null, List<ChangeContentFlags>? Filter = null, List<string>? Files = null, Dictionary<string, string>? TemplateParameters = null, List<SchedulePattern>? Patterns = null)
 		{
 			this.Enabled = Enabled;
 			this.MaxActive = MaxActive;
@@ -245,7 +253,7 @@ namespace HordeServer.Models
 			this.Files = Files;
 			this.TemplateParameters = TemplateParameters ?? new Dictionary<string, string>();
 			this.Patterns = Patterns ?? new List<SchedulePattern>();
-			this.LastTriggerTime = DateTimeOffset.Now;
+			this.LastTriggerTimeUtc = CurrentTimeUtc;
 		}
 
 		/// <summary>
@@ -278,37 +286,51 @@ namespace HordeServer.Models
 			return Flags;
 		}
 
+		/// <summary>
+		/// Gets the last time that the schedule triggered
+		/// </summary>
+		/// <returns>Last trigger time</returns>
+		public DateTime GetLastTriggerTimeUtc()
+		{
+			if (LastTriggerTime != null)
+			{
+				LastTriggerTimeUtc = LastTriggerTime.Value.UtcDateTime;
+				LastTriggerTime = null;
+			}
+			return LastTriggerTimeUtc;
+		}
+
 		/// <summary>
 		/// Get the next time that the schedule will trigger
 		/// </summary>
 		/// <param name="TimeZone">Timezone to evaluate the trigger</param>
 		/// <returns>Next time at which the schedule will trigger</returns>
-		public DateTimeOffset? GetNextTriggerTime(TimeZoneInfo TimeZone)
+		public DateTime? GetNextTriggerTimeUtc(TimeZoneInfo TimeZone)
 		{
-			return GetNextTriggerTime(LastTriggerTime, TimeZone);
+			return GetNextTriggerTimeUtc(GetLastTriggerTimeUtc(), TimeZone);
 		}
 
 		/// <summary>
 		/// Get the next time that the schedule will trigger
 		/// </summary>
-		/// <param name="LastTime">Last time at which the schedule triggered</param>
+		/// <param name="LastTimeUtc">Last time at which the schedule triggered</param>
 		/// <param name="TimeZone">Timezone to evaluate the trigger</param>
 		/// <returns>Next time at which the schedule will trigger</returns>
-		public DateTimeOffset? GetNextTriggerTime(DateTimeOffset LastTime, TimeZoneInfo TimeZone)
+		public DateTime? GetNextTriggerTimeUtc(DateTime LastTimeUtc, TimeZoneInfo TimeZone)
 		{
-			DateTimeOffset? NextTriggerTime = null;
+			DateTime? NextTriggerTimeUtc = null;
 			if (Enabled)
 			{
 				foreach (SchedulePattern Pattern in Patterns)
 				{
-					DateTimeOffset PatternTriggerTime = Pattern.GetNextTriggerTime(LastTime, TimeZone);
-					if (NextTriggerTime == null || NextTriggerTime < PatternTriggerTime)
+					DateTime PatternTriggerTime = Pattern.GetNextTriggerTimeUtc(LastTimeUtc, TimeZone);
+					if (NextTriggerTimeUtc == null || NextTriggerTimeUtc < PatternTriggerTime)
 					{
-						NextTriggerTime = PatternTriggerTime;
+						NextTriggerTimeUtc = PatternTriggerTime;
 					}
 				}
 			}
-			return NextTriggerTime;
+			return NextTriggerTimeUtc;
 		}
 	}
 }
diff --git a/Engine/Source/Programs/Horde/Horde.Build/Projects/IStreamCollection.cs b/Engine/Source/Programs/Horde/Horde.Build/Projects/IStreamCollection.cs
--- a/Engine/Source/Programs/Horde/Horde.Build/Projects/IStreamCollection.cs
+++ b/Engine/Source/Programs/Horde/Horde.Build/Projects/IStreamCollection.cs
@@ -78,11 +78,11 @@ namespace HordeServer.Collections
 		/// </summary>
 		/// <param name="Stream">The stream to update</param>
 		/// <param name="TemplateRefId">The template ref id</param>
-		/// <param name="LastTriggerTime">New last trigger time for the schedule</param>
+		/// <param name="LastTriggerTimeUtc">New last trigger time for the schedule</param>
 		/// <param name="LastTriggerChange">New last trigger changelist for the schedule</param>
 		/// <param name="NewActiveJobs">New list of active jobs</param>
 		/// <returns>The updated stream if successful, null otherwise</returns>
-		Task<IStream?> TryUpdateScheduleTriggerAsync(IStream Stream, TemplateRefId TemplateRefId, DateTimeOffset? LastTriggerTime, int? LastTriggerChange, List<JobId> NewActiveJobs);
+		Task<IStream?> TryUpdateScheduleTriggerAsync(IStream Stream, TemplateRefId TemplateRefId, DateTime? LastTriggerTimeUtc, int? LastTriggerChange, List<JobId> NewActiveJobs);
 
 		/// <summary>
 		/// Delete a stream
diff --git a/Engine/Source/Programs/Horde/Horde.Build/Projects/Impl/StreamCollection.cs b/Engine/Source/Programs/Horde/Horde.Build/Projects/Impl/StreamCollection.cs
--- a/Engine/Source/Programs/Horde/Horde.Build/Projects/Impl/StreamCollection.cs
+++ b/Engine/Source/Programs/Horde/Horde.Build/Projects/Impl/StreamCollection.cs
@@ -189,8 +189,7 @@ namespace HordeServer.Collections.Impl
 				Schedule? Schedule = null;
 				if (Request.Schedule != null)
 				{
-					Schedule = Request.Schedule.ToModel();
-					Schedule.LastTriggerTime = Clock.UtcNow;
+					Schedule = Request.Schedule.ToModel(Clock.UtcNow);
 				}
 
 				// Add it to the list
@@ -330,18 +329,18 @@ namespace HordeServer.Collections.Impl
 		}
 
 		/// <inheritdoc/>
-		public async Task<IStream?> TryUpdateScheduleTriggerAsync(IStream StreamInterface, TemplateRefId TemplateRefId, DateTimeOffset? LastTriggerTime, int? LastTriggerChange, List<JobId> NewActiveJobs)
+		public async Task<IStream?> TryUpdateScheduleTriggerAsync(IStream StreamInterface, TemplateRefId TemplateRefId, DateTime? LastTriggerTimeUtc, int? LastTriggerChange, List<JobId> NewActiveJobs)
 		{
 			StreamDocument Stream = (StreamDocument)StreamInterface;
 			Schedule Schedule = Stream.Templates[TemplateRefId].Schedule!;
 
 			// Build the updates. MongoDB driver cannot parse TemplateRefId in expression tree; need to specify field name explicitly
 			List<UpdateDefinition<StreamDocument>> Updates = new List<UpdateDefinition<StreamDocument>>();
-			if (LastTriggerTime.HasValue && LastTriggerTime.Value != Schedule.LastTriggerTime)
+			if (LastTriggerTimeUtc.HasValue && LastTriggerTimeUtc.Value != Schedule.LastTriggerTime)
 			{
 				FieldDefinition<StreamDocument, DateTimeOffset> LastTriggerTimeField = $"{nameof(Stream.Templates)}.{TemplateRefId}.{nameof(Schedule)}.{nameof(Schedule.LastTriggerTime)}";
-				Updates.Add(Builders<StreamDocument>.Update.Set(LastTriggerTimeField, LastTriggerTime.Value));
-				Schedule.LastTriggerTime = LastTriggerTime.Value;
+				Updates.Add(Builders<StreamDocument>.Update.Set(LastTriggerTimeField, LastTriggerTimeUtc.Value));
+				Schedule.LastTriggerTimeUtc = LastTriggerTimeUtc.Value;
 			}
 			if (LastTriggerChange.HasValue && LastTriggerChange.Value > Schedule.LastTriggerChange)
 			{
diff --git a/Engine/Source/Programs/Horde/Horde.Build/Projects/Impl/StreamService.cs b/Engine/Source/Programs/Horde/Horde.Build/Projects/Impl/StreamService.cs
--- a/Engine/Source/Programs/Horde/Horde.Build/Projects/Impl/StreamService.cs
+++ b/Engine/Source/Programs/Horde/Horde.Build/Projects/Impl/StreamService.cs
@@ -114,12 +114,12 @@ namespace HordeServer.Services
 		/// </summary>
 		/// <param name="Stream">The stream to update</param>
 		/// <param name="TemplateRefId">The template ref id</param>
-		/// <param name="LastTriggerTime"></param>
+		/// <param name="LastTriggerTimeUtc"></param>
 		/// <param name="LastTriggerChange"></param>
 		/// <param name="AddJobs">Jobs to add</param>
 		/// <param name="RemoveJobs">Jobs to remove</param>
 		/// <returns>True if the stream was updated</returns>
-		public async Task<IStream?> UpdateScheduleTriggerAsync(IStream Stream, TemplateRefId TemplateRefId, DateTimeOffset? LastTriggerTime, int? LastTriggerChange, List<JobId> AddJobs, List<JobId> RemoveJobs)
+		public async Task<IStream?> UpdateScheduleTriggerAsync(IStream Stream, TemplateRefId TemplateRefId, DateTime? LastTriggerTimeUtc = null, int? LastTriggerChange = null, List<JobId>? AddJobs = null, List<JobId>? RemoveJobs = null)
 		{
 			IStream? NewStream = Stream;
 			while (NewStream != null)
@@ -134,9 +134,17 @@ namespace HordeServer.Services
 					break;
 				}
 
-				List<JobId> NewActiveJobs = TemplateRef.Schedule.ActiveJobs.Except(RemoveJobs).Union(AddJobs).ToList();
+				IEnumerable<JobId> NewActiveJobs = TemplateRef.Schedule.ActiveJobs;
+				if (RemoveJobs != null)
+				{
+					NewActiveJobs = NewActiveJobs.Except(RemoveJobs);
+				}
+				if (AddJobs != null)
+				{
+					NewActiveJobs = NewActiveJobs.Union(AddJobs);
+				}
 
-				NewStream = await Streams.TryUpdateScheduleTriggerAsync(NewStream, TemplateRefId, LastTriggerTime, LastTriggerChange, NewActiveJobs);
+				NewStream = await Streams.TryUpdateScheduleTriggerAsync(NewStream, TemplateRefId, LastTriggerTimeUtc, LastTriggerChange, NewActiveJobs.ToList());
 
 				if (NewStream != null)
 				{
diff --git a/Engine/Source/Programs/Horde/Horde.Build/Projects/ScheduleService.cs b/Engine/Source/Programs/Horde/Horde.Build/Projects/ScheduleService.cs
--- a/Engine/Source/Programs/Horde/Horde.Build/Projects/ScheduleService.cs
+++ b/Engine/Source/Programs/Horde/Horde.Build/Projects/ScheduleService.cs
@@ -25,6 +25,9 @@ using System.Security.Claims;
 using System.Threading;
 using System.Threading.Tasks;
 using TimeZoneConverter;
+using StackExchange.Redis;
+using EpicGames.Redis;
+using EpicGames.Serialization;
 
 namespace HordeServer.Services
 {
@@ -35,62 +38,49 @@ namespace HordeServer.Services
 	/// <summary>
 	/// Manipulates schedule instances
 	/// </summary>
-	public class ScheduleService : ElectedBackgroundService
+	public sealed class ScheduleService : IHostedService, IDisposable
 	{
-		/// <summary>
-		/// The downtime service instance
-		/// </summary>
-		IDowntimeService DowntimeService;
+		[RedisConverter(typeof(RedisCbConverter<QueueItem>))]
+		class QueueItem
+		{
+			[CbField("sid")]
+			public StreamId StreamId { get; set; }
 
-		/// <summary>
-		/// Collection of graph documents
-		/// </summary>
-		IGraphCollection Graphs;
+			[CbField("tid")]
+			public TemplateRefId TemplateId { get; set; }
 
-		/// <summary>
-		/// The perforce service singleton
-		/// </summary>
-		IPerforceService Perforce;
+			public QueueItem()
+			{
+			}
 
-		/// <summary>
-		/// Collection of job documents
-		/// </summary>
-		IJobCollection JobCollection;
+			public QueueItem(StreamId StreamId, TemplateRefId TemplateId)
+			{
+				this.StreamId = StreamId;
+				this.TemplateId = TemplateId;
+			}
 
-		/// <summary>
-		/// The job service
-		/// </summary>
-		JobService JobService;
+			public static DateTime GetTimeFromScore(double Score) => DateTime.UnixEpoch + TimeSpan.FromSeconds(Score);
+			public static double GetScoreFromTime(DateTime Time) => (Time.ToUniversalTime() - DateTime.UnixEpoch).TotalSeconds;
+		}
 
-		/// <summary>
-		/// The stream service
-		/// </summary>
+		IDowntimeService DowntimeService;
+		IGraphCollection Graphs;
+		IPerforceService Perforce;
+		IJobCollection JobCollection;
+		JobService JobService;
 		StreamService StreamService;
-
-		/// <summary>
-		/// The template service
-		/// </summary>
 		ITemplateCollection TemplateCollection;
-
-		/// <summary>
-		/// The timezone to use for schedules
-		/// </summary>
 		TimeZoneInfo TimeZone;
-
-		/// <summary>
-		/// The clock instance
-		/// </summary>
 		IClock Clock;
-
-		/// <summary>
-		/// The logging instance
-		/// </summary>
 		ILogger Logger;
+		RedisService Redis;
+		RedisSortedSet<QueueItem> Queue;
+		BackgroundTask BackgroundTask;
 
 		/// <summary>
 		/// Constructor
 		/// </summary>
-		/// <param name="DatabaseService">Database service instance</param>
+		/// <param name="Redis"></param>
 		/// <param name="Graphs">Collection of graph documents</param>
 		/// <param name="DowntimeService">The downtime service</param>
 		/// <param name="Perforce">The Perforce service</param>
@@ -101,8 +91,7 @@ namespace HordeServer.Services
 		/// <param name="Clock"></param>
 		/// <param name="Settings">Settings for the server</param>
 		/// <param name="Logger">Logging instance</param>
-		public ScheduleService(DatabaseService DatabaseService, IGraphCollection Graphs, IDowntimeService DowntimeService, IPerforceService Perforce, IJobCollection JobCollection, JobService JobService, StreamService StreamService, ITemplateCollection TemplateCollection, IClock Clock, IOptionsMonitor<ServerSettings> Settings, ILogger<ScheduleService> Logger)
-			: base(DatabaseService, new ObjectId("6035593e6d721d80fb9efa5c"), Logger)
+		public ScheduleService(RedisService Redis, IGraphCollection Graphs, IDowntimeService DowntimeService, IPerforceService Perforce, IJobCollection JobCollection, JobService JobService, StreamService StreamService, ITemplateCollection TemplateCollection, IClock Clock, IOptionsMonitor<ServerSettings> Settings, ILogger<ScheduleService> Logger)
 		{
 			this.Graphs = Graphs;
 			this.DowntimeService = DowntimeService;
@@ -117,94 +106,223 @@ namespace HordeServer.Services
 			this.TimeZone = (ScheduleTimeZone == null) ? TimeZoneInfo.Local : TZConvert.GetTimeZoneInfo(ScheduleTimeZone);
 
 			this.Logger = Logger;
+
+			this.Redis = Redis;
+			this.Queue = new RedisSortedSet<QueueItem>(Redis.Database, "scheduler/queue");
+			this.BackgroundTask = new BackgroundTask(RunAsync);
 		}
 
-		/// <summary>
-		/// Execute any scheduled tasks
-		/// </summary>
-		/// <param name="StoppingToken">Cancellation token for stopping the service</param>
-		/// <returns>Async task</returns>
-		protected override async Task<DateTime> TickLeaderAsync(CancellationToken StoppingToken)
+		/// <inheritdoc/>
+		public async Task StartAsync(CancellationToken CancellationToken)
+		{
+			await RefreshSchedulesAsync(Clock.UtcNow);
+			BackgroundTask.Start();
+		}
+
+		/// <inheritdoc/>
+		public async Task StopAsync(CancellationToken CancellationToken)
+		{
+			await BackgroundTask.StopAsync();
+		}
+
+		/// <inheritdoc/>
+		public void Dispose()
 		{
-			DateTime NextTickTime = DateTime.UtcNow.AddMinutes(1.0);
-			if (!DowntimeService.IsDowntimeActive)
+			BackgroundTask?.Dispose();
+		}
+
+		async Task RunAsync(CancellationToken CancellationToken)
+		{
+			await RefreshSchedulesAsync(Clock.UtcNow);
+
+			for(; ;)
 			{
-				Logger.LogInformation("Updating scheduled triggers...");
-				Stopwatch Stopwatch = Stopwatch.StartNew();
-				using (IScope Scope = GlobalTracer.Instance.BuildSpan("ScheduleService Tick").StartActive())
+				TimeSpan DelayTime = await TickAsync(CancellationToken);
+				if (CancellationToken.IsCancellationRequested)
 				{
-					NextTickTime = await UpdateSchedulesAsync();
+					break;
 				}
-				Stopwatch.Stop();
-				Logger.LogInformation("Scheduling triggers took {ElapsedTime} ms", Stopwatch.ElapsedMilliseconds);
+				await Task.Delay(DelayTime, CancellationToken);
 			}
-			return NextTickTime;
 		}
 
-		/// <summary>
-		/// Update all the schedules
-		/// </summary>
-		/// <returns>Next time that checkes need to be updated</returns>
-		[System.Diagnostics.CodeAnalysis.SuppressMessage("Design", "CA1031:Do not catch general exception types", Justification = "<Pending>")]
-		async Task<DateTime> UpdateSchedulesAsync()
+		async Task<TimeSpan> TickAsync(CancellationToken CancellationToken)
 		{
-			// Find the next time at which each schedule should run
-			DateTimeOffset Now = Clock.UtcNow;
-			DateTimeOffset NextUpdateTime = Now.AddMinutes(1.0);
+			DateTime UtcNow = Clock.UtcNow;
 
-			// Find all the enabled schedules
-			List<IStream> Streams = await StreamService.GetStreamsAsync();
-			foreach (IStream Stream in Streams)
+			// Keep updating schedules
+			while (!CancellationToken.IsCancellationRequested)
 			{
-				if (Stream.IsPaused(Clock.UtcNow))
+				// Get the item with the lowest score (ie. the one that hasn't been updated in the longest time)
+				SortedSetEntry<QueueItem>[] Entries = await Queue.RangeByScoreWithScoresAsync(Take: 1);
+				if (Entries.Length == 0)
+				{
+					await RefreshSchedulesAsync(UtcNow);
+					return TimeSpan.FromSeconds(10.0);
+				}
+
+				SortedSetEntry<QueueItem> Entry = Entries[0];
+				QueueItem Item = Entry.Element;
+				DateTime ScheduledTime = QueueItem.GetTimeFromScore(Entry.Score);
+
+				// Check if the first entry is due to run
+				if (UtcNow < ScheduledTime)
+				{
+					TimeSpan Delay = (ScheduledTime - UtcNow).Add(TimeSpan.FromSeconds(1.0));
+					return Delay;
+				}
+
+				// If the stream id is null, it's a request to update the schedules
+				if (Item.StreamId == StreamId.Empty)
 				{
-					Logger.LogDebug("Skipping schedule update for stream {StreamId}. It has been paused until {PausedUntil} with comment '{PauseComment}'.", Stream.Id, Stream.PausedUntil, Stream.PauseComment);
+					await RefreshSchedulesAsync(UtcNow);
 					continue;
 				}
-				
-				foreach ((TemplateRefId TemplateRefId, TemplateRef TemplateRef) in Stream.Templates)
+
+				// Get the stream and template
+				IStream? Stream = await StreamService.GetStreamAsync(Item.StreamId);
+				if (Stream == null || !Stream.Templates.TryGetValue(Item.TemplateId, out TemplateRef? TemplateRef) || TemplateRef.Schedule == null)
 				{
-					Schedule? Schedule = TemplateRef.Schedule;
-					if (Schedule != null && Schedule.Enabled)
-					{
-						using IDisposable Scope = Logger.BeginScope("Updating schedule {StreamId}/{TemplateRefId}", Stream.Id, TemplateRefId);
+					await Queue.RemoveAsync(Item);
+					continue;
+				}
 
-						DateTimeOffset? NextTriggerTime = await UpdateScheduleAsync(Stream, TemplateRefId, TemplateRef, Schedule, Now);
-						if (NextTriggerTime.HasValue && NextTriggerTime.Value < NextUpdateTime)
+				// Calculate the next trigger time for this schedule; it may have changed since the queue item was created.
+				DateTime? NextUpdateTimeUtc = TemplateRef.Schedule.GetNextTriggerTimeUtc(TimeZone);
+				if (NextUpdateTimeUtc == null)
+				{
+					await Queue.RemoveAsync(Item);
+					continue;
+				}
+
+				// Check if we're ready to trigger now
+				if (UtcNow < NextUpdateTimeUtc)
+				{
+					await TrySetNextUpdateTime(Entry, NextUpdateTimeUtc.Value);
+					continue;
+				}
+
+				// Try to prevent any updates being done for 5 minutes, but keep extending that again every minute. This
+				// creates a reservation for the stream update that will expire if the pod is terminated.
+				TimeSpan ReservationTime = TimeSpan.FromMinutes(5.0);
+				if (!await TrySetNextUpdateTime(Entry, DateTime.UtcNow + ReservationTime))
+				{
+					continue;
+				}
+
+				// Try to do the update
+				try
+				{
+					Task TriggerTask = Task.Run(() => TriggerAsync(Stream, Item.TemplateId, TemplateRef, UtcNow));
+					for (; ; )
+					{
+						Task DelayTask = Task.Delay(TimeSpan.FromMinutes(1.0));
+						if (await Task.WhenAny(DelayTask, TriggerTask) == TriggerTask)
 						{
-							NextUpdateTime = NextTriggerTime.Value;
+							break;
 						}
+						await SetNextUpdateTime(Item, Clock.UtcNow + ReservationTime);
 					}
+					await TriggerTask;
+				}
+				catch (OperationCanceledException)
+				{
+					throw;
+				}
+				catch (Exception Ex)
+				{
+					Logger.LogError(Ex, "Error while updating schedule for {StreamId}/{TemplateId}", Item.StreamId, Item.TemplateId);
+				}
+
+				// Update the next trigger time
+				DateTime? NextTime = TemplateRef.Schedule.GetNextTriggerTimeUtc(UtcNow, TimeZone);
+				if (NextTime == null)
+				{
+					await Queue.RemoveAsync(Item);
+				}
+				else
+				{
+					await SetNextUpdateTime(Item, NextTime.Value);
 				}
 			}
+			return TimeSpan.Zero;
+		}
 
-			return NextUpdateTime.UtcDateTime;
+		internal async Task ResetAsync()
+		{
+			await Redis.Database.KeyDeleteAsync(Queue.Key);
+		}
+
+		internal async Task TickForTestingAsync()
+		{
+			await RefreshSchedulesAsync(Clock.UtcNow);
+			await TickAsync(CancellationToken.None);
+		}
+
+		Task<bool> TrySetNextUpdateTime(SortedSetEntry<QueueItem> Entry, DateTime NextTimeUtc)
+		{
+			ITransaction Transaction = Redis.Database.CreateTransaction();
+			Transaction.AddCondition(Condition.SortedSetEqual(Queue.Key, Entry.ElementValue, Entry.Score));
+			_ = Transaction.With(Queue).AddAsync(Entry.Element, QueueItem.GetScoreFromTime(NextTimeUtc));
+			return Transaction.ExecuteAsync();
+		}
+
+		async Task SetNextUpdateTime(QueueItem Item, DateTime NextTimeUtc)
+		{
+			await Queue.AddAsync(Item, QueueItem.GetScoreFromTime(NextTimeUtc));
 		}
 
 		/// <summary>
-		/// Updates a schedule
+		/// Get the current set of streams and ensure there's an entry for each item
 		/// </summary>
-		/// <returns>Next time to trigger the schedule</returns>
-		[System.Diagnostics.CodeAnalysis.SuppressMessage("Design", "CA1031:Do not catch general exception types", Justification = "<Pending>")]
-		private async Task<DateTimeOffset?> UpdateScheduleAsync(IStream Stream, TemplateRefId TemplateRefId, TemplateRef TemplateRef, Schedule Schedule, DateTimeOffset Now)
+		public async Task RefreshSchedulesAsync(DateTime UtcNow)
 		{
-			using IScope Scope = GlobalTracer.Instance.BuildSpan("ScheduleService.UpdateScheduleAsync").StartActive();
-			Scope.Span.SetTag("Stream.Id", Stream.Id);
-			Scope.Span.SetTag("Stream.Name", Stream.Name);
-			
-			// Check if we can run the trigger
-			DateTimeOffset? NextTriggerTime = Schedule.GetNextTriggerTime(TimeZone);
-			if (!NextTriggerTime.HasValue)
+			List<SortedSetEntry<QueueItem>> QueueItems = new List<SortedSetEntry<QueueItem>>();
+
+			List<IStream> Streams = await StreamService.GetStreamsAsync();
+			foreach (IStream Stream in Streams)
 			{
-				Logger.LogInformation("Schedule {StreamId}/{TemplateRefId} is disabled", Stream.Id, TemplateRefId);
-				return NextTriggerTime;
+				foreach ((TemplateRefId TemplateId, TemplateRef TemplateRef) in Stream.Templates)
+				{
+					if (TemplateRef.Schedule != null)
+					{
+						DateTime? NextTriggerTimeUtc = TemplateRef.Schedule.GetNextTriggerTimeUtc(TimeZone);
+						if (NextTriggerTimeUtc != null)
+						{
+							double Score = GetQueueItemScore(NextTriggerTimeUtc.Value);
+							QueueItems.Add(new SortedSetEntry<QueueItem>(new QueueItem(Stream.Id, TemplateId), Score));
+						}
+					}
+				}
 			}
-			if (NextTriggerTime.Value > Now)
+			QueueItems.Add(new SortedSetEntry<QueueItem>(new QueueItem(StreamId.Empty, TemplateRefId.Empty), GetQueueItemScore(UtcNow.AddMinutes(1.0))));
+
+			await Queue.AddAsync(QueueItems.ToArray());
+		}
+
+		/// <summary>
+		/// Trigger a schedule to run
+		/// </summary>
+		/// <param name="Stream">Stream for the schedule</param>
+		/// <param name="TemplateId"></param>
+		/// <param name="TemplateRef"></param>
+		/// <param name="UtcNow"></param>
+		/// <returns>Async task</returns>
+		private async Task TriggerAsync(IStream Stream, TemplateRefId TemplateId, TemplateRef TemplateRef, DateTime UtcNow)
+		{
+			using IScope Scope = GlobalTracer.Instance.BuildSpan("ScheduleService.TriggerAsync").StartActive();
+			Scope.Span.SetTag("StreamId", Stream.Id);
+			Scope.Span.SetTag("TemplateId", TemplateId);
+
+			Stopwatch Stopwatch = Stopwatch.StartNew();
+			Logger.LogInformation("Updating schedule for {StreamId} template {TemplateId}", Stream.Id, TemplateId);
+
+			// Make sure the schedule is valid
+			Schedule? Schedule = TemplateRef.Schedule;
+			if(Schedule == null)
 			{
-				Logger.LogDebug("Schedule {StreamId}/{TemplateRefId} will next trigger at {NextTriggerTime}", Stream.Id, TemplateRefId, NextTriggerTime.Value);
-				return NextTriggerTime;
+				return;
 			}
-			Logger.LogInformation("Schedule {StreamId}/{TemplateRefId} is ready to trigger (last: {Change} at {Time})", Stream.Id, TemplateRefId, Schedule.LastTriggerChange, Schedule.LastTriggerTime);
 
 			// Get a list of jobs that we need to remove
 			List<JobId> RemoveJobIds = new List<JobId>();
@@ -217,23 +335,28 @@ namespace HordeServer.Services
 					RemoveJobIds.Add(ActiveJobId);
 				}
 			}
-			if (RemoveJobIds.Count > 0)
+			await StreamService.UpdateScheduleTriggerAsync(Stream, TemplateId, RemoveJobs: RemoveJobIds);
+
+			// If the stream is paused, bail out
+			if (Stream.IsPaused(Clock.UtcNow))
 			{
-				await StreamService.UpdateScheduleTriggerAsync(Stream, TemplateRefId, Now, null, new List<JobId>(), RemoveJobIds);
+				Logger.LogDebug("Skipping schedule update for stream {StreamId}. It has been paused until {PausedUntil} with comment '{PauseComment}'.", Stream.Id, Stream.PausedUntil, Stream.PauseComment);
+				return;
 			}
 
 			// Trigger this schedule
 			try
 			{
-				await TriggerAsync(Stream, TemplateRefId, TemplateRef, Schedule, Schedule.ActiveJobs.Count - RemoveJobIds.Count, Now);
+				await TriggerAsync(Stream, TemplateId, TemplateRef, Schedule, Schedule.ActiveJobs.Count - RemoveJobIds.Count, UtcNow);
 			}
 			catch (Exception Ex)
 			{
-				Logger.LogError(Ex, "Failed to start schedule {StreamId}/{TemplateRefId}", Stream.Id, TemplateRefId);
+				Logger.LogError(Ex, "Failed to start schedule {StreamId}/{TemplateRefId}", Stream.Id, TemplateId);
 			}
 
-			// Return the next update time
-			return Schedule.GetNextTriggerTime(Now, TimeZone);
+			// Print some timing info
+			Stopwatch.Stop();
+			Logger.LogInformation("Updated schedule for {StreamId} template {TemplateId} in {TimeSeconds}ms", Stream.Id, TemplateId, (long)Stopwatch.Elapsed.TotalMilliseconds);
 		}
 
 		/// <summary>
@@ -244,15 +367,10 @@ namespace HordeServer.Services
 		/// <param name="TemplateRef"></param>
 		/// <param name="Schedule"></param>
 		/// <param name="NumActiveJobs"></param>
-		/// <param name="Now"></param>
+		/// <param name="UtcNow">The current time</param>
 		/// <returns>Async task</returns>
-		private async Task TriggerAsync(IStream Stream, TemplateRefId TemplateRefId, TemplateRef TemplateRef, Schedule Schedule, int NumActiveJobs, DateTimeOffset Now)
+		private async Task TriggerAsync(IStream Stream, TemplateRefId TemplateRefId, TemplateRef TemplateRef, Schedule Schedule, int NumActiveJobs, DateTime UtcNow)
 		{
-			using IScope Scope = GlobalTracer.Instance.BuildSpan("ScheduleService.TriggerAsync").StartActive();
-			Scope.Span.SetTag("Stream.Id", Stream.Id);
-			Scope.Span.SetTag("Stream.Name", Stream.Name);
-			Scope.Span.SetTag("TemplateRefId", TemplateRefId);
-			
 			// Check we're not already at the maximum number of allowed jobs
 			if (Schedule.MaxActive != 0 && NumActiveJobs >= Schedule.MaxActive)
 			{
@@ -377,7 +495,7 @@ namespace HordeServer.Services
 				List<string> DefaultArguments = Template.GetDefaultArguments();
 				IJob NewJob = await JobService.CreateJobAsync(null, Stream, TemplateRefId, Template.Id, Graph, Template.Name, Change, CodeChange, null, null, null, Template.Priority, null, null, TemplateRef.ChainedJobs, TemplateRef.ShowUgsBadges, TemplateRef.ShowUgsAlerts, TemplateRef.NotificationChannel, TemplateRef.NotificationChannelFilter, DefaultArguments);
 				Logger.LogInformation("Started new job for {StreamName} template {TemplateName} at CL {Change} (Code CL {CodeChange}): {JobId}", Stream.Id, TemplateRef.Name, Change, CodeChange, NewJob.Id);
-				await StreamService.UpdateScheduleTriggerAsync(Stream, TemplateRefId, Now, Change, new List<JobId> { NewJob.Id }, new List<JobId>());
+				await StreamService.UpdateScheduleTriggerAsync(Stream, TemplateRefId, UtcNow, Change, new List<JobId> { NewJob.Id }, new List<JobId>());
 			}
 		}
 
diff --git a/Engine/Source/Programs/Horde/Horde.Build/Utilities/Clock.cs b/Engine/Source/Programs/Horde/Horde.Build/Utilities/Clock.cs
--- a/Engine/Source/Programs/Horde/Horde.Build/Utilities/Clock.cs
+++ b/Engine/Source/Programs/Horde/Horde.Build/Utilities/Clock.cs
@@ -32,12 +32,14 @@ namespace HordeCommon
 	/// </summary>
 	public class FakeClock : IClock
 	{
+		DateTime UtcNowPrivate;
+
 		/// <summary>
 		/// Constructor
 		/// </summary>
 		public FakeClock()
 		{
-			UtcNow = DateTime.UtcNow;
+			UtcNowPrivate = DateTime.UtcNow;
 		}
 
 		/// <summary>
@@ -47,10 +49,14 @@ namespace HordeCommon
 		/// <param name="Period">Time span to advance</param>
 		public void Advance(TimeSpan Period)
 		{
-			UtcNow = UtcNow.Add(Period);
+			UtcNowPrivate = UtcNowPrivate.Add(Period);
 		}
 
 		/// <inheritdoc/>
-		public DateTime UtcNow { get; set; }
+		public DateTime UtcNow
+		{ 
+			get => UtcNowPrivate;
+			set => UtcNowPrivate = value.ToUniversalTime(); 
+		}
 	}
 }
\ No newline at end of file
diff --git a/Engine/Source/Programs/Shared/EpicGames.Redis/RedisSortedSet.cs b/Engine/Source/Programs/Shared/EpicGames.Redis/RedisSortedSet.cs
--- a/Engine/Source/Programs/Shared/EpicGames.Redis/RedisSortedSet.cs
+++ b/Engine/Source/Programs/Shared/EpicGames.Redis/RedisSortedSet.cs
@@ -20,6 +20,11 @@ namespace EpicGames.Redis
 		/// </summary>
 		public readonly T Element { get; }
 
+		/// <summary>
+		/// The encoded element value
+		/// </summary>
+		public readonly RedisValue ElementValue { get; }
+
 		/// <summary>
 		/// Score for the entry
 		/// </summary>
@@ -32,6 +37,7 @@ namespace EpicGames.Redis
 		public SortedSetEntry(SortedSetEntry Entry)
 		{
 			this.Element = RedisSerializer.Deserialize<T>(Entry.Element);
+			this.ElementValue = Entry.Element;
 			this.Score = Entry.Score;
 		}
 
@@ -43,8 +49,20 @@ namespace EpicGames.Redis
 		public SortedSetEntry(T Element, double Score)
 		{
 			this.Element = Element;
+			this.ElementValue = RedisSerializer.Serialize<T>(Element);
 			this.Score = Score;
 		}
+
+		/// <summary>
+		/// Deconstruct this item into a tuple
+		/// </summary>
+		/// <param name="OutElement"></param>
+		/// <param name="OutScore"></param>
+		public void Deconstruct(out T OutElement, out double OutScore)
+		{
+			OutElement = Element;
+			OutScore = Score;
+		}
 	}
 
 	/// <summary>
@@ -79,7 +97,7 @@ namespace EpicGames.Redis
 		/// <inheritdoc cref="IDatabaseAsync.SortedSetAddAsync(RedisKey, RedisValue, double, When, CommandFlags)"/>
 		public Task<long> AddAsync(SortedSetEntry<TElement>[] Values, When When = When.Always, CommandFlags Flags = CommandFlags.None)
 		{
-			SortedSetEntry[] Untyped = Array.ConvertAll(Values, x => new SortedSetEntry(RedisSerializer.Serialize(x.Element), x.Score));
+			SortedSetEntry[] Untyped = Array.ConvertAll(Values, x => new SortedSetEntry(x.ElementValue, x.Score));
 			return Database.SortedSetAddAsync(Key, Untyped, When, Flags);
 		}
 
