Horde: Missing file.

diff --git a/Engine/Source/Programs/Horde/Horde.Build/Utilities/BackgroundTask.cs b/Engine/Source/Programs/Horde/Horde.Build/Utilities/BackgroundTask.cs
new file mode 100644
--- /dev/null
+++ b/Engine/Source/Programs/Horde/Horde.Build/Utilities/BackgroundTask.cs
@@ -0,0 +1,66 @@
+// Copyright Epic Games, Inc. All Rights Reserved.
+
+using System;
+using System.Threading;
+using System.Threading.Tasks;
+
+namespace Horde.Build.Utilities
+{
+	/// <summary>
+	/// Runs a task in the background and allows stopping it on demand
+	/// </summary>
+	public sealed class BackgroundTask : IDisposable
+	{
+		Func<CancellationToken, Task> RunTask;
+		Task Task = Task.CompletedTask;
+		CancellationTokenSource CancellationTokenSource = new CancellationTokenSource();
+
+		/// <summary>
+		/// Constructor
+		/// </summary>
+		/// <param name="RunTask"></param>
+		public BackgroundTask(Func<CancellationToken, Task> RunTask)
+		{
+			this.RunTask = RunTask;
+		}
+
+		/// <inheritdoc/>
+		public void Dispose()
+		{
+			if (!Task.IsCompleted)
+			{
+				CancellationTokenSource.Cancel();
+				StopAsync().Wait();
+			}
+			CancellationTokenSource.Dispose();
+		}
+
+		/// <summary>
+		/// Starts the task
+		/// </summary>
+		public void Start()
+		{
+			if (!Task.IsCompleted)
+			{
+				throw new InvalidOperationException("Background task is already running");
+			}
+			Task = Task.Run(() => RunTask(CancellationTokenSource.Token), CancellationTokenSource.Token);
+		}
+
+		/// <summary>
+		/// Signals the cancellation token and waits for the task to finish
+		/// </summary>
+		/// <returns></returns>
+		public async Task StopAsync()
+		{
+			try
+			{
+				CancellationTokenSource.Cancel();
+				await Task;
+			}
+			catch (OperationCanceledException)
+			{
+			}
+		}
+	}
+}
