Horde: Fix build error.

diff --git a/Engine/Source/Programs/Horde/Horde.Build/Projects/ScheduleService.cs b/Engine/Source/Programs/Horde/Horde.Build/Projects/ScheduleService.cs
--- a/Engine/Source/Programs/Horde/Horde.Build/Projects/ScheduleService.cs
+++ b/Engine/Source/Programs/Horde/Horde.Build/Projects/ScheduleService.cs
@@ -289,13 +289,13 @@ namespace HordeServer.Services
 						DateTime? NextTriggerTimeUtc = TemplateRef.Schedule.GetNextTriggerTimeUtc(TimeZone);
 						if (NextTriggerTimeUtc != null)
 						{
-							double Score = GetQueueItemScore(NextTriggerTimeUtc.Value);
+							double Score = QueueItem.GetScoreFromTime(NextTriggerTimeUtc.Value);
 							QueueItems.Add(new SortedSetEntry<QueueItem>(new QueueItem(Stream.Id, TemplateId), Score));
 						}
 					}
 				}
 			}
-			QueueItems.Add(new SortedSetEntry<QueueItem>(new QueueItem(StreamId.Empty, TemplateRefId.Empty), GetQueueItemScore(UtcNow.AddMinutes(1.0))));
+			QueueItems.Add(new SortedSetEntry<QueueItem>(new QueueItem(StreamId.Empty, TemplateRefId.Empty), QueueItem.GetScoreFromTime(UtcNow.AddMinutes(1.0))));
 
 			await Queue.AddAsync(QueueItems.ToArray());
 		}
