Horde: Do not remove untracked files by default when running a conform. Add a new conform type (a "full conform") which will erase everything, including Temp and Saved folders.

diff --git a/Engine/Source/Programs/Horde/Horde.Agent/Execution/PerforceExecutor.cs b/Engine/Source/Programs/Horde/Horde.Agent/Execution/PerforceExecutor.cs
--- a/Engine/Source/Programs/Horde/Horde.Agent/Execution/PerforceExecutor.cs
+++ b/Engine/Source/Programs/Horde/Horde.Agent/Execution/PerforceExecutor.cs
@@ -212,7 +212,7 @@ namespace HordeAgent.Execution
 			await Workspace.CleanAsync(CancellationToken);
 		}
 
-		public static async Task ConformAsync(DirectoryReference RootDir, IList<AgentWorkspace> PendingWorkspaces, ILogger Logger, CancellationToken CancellationToken)
+		public static async Task ConformAsync(DirectoryReference RootDir, IList<AgentWorkspace> PendingWorkspaces, bool bRemoveUntrackedFiles, ILogger Logger, CancellationToken CancellationToken)
 		{
 			// Print out all the workspaces we're going to sync
 			Logger.LogInformation("Workspaces:");
@@ -302,8 +302,11 @@ namespace HordeAgent.Execution
 
 				// Build a set of directories to protect
 				HashSet<DirectoryReference> ProtectDirs = new HashSet<DirectoryReference>();
-				ProtectDirs.Add(DirectoryReference.Combine(RootDir, "Temp"));
-				ProtectDirs.Add(DirectoryReference.Combine(RootDir, "Saved"));
+				if (!bRemoveUntrackedFiles)
+				{
+					ProtectDirs.Add(DirectoryReference.Combine(RootDir, "Temp"));
+					ProtectDirs.Add(DirectoryReference.Combine(RootDir, "Saved"));
+				}
 				ProtectDirs.UnionWith(Workspaces.Select(x => x.MetadataDir));
 
 				// Delete all the directories which aren't a workspace root
@@ -343,9 +346,18 @@ namespace HordeAgent.Execution
 					PopulateRequests.Add(new PopulateRequest(PerforceClient, Workspace.StreamName, Workspace.View));
 				}
 
-				ManagedWorkspace Repository = WorkspaceGroup.First().Repository;
-				Tuple<int, StreamSnapshot>[] StreamStates = await Repository.PopulateCleanAsync(PopulateRequests, false, CancellationToken);
-				SyncFuncs.Add(() => Repository.PopulateSyncAsync(PopulateRequests, StreamStates, false, CancellationToken));
+				WorkspaceInfo? FirstWorkspace = WorkspaceGroup.First();
+				if (PopulateRequests.Count == 1 && !FirstWorkspace.bRemoveUntrackedFiles && !bRemoveUntrackedFiles)
+				{
+					await FirstWorkspace.CleanAsync(CancellationToken);
+					SyncFuncs.Add(() => FirstWorkspace.SyncAsync(-1, -1, null, Logger, CancellationToken));
+				}
+				else
+				{
+					ManagedWorkspace Repository = FirstWorkspace.Repository;
+					Tuple<int, StreamSnapshot>[] StreamStates = await Repository.PopulateCleanAsync(PopulateRequests, false, CancellationToken);
+					SyncFuncs.Add(() => Repository.PopulateSyncAsync(PopulateRequests, StreamStates, false, CancellationToken));
+				}
 			}
 			foreach (Func<Task> SyncFunc in SyncFuncs)
 			{
diff --git a/Engine/Source/Programs/Horde/Horde.Agent/Services/WorkerService.cs b/Engine/Source/Programs/Horde/Horde.Agent/Services/WorkerService.cs
--- a/Engine/Source/Programs/Horde/Horde.Agent/Services/WorkerService.cs
+++ b/Engine/Source/Programs/Horde/Horde.Agent/Services/WorkerService.cs
@@ -881,13 +881,14 @@ namespace HordeAgent.Services
 			ConformLogger.LogInformation("Conforming, lease {LeaseId}", LeaseId);
 			TerminateProcesses(ConformLogger, CancellationToken);
 
+			bool RemoveUntrackedFiles = ConformTask.RemoveUntrackedFiles;
 			IList<AgentWorkspace> PendingWorkspaces = ConformTask.Workspaces;
 			for (; ;)
 			{
 				// Run the conform task
 				if (Settings.Executor == ExecutorType.Perforce && Settings.PerforceExecutor.RunConform)
 				{
-					await PerforceExecutor.ConformAsync(WorkingDir, PendingWorkspaces, ConformLogger, CancellationToken);
+					await PerforceExecutor.ConformAsync(WorkingDir, PendingWorkspaces, RemoveUntrackedFiles, ConformLogger, CancellationToken);
 				}
 				else
 				{
@@ -898,6 +899,7 @@ namespace HordeAgent.Services
 				UpdateAgentWorkspacesRequest Request = new UpdateAgentWorkspacesRequest();
 				Request.AgentId = AgentId;
 				Request.Workspaces.AddRange(PendingWorkspaces);
+				Request.RemoveUntrackedFiles = RemoveUntrackedFiles;
 
 				UpdateAgentWorkspacesResponse Response = await RpcConnection.InvokeAsync(x => x.UpdateAgentWorkspacesAsync(Request, null, null, CancellationToken), new RpcContext(), CancellationToken);
 				if (!Response.Retry)
@@ -908,6 +910,7 @@ namespace HordeAgent.Services
 
 				ConformLogger.LogInformation("Pending workspaces have changed - running conform again...");
 				PendingWorkspaces = Response.PendingWorkspaces;
+				RemoveUntrackedFiles = Response.RemoveUntrackedFiles;
 			}
 
 			return LeaseResult.Success;
diff --git a/Engine/Source/Programs/Horde/Horde.Build/Agents/IAgent.cs b/Engine/Source/Programs/Horde/Horde.Build/Agents/IAgent.cs
--- a/Engine/Source/Programs/Horde/Horde.Build/Agents/IAgent.cs
+++ b/Engine/Source/Programs/Horde/Horde.Build/Agents/IAgent.cs
@@ -323,6 +323,7 @@ namespace HordeServer.Models
 				{
 					Details["Type"] = "Conform";
 					Details["LogId"] = ConformTask.LogId;
+					Details["Full"] = ConformTask.RemoveUntrackedFiles.ToString();
 				}
 
 				ExecuteJobTask JobTask;
@@ -475,6 +476,11 @@ namespace HordeServer.Models
 		/// </summary>
 		public bool RequestConform { get; }
 
+		/// <summary>
+		/// Whether a full conform is requested
+		/// </summary>
+		public bool RequestFullConform { get; }
+
 		/// <summary>
 		/// Whether a machine restart is requested
 		/// </summary>
diff --git a/Engine/Source/Programs/Horde/Horde.Build/Agents/IAgentCollection.cs b/Engine/Source/Programs/Horde/Horde.Build/Agents/IAgentCollection.cs
--- a/Engine/Source/Programs/Horde/Horde.Build/Agents/IAgentCollection.cs
+++ b/Engine/Source/Programs/Horde/Horde.Build/Agents/IAgentCollection.cs
@@ -83,6 +83,7 @@ namespace HordeServer.Collections
 		/// <param name="Agent">Agent instance</param>
 		/// <param name="bEnabled">Whether the agent is enabled or not</param>
 		/// <param name="bRequestConform">Whether to request a conform job be run</param>
+		/// <param name="bRequestFullConform">Whether to request a full conform job be run</param>
 		/// <param name="bRequestRestart">Whether to request the machine be restarted</param>
 		/// <param name="bRequestShutdown">Whether to request the machine be shut down</param>
 		/// <param name="Channel">Override for the desired software channel</param>
@@ -90,7 +91,7 @@ namespace HordeServer.Collections
 		/// <param name="Acl">New ACL for this agent</param>
 		/// <param name="Comment">New comment</param>
 		/// <returns>Version of the software that needs to be installed on the agent. Null if the agent is running the correct version.</returns>
-		Task<IAgent?> TryUpdateSettingsAsync(IAgent Agent, bool? bEnabled = null, bool? bRequestConform = null, bool? bRequestRestart = null, bool? bRequestShutdown = null, AgentSoftwareChannelName? Channel = null, List<PoolId>? Pools = null, Acl? Acl = null, string? Comment = null);
+		Task<IAgent?> TryUpdateSettingsAsync(IAgent Agent, bool? bEnabled = null, bool? bRequestConform = null, bool? bRequestFullConform = null, bool? bRequestRestart = null, bool? bRequestShutdown = null, AgentSoftwareChannelName? Channel = null, List<PoolId>? Pools = null, Acl? Acl = null, string? Comment = null);
 
 		/// <summary>
 		/// Update the current workspaces for an agent.
diff --git a/Engine/Source/Programs/Horde/Horde.Build/Agents/Impl/AgentCollection.cs b/Engine/Source/Programs/Horde/Horde.Build/Agents/Impl/AgentCollection.cs
--- a/Engine/Source/Programs/Horde/Horde.Build/Agents/Impl/AgentCollection.cs
+++ b/Engine/Source/Programs/Horde/Horde.Build/Agents/Impl/AgentCollection.cs
@@ -91,12 +91,17 @@ namespace HordeServer.Collections.Impl
 
 			public List<PoolId> DynamicPools { get; set; } = new List<PoolId>();
 			public List<PoolId> Pools { get; set; } = new List<PoolId>();
+
+			[BsonIgnoreIfDefault, BsonDefaultValue(false)]
 			public bool RequestConform { get; set; }
 
-			[BsonIgnoreIfNull]
+			[BsonIgnoreIfDefault, BsonDefaultValue(false)]
+			public bool RequestFullConform { get; set; }
+
+			[BsonIgnoreIfDefault, BsonDefaultValue(false)]
 			public bool RequestRestart { get; set; }
 
-			[BsonIgnoreIfNull]
+			[BsonIgnoreIfDefault, BsonDefaultValue(false)]
 			public bool RequestShutdown { get; set; }
 
 			public List<AgentWorkspace> Workspaces { get; set; } = new List<AgentWorkspace>();
@@ -249,7 +254,7 @@ namespace HordeServer.Collections.Impl
 		}
 
 		/// <inheritdoc/>
-		public async Task<IAgent?> TryUpdateSettingsAsync(IAgent AgentInterface, bool? Enabled = null, bool? RequestConform = null, bool? RequestRestart = null, bool? RequestShutdown = null, AgentSoftwareChannelName? Channel = null, List<PoolId>? Pools = null, Acl? Acl = null, string? Comment = null)
+		public async Task<IAgent?> TryUpdateSettingsAsync(IAgent AgentInterface, bool? Enabled = null, bool? RequestConform = null, bool? RequestFullConform = null, bool? RequestRestart = null, bool? RequestShutdown = null, AgentSoftwareChannelName? Channel = null, List<PoolId>? Pools = null, Acl? Acl = null, string? Comment = null)
 		{
 			AgentDocument Agent = (AgentDocument)AgentInterface;
 
@@ -270,6 +275,11 @@ namespace HordeServer.Collections.Impl
 				Updates.Add(UpdateBuilder.Set(x => x.RequestConform, RequestConform.Value));
 				Updates.Add(UpdateBuilder.Unset(x => x.ConformAttemptCount));
 			}
+			if (RequestFullConform != null)
+			{
+				Updates.Add(UpdateBuilder.Set(x => x.RequestFullConform, RequestFullConform.Value));
+				Updates.Add(UpdateBuilder.Unset(x => x.ConformAttemptCount));
+			}
 			if (RequestRestart != null)
 			{
 				if (RequestRestart.Value)
@@ -420,7 +430,11 @@ namespace HordeServer.Collections.Impl
 			UpdateDefinition<AgentDocument> Update = Builders<AgentDocument>.Update.Set(x => x.Workspaces, Workspaces);
 			Update = Update.Set(x => x.LastConformTime, LastConformTime);
 			Update = Update.Unset(x => x.ConformAttemptCount);
-			Update = Update.Set(x => x.RequestConform, RequestConform);
+			if (!RequestConform)
+			{
+				Update = Update.Unset(x => x.RequestConform);
+				Update = Update.Unset(x => x.RequestFullConform);
+			}
 
 			// Update the agent
 			return await TryUpdateAsync(Agent, Update);
diff --git a/Engine/Source/Programs/Horde/Horde.Build/Agents/Impl/Tasks/ConformTaskSource.cs b/Engine/Source/Programs/Horde/Horde.Build/Agents/Impl/Tasks/ConformTaskSource.cs
--- a/Engine/Source/Programs/Horde/Horde.Build/Agents/Impl/Tasks/ConformTaskSource.cs
+++ b/Engine/Source/Programs/Horde/Horde.Build/Agents/Impl/Tasks/ConformTaskSource.cs
@@ -185,6 +185,7 @@ namespace HordeServer.Tasks.Impl
 
 					ILogFile Log = await LogService.CreateLogFileAsync(JobId.Empty, Agent.SessionId, LogType.Json);
 					Task.LogId = Log.Id.ToString();
+					Task.RemoveUntrackedFiles = Agent.RequestFullConform;
 
 					byte[] Payload = Any.Pack(Task).ToByteArray();
 
@@ -192,7 +193,7 @@ namespace HordeServer.Tasks.Impl
 				}
 			}
 
-			if (Agent.RequestConform)
+			if (Agent.RequestConform || Agent.RequestFullConform)
 			{
 				return AgentLease.Drain;
 			}
@@ -321,7 +322,7 @@ namespace HordeServer.Tasks.Impl
 		private async Task<bool> IsConformPendingAsync(IAgent Agent, DateTime UtcNow)
 		{
 			// If a conform was manually requested, allow it to run even if the agent is disabled
-			if (Agent.RequestConform)
+			if (Agent.RequestConform || Agent.RequestFullConform)
 			{
 				return !IsConformCoolDownPeriod(Agent, UtcNow);
 			}
diff --git a/Engine/Source/Programs/Horde/Horde.Build/Api/Agents.cs b/Engine/Source/Programs/Horde/Horde.Build/Api/Agents.cs
--- a/Engine/Source/Programs/Horde/Horde.Build/Api/Agents.cs
+++ b/Engine/Source/Programs/Horde/Horde.Build/Api/Agents.cs
@@ -86,6 +86,11 @@ namespace HordeServer.Api
 		/// </summary>
 		public bool? RequestConform { get; set; }
 
+		/// <summary>
+		/// Request that a full conform be performed, removing all intermediate files
+		/// </summary>
+		public bool? RequestFullConform { get; set; }
+
 		/// <summary>
 		/// Request the machine be restarted
 		/// </summary>
diff --git a/Engine/Source/Programs/Horde/Horde.Build/Controllers/AgentsController.cs b/Engine/Source/Programs/Horde/Horde.Build/Controllers/AgentsController.cs
--- a/Engine/Source/Programs/Horde/Horde.Build/Controllers/AgentsController.cs
+++ b/Engine/Source/Programs/Horde/Horde.Build/Controllers/AgentsController.cs
@@ -159,7 +159,7 @@ namespace HordeServer.Controllers
 					return Forbid(AclAction.ChangePermissions, AgentId);
 				}
 
-				IAgent? NewAgent = await AgentService.Agents.TryUpdateSettingsAsync(Agent, Update.Enabled, Update.RequestConform, Update.RequestRestart, Update.RequestShutdown, Channel, Update.Pools?.ConvertAll(x => new PoolId(x)), Acl.Merge(Agent.Acl, Update.Acl), Update.Comment);
+				IAgent? NewAgent = await AgentService.Agents.TryUpdateSettingsAsync(Agent, Update.Enabled, Update.RequestConform, Update.RequestFullConform, Update.RequestRestart, Update.RequestShutdown, Channel, Update.Pools?.ConvertAll(x => new PoolId(x)), Acl.Merge(Agent.Acl, Update.Acl), Update.Comment);
 				if (NewAgent == null)
 				{
 					continue;
diff --git a/Engine/Source/Programs/Horde/Horde.Build/Services/RpcService.cs b/Engine/Source/Programs/Horde/Horde.Build/Services/RpcService.cs
--- a/Engine/Source/Programs/Horde/Horde.Build/Services/RpcService.cs
+++ b/Engine/Source/Programs/Horde/Horde.Build/Services/RpcService.cs
@@ -271,7 +271,7 @@ namespace HordeServer.Services
 
 				// Get the set of workspaces that are currently required
 				HashSet<AgentWorkspace> ConformWorkspaces = await PoolService.GetWorkspacesAsync(Agent, DateTime.UtcNow);
-				bool bPendingConform = !ConformWorkspaces.SetEquals(NewWorkspaces);
+				bool bPendingConform = !ConformWorkspaces.SetEquals(NewWorkspaces) || (Agent.RequestFullConform && !Request.RemoveUntrackedFiles);
 
 				// Update the workspaces
 				if (await AgentService.TryUpdateWorkspacesAsync(Agent, NewWorkspaces, bPendingConform))
@@ -280,6 +280,7 @@ namespace HordeServer.Services
 					if (bPendingConform)
 					{
 						Response.Retry = await ConformTaskSource.GetWorkspacesAsync(Agent, Response.PendingWorkspaces);
+						Response.RemoveUntrackedFiles = Request.RemoveUntrackedFiles || Agent.RequestFullConform;
 					}
 					return Response;
 				}
diff --git a/Engine/Source/Programs/Horde/Horde.Build/Utilities/MongoExtensions.cs b/Engine/Source/Programs/Horde/Horde.Build/Utilities/MongoExtensions.cs
--- a/Engine/Source/Programs/Horde/Horde.Build/Utilities/MongoExtensions.cs
+++ b/Engine/Source/Programs/Horde/Horde.Build/Utilities/MongoExtensions.cs
@@ -177,6 +177,26 @@ namespace HordeServer.Utilities
 			}
 		}
 
+		/// <summary>
+		/// Sets a field to a value, or unsets it if the value is null
+		/// </summary>
+		/// <typeparam name="TDocument">The document type</typeparam>
+		/// <param name="Update">Update builder</param>
+		/// <param name="Field">Expression for the field to set</param>
+		/// <param name="Value">New value to set</param>
+		/// <returns>Update defintiion</returns>
+		public static UpdateDefinition<TDocument> SetOrUnsetBool<TDocument>(this UpdateDefinition<TDocument> Update, Expression<Func<TDocument, bool>> Field, bool Value)
+		{
+			if (Value)
+			{
+				return Update.Set(Field, Value);
+			}
+			else
+			{
+				return Update.Unset(new ExpressionFieldDefinition<TDocument>(Field));
+			}
+		}
+
 		/// <summary>
 		/// Sets a field to a value, or unsets it if the value is null
 		/// </summary>
