Horde: Add "Full" conform option to agent view
#rnx

diff --git a/Engine/Source/Programs/Horde/HordeDashboard/src/backend/Api.ts b/Engine/Source/Programs/Horde/HordeDashboard/src/backend/Api.ts
--- a/Engine/Source/Programs/Horde/HordeDashboard/src/backend/Api.ts
+++ b/Engine/Source/Programs/Horde/HordeDashboard/src/backend/Api.ts
@@ -388,6 +388,9 @@ export type UpdateAgentRequest = {
 	/**Boolean to request a conform */
 	requestConform?: boolean;
 
+	/** Request that a full conform be performed, removing all intermediate files */
+	requestFullConform?: boolean;
+
 	requestRestart?: boolean;
 
 	/**Per-agent override for the desired client version */
