Horde: Skip building changes automatically if #skipci is in the changelist description.

diff --git a/Engine/Source/Programs/Horde/Horde.Build.Tests/SchedulerTests.cs b/Engine/Source/Programs/Horde/Horde.Build.Tests/SchedulerTests.cs
--- a/Engine/Source/Programs/Horde/Horde.Build.Tests/SchedulerTests.cs
+++ b/Engine/Source/Programs/Horde/Horde.Build.Tests/SchedulerTests.cs
@@ -279,6 +279,43 @@ namespace HordeServerTests
 			Assert.AreEqual(106, Jobs2[1].CodeChange);
 		}
 
+		[TestMethod]
+		public async Task SkipCiTestAsync()
+		{
+			DateTime StartTime = new DateTime(2021, 1, 1, 12, 0, 0, DateTimeKind.Local); // Friday Jan 1, 2021 
+			Clock.UtcNow = StartTime;
+
+			CreateScheduleRequest Schedule = new CreateScheduleRequest();
+			Schedule.Enabled = true;
+			Schedule.Patterns.Add(new CreateSchedulePatternRequest { MinTime = 13 * 60, MaxTime = 14 * 60, Interval = 15 });
+			Schedule.MaxChanges = 2;
+			Schedule.Filter = new List<ChangeContentFlags> { ChangeContentFlags.ContainsCode };
+			await SetScheduleAsync(Schedule);
+
+			// Initial tick
+			await ScheduleService.TickForTestingAsync();
+
+			List<IJob> Jobs1 = await GetNewJobs();
+			Assert.AreEqual(0, Jobs1.Count);
+
+			// Trigger some jobs
+			IUser Bob = UserCollection.FindOrAddUserByLoginAsync("Bob").Result;
+			PerforceService.AddChange("//UE5/Main", 103, Bob, "", new string[] { "foo.cpp" });
+			PerforceService.AddChange("//UE5/Main", 104, Bob, "Don't build this change!\n#skipci", new string[] { "foo.cpp" });
+			PerforceService.AddChange("//UE5/Main", 105, Bob, "", new string[] { "foo.uasset" });
+			PerforceService.AddChange("//UE5/Main", 106, Bob, "", new string[] { "foo.cpp" });
+
+			Clock.Advance(TimeSpan.FromHours(1.25));
+			await ScheduleService.TickForTestingAsync();
+
+			List<IJob> Jobs2 = await GetNewJobs();
+			Assert.AreEqual(2, Jobs2.Count);
+			Assert.AreEqual(103, Jobs2[0].Change);
+			Assert.AreEqual(103, Jobs2[0].CodeChange);
+			Assert.AreEqual(106, Jobs2[1].Change);
+			Assert.AreEqual(106, Jobs2[1].CodeChange);
+		}
+
 		[TestMethod]
 		public async Task MaxActiveTestAsync()
 		{
diff --git a/Engine/Source/Programs/Horde/Horde.Build/Projects/ScheduleService.cs b/Engine/Source/Programs/Horde/Horde.Build/Projects/ScheduleService.cs
--- a/Engine/Source/Programs/Horde/Horde.Build/Projects/ScheduleService.cs
+++ b/Engine/Source/Programs/Horde/Horde.Build/Projects/ScheduleService.cs
@@ -28,6 +28,7 @@ using TimeZoneConverter;
 using StackExchange.Redis;
 using EpicGames.Redis;
 using EpicGames.Serialization;
+using System.Text.RegularExpressions;
 
 namespace HordeServer.Services
 {
@@ -510,9 +511,13 @@ namespace HordeServer.Services
 		/// <returns></returns>
 		private async Task<bool> ShouldBuildChangeAsync(string ClusterName, string StreamName, int Change, ChangeContentFlags? FilterFlags, FileFilter? FileFilter)
 		{
+			ChangeDetails Details = await Perforce.GetChangeDetailsAsync(ClusterName, StreamName, Change, null);
+			if (Regex.IsMatch(Details.Description, @"^\s*#\s*skipci", RegexOptions.Multiline))
+			{
+				return false;
+			}
 			if (FilterFlags != null && FilterFlags.Value != 0)
 			{
-				ChangeDetails Details = await Perforce.GetChangeDetailsAsync(ClusterName, StreamName, Change, null);
 				if ((Details.GetContentFlags() & FilterFlags.Value) == 0)
 				{
 					Logger.LogDebug("Not building change {Change} ({ChangeFlags}) due to filter flags ({FilterFlags})", Change, Details.GetContentFlags().ToString(), FilterFlags.Value.ToString());
@@ -521,7 +526,6 @@ namespace HordeServer.Services
 			}
 			if (FileFilter != null)
 			{
-				ChangeDetails Details = await Perforce.GetChangeDetailsAsync(ClusterName, StreamName, Change, null);
 				if (!Details.Files.Any(x => FileFilter.Matches(x.Path)))
 				{
 					Logger.LogDebug("Not building change {Change} due to file filter", Change);
