Horde: Possible fix for tests.

diff --git a/Engine/Source/Programs/Horde/Horde.Build.Tests/JobServiceTests.cs b/Engine/Source/Programs/Horde/Horde.Build.Tests/JobServiceTests.cs
--- a/Engine/Source/Programs/Horde/Horde.Build.Tests/JobServiceTests.cs
+++ b/Engine/Source/Programs/Horde/Horde.Build.Tests/JobServiceTests.cs
@@ -169,17 +169,15 @@ namespace HordeServerTests
 			IJob Job = await JobService.CreateJobAsync(null, Stream!, new TemplateRefId("temp"), Template.Id, Graph, "Hello", 1234, 1233, 999, null, null, null, null, null, null, true, true, null, null, new List<string> { "-Target=Pak" });
 
 			Job = Deref(await JobService.UpdateBatchAsync(Job, Job.Batches[0].Id, LogId.GenerateNewId(), JobStepBatchState.Running));
-			Assert.IsNotNull(await JobService.UpdateStepAsync(Job, Job.Batches[0].Id, Job.Batches[0].Steps[0].Id, JobStepState.Running));
-			Assert.IsNotNull(await JobService.UpdateStepAsync(Job, Job.Batches[0].Id, Job.Batches[0].Steps[0].Id, JobStepState.Completed, JobStepOutcome.Success));
+			Job = Deref(await JobService.UpdateStepAsync(Job, Job.Batches[0].Id, Job.Batches[0].Steps[0].Id, JobStepState.Running));
+			Job = Deref(await JobService.UpdateStepAsync(Job, Job.Batches[0].Id, Job.Batches[0].Steps[0].Id, JobStepState.Completed, JobStepOutcome.Success));
 
 			Job = Deref(await JobService.UpdateBatchAsync(Job, Job.Batches[1].Id, LogId.GenerateNewId(), JobStepBatchState.Running));
-			Assert.IsNotNull(await JobService.UpdateStepAsync(Job, Job.Batches[1].Id, Job.Batches[1].Steps[0].Id, JobStepState.Running));
+			Job = Deref(await JobService.UpdateStepAsync(Job, Job.Batches[1].Id, Job.Batches[1].Steps[0].Id, JobStepState.Running));
 
 			Job = Deref(await JobService.UpdateBatchAsync(Job, Job.Batches[2].Id, LogId.GenerateNewId(), JobStepBatchState.Running));
-			Assert.IsNotNull(await JobService.UpdateStepAsync(Job, Job.Batches[2].Id, Job.Batches[2].Steps[0].Id, JobStepState.Running));
-			Assert.IsNotNull(await JobService.UpdateStepAsync(Job, Job.Batches[2].Id, Job.Batches[2].Steps[0].Id, JobStepState.Completed, JobStepOutcome.Success));
-
-			Job = (await JobService.GetJobAsync(Job.Id))!;
+			Job = Deref(await JobService.UpdateStepAsync(Job, Job.Batches[2].Id, Job.Batches[2].Steps[0].Id, JobStepState.Running));
+			Job = Deref(await JobService.UpdateStepAsync(Job, Job.Batches[2].Id, Job.Batches[2].Steps[0].Id, JobStepState.Completed, JobStepOutcome.Success));
 
 			Assert.AreEqual(JobStepState.Waiting, Job.Batches[2].Steps[1].State);
 		}
