Horde: Fix missing log output for job leases finishing.

diff --git a/Engine/Source/Programs/Horde/Horde.Build/Jobs/Impl/JobTaskSource.cs b/Engine/Source/Programs/Horde/Horde.Build/Jobs/Impl/JobTaskSource.cs
--- a/Engine/Source/Programs/Horde/Horde.Build/Jobs/Impl/JobTaskSource.cs
+++ b/Engine/Source/Programs/Horde/Horde.Build/Jobs/Impl/JobTaskSource.cs
@@ -995,6 +995,8 @@ namespace HordeServer.Tasks.Impl
 		/// <inheritdoc/>
 		public override async Task OnLeaseFinishedAsync(IAgent Agent, LeaseId LeaseId, ExecuteJobTask Task, LeaseOutcome Outcome, ReadOnlyMemory<byte> Output, ILogger Logger)
 		{
+			await base.OnLeaseFinishedAsync(Agent, LeaseId, Task, Outcome, Output, Logger);
+
 			if (Outcome != LeaseOutcome.Success)
 			{
 				AgentId AgentId = Agent.Id;
