Horde: Fix missing summary when lease terminated
#rnx

diff --git a/Engine/Source/Programs/Horde/HordeDashboard/src/backend/JobDetails.ts b/Engine/Source/Programs/Horde/HordeDashboard/src/backend/JobDetails.ts
--- a/Engine/Source/Programs/Horde/HordeDashboard/src/backend/JobDetails.ts
+++ b/Engine/Source/Programs/Horde/HordeDashboard/src/backend/JobDetails.ts
@@ -1135,16 +1135,19 @@ export const getStepSummaryMarkdown = (jobDetails: JobDetails, stepId: string):
         if (step.abortByUser) {
             aborted = "This step was aborted";
             aborted += ` by ${step.abortByUser}.`;
-        }
-        else if (jobDetails.jobdata?.abortedByUser) {
+        } else if (jobDetails.jobdata?.abortedByUser) {
             aborted = "The job was aborted";
             aborted += ` by ${jobDetails.jobdata?.abortedByUser}.`;
+        } else {
+            aborted = "The step was aborted";
         }
         text.push(aborted);
     } else if (step.state === JobStepState.Ready || step.state === JobStepState.Waiting) {
 
         text.push(batchText() ?? `The step is pending in ${step.state} state`);
-    } else if (batch?.agentId) {
+    }
+    
+    if (batch?.agentId) {
 
         if (step.startTime) {
             const str = getNiceTime(step.startTime);
