Horde Storage

Fixed up docker composes for local testing which had stopped working with the rename to Horde.Storage

Changed the expected name of the main configuration object from Horde.Storage to Horde_Storage, this resolves issues attempting to set this objects values from environment variables.

diff --git a/Engine/Source/Programs/Horde/HordeStorage/Horde.Storage/Horde.Storage.Startup.cs b/Engine/Source/Programs/Horde/HordeStorage/Horde.Storage/Horde.Storage.Startup.cs
--- a/Engine/Source/Programs/Horde/HordeStorage/Horde.Storage/Horde.Storage.Startup.cs
+++ b/Engine/Source/Programs/Horde/HordeStorage/Horde.Storage/Horde.Storage.Startup.cs
@@ -154,7 +154,7 @@ namespace Horde.Storage
 
             services.AddHttpClient();
 
-            services.AddOptions<HordeStorageSettings>().Bind(Configuration.GetSection("Horde.Storage")).ValidateDataAnnotations();
+            services.AddOptions<HordeStorageSettings>().Bind(Configuration.GetSection("Horde_Storage")).ValidateDataAnnotations();
             services.AddOptions<MongoSettings>().Bind(Configuration.GetSection("Mongo")).ValidateDataAnnotations();
             services.AddOptions<CosmosSettings>().Bind(Configuration.GetSection("Cosmos")).ValidateDataAnnotations();
             services.AddOptions<DynamoDbSettings>().Bind(Configuration.GetSection("DynamoDb")).ValidateDataAnnotations();
diff --git a/Engine/Source/Programs/Horde/HordeStorage/Horde.Storage/Tests/Functional/Blob/BlobIndexTests.cs b/Engine/Source/Programs/Horde/HordeStorage/Horde.Storage/Tests/Functional/Blob/BlobIndexTests.cs
--- a/Engine/Source/Programs/Horde/HordeStorage/Horde.Storage/Tests/Functional/Blob/BlobIndexTests.cs
+++ b/Engine/Source/Programs/Horde/HordeStorage/Horde.Storage/Tests/Functional/Blob/BlobIndexTests.cs
@@ -156,7 +156,7 @@ namespace Horde.Storage.FunctionalTests.Storage
     {
         protected override IEnumerable<KeyValuePair<string, string>> GetSettings()
         {
-            return new[] { new KeyValuePair<string, string>("Horde.Storage:BlobIndexImplementation", HordeStorageSettings.BlobIndexImplementations.Memory.ToString()) };
+            return new[] { new KeyValuePair<string, string>("Horde_Storage:BlobIndexImplementation", HordeStorageSettings.BlobIndexImplementations.Memory.ToString()) };
         }
 
         protected override Task Seed(IServiceProvider serverServices)
@@ -175,7 +175,7 @@ namespace Horde.Storage.FunctionalTests.Storage
     {
         protected override IEnumerable<KeyValuePair<string, string>> GetSettings()
         {
-            return new[] { new KeyValuePair<string, string>("Horde.Storage:BlobIndexImplementation", HordeStorageSettings.BlobIndexImplementations.Scylla.ToString()) };
+            return new[] { new KeyValuePair<string, string>("Horde_Storage:BlobIndexImplementation", HordeStorageSettings.BlobIndexImplementations.Scylla.ToString()) };
         }
 
         protected override Task Seed(IServiceProvider serverServices)
diff --git a/Engine/Source/Programs/Horde/HordeStorage/Horde.Storage/Tests/Functional/ContentId/CompressedBlobTests.cs b/Engine/Source/Programs/Horde/HordeStorage/Horde.Storage/Tests/Functional/ContentId/CompressedBlobTests.cs
--- a/Engine/Source/Programs/Horde/HordeStorage/Horde.Storage/Tests/Functional/ContentId/CompressedBlobTests.cs
+++ b/Engine/Source/Programs/Horde/HordeStorage/Horde.Storage/Tests/Functional/ContentId/CompressedBlobTests.cs
@@ -30,7 +30,7 @@ namespace Horde.Storage.FunctionalTests.CompressedBlobs
             // Use the S3 storage backend so we can handle large blobs
             return new[]
             {
-                new KeyValuePair<string, string>("Horde.Storage:StorageImplementations:0", HordeStorageSettings.StorageBackendImplementations.S3.ToString()),
+                new KeyValuePair<string, string>("Horde_Storage:StorageImplementations:0", HordeStorageSettings.StorageBackendImplementations.S3.ToString()),
                 new KeyValuePair<string, string>("S3:BucketName", $"tests-{TestNamespace}")
             };
         }
diff --git a/Engine/Source/Programs/Horde/HordeStorage/Horde.Storage/Tests/Functional/GC/GCBlobTests.cs b/Engine/Source/Programs/Horde/HordeStorage/Horde.Storage/Tests/Functional/GC/GCBlobTests.cs
--- a/Engine/Source/Programs/Horde/HordeStorage/Horde.Storage/Tests/Functional/GC/GCBlobTests.cs
+++ b/Engine/Source/Programs/Horde/HordeStorage/Horde.Storage/Tests/Functional/GC/GCBlobTests.cs
@@ -49,7 +49,7 @@ namespace Horde.Storage.FunctionalTests.GC
                 // we are not reading the base appSettings here as we want exact control over what runs in the tests
                 .AddJsonFile("appsettings.Testing.json", false)
                 .AddEnvironmentVariables()
-                .AddInMemoryCollection(new List<KeyValuePair<string, string>>() { new KeyValuePair<string, string>("Horde.Storage:StorageImplementations:0", "MemoryBlobStore")})
+                .AddInMemoryCollection(new List<KeyValuePair<string, string>>() { new KeyValuePair<string, string>("Horde_Storage:StorageImplementations:0", "MemoryBlobStore")})
                 .Build();
 
             Logger logger = new LoggerConfiguration()
@@ -233,10 +233,10 @@ namespace Horde.Storage.FunctionalTests.GC
                 .AddEnvironmentVariables()
                 .AddInMemoryCollection(new List<KeyValuePair<string, string>>()
                 {
-                    new KeyValuePair<string, string>("Horde.Storage:StorageImplementations:0", "MemoryBlobStore"),
+                    new KeyValuePair<string, string>("Horde_Storage:StorageImplementations:0", "MemoryBlobStore"),
                     new KeyValuePair<string, string>("GC:CleanNamespaces:0", TestNamespace.ToString()),
                     new KeyValuePair<string, string>("GC:CleanOldBlobs", true.ToString()),
-                    new KeyValuePair<string, string>("Horde.Storage:BlobIndexImplementation", GetImplementation()),
+                    new KeyValuePair<string, string>("Horde_Storage:BlobIndexImplementation", GetImplementation()),
                     
                 })
                 .Build();
diff --git a/Engine/Source/Programs/Horde/HordeStorage/Horde.Storage/Tests/Functional/Objects/RefReplicatorTests.cs b/Engine/Source/Programs/Horde/HordeStorage/Horde.Storage/Tests/Functional/Objects/RefReplicatorTests.cs
--- a/Engine/Source/Programs/Horde/HordeStorage/Horde.Storage/Tests/Functional/Objects/RefReplicatorTests.cs
+++ b/Engine/Source/Programs/Horde/HordeStorage/Horde.Storage/Tests/Functional/Objects/RefReplicatorTests.cs
@@ -76,8 +76,8 @@ namespace Horde.Storage.FunctionalTests.Replication
         {
             return new[]
             {
-                new KeyValuePair<string, string>("Horde.Storage:ReferencesDbImplementation", HordeStorageSettings.ReferencesDbImplementations.Scylla.ToString()),
-                new KeyValuePair<string, string>("Horde.Storage:ReplicationLogWriterImplementation", HordeStorageSettings.ReplicationLogWriterImplementations.Scylla.ToString()),
+                new KeyValuePair<string, string>("Horde_Storage:ReferencesDbImplementation", HordeStorageSettings.ReferencesDbImplementations.Scylla.ToString()),
+                new KeyValuePair<string, string>("Horde_Storage:ReplicationLogWriterImplementation", HordeStorageSettings.ReplicationLogWriterImplementations.Scylla.ToString()),
             };
         }
 
diff --git a/Engine/Source/Programs/Horde/HordeStorage/Horde.Storage/Tests/Functional/Objects/ReferencesTests.cs b/Engine/Source/Programs/Horde/HordeStorage/Horde.Storage/Tests/Functional/Objects/ReferencesTests.cs
--- a/Engine/Source/Programs/Horde/HordeStorage/Horde.Storage/Tests/Functional/Objects/ReferencesTests.cs
+++ b/Engine/Source/Programs/Horde/HordeStorage/Horde.Storage/Tests/Functional/Objects/ReferencesTests.cs
@@ -37,9 +37,9 @@ namespace Horde.Storage.FunctionalTests.References
         {
             return new[]
             {
-                new KeyValuePair<string, string>("Horde.Storage:ReferencesDbImplementation", HordeStorageSettings.ReferencesDbImplementations.Scylla.ToString()),
-                new KeyValuePair<string, string>("Horde.Storage:ContentIdStoreImplementation", HordeStorageSettings.ContentIdStoreImplementations.Scylla.ToString()),
-                new KeyValuePair<string, string>("Horde.Storage:ReplicationLogWriterImplementation", HordeStorageSettings.ReplicationLogWriterImplementations.Scylla.ToString()),
+                new KeyValuePair<string, string>("Horde_Storage:ReferencesDbImplementation", HordeStorageSettings.ReferencesDbImplementations.Scylla.ToString()),
+                new KeyValuePair<string, string>("Horde_Storage:ContentIdStoreImplementation", HordeStorageSettings.ContentIdStoreImplementations.Scylla.ToString()),
+                new KeyValuePair<string, string>("Horde_Storage:ReplicationLogWriterImplementation", HordeStorageSettings.ReplicationLogWriterImplementations.Scylla.ToString()),
             };
         }
 
@@ -90,9 +90,9 @@ namespace Horde.Storage.FunctionalTests.References
         {
             return new[]
             {
-                new KeyValuePair<string, string>("Horde.Storage:ReferencesDbImplementation", HordeStorageSettings.ReferencesDbImplementations.Memory.ToString()),
-                new KeyValuePair<string, string>("Horde.Storage:ContentIdStoreImplementation", HordeStorageSettings.ContentIdStoreImplementations.Memory.ToString()),
-                new KeyValuePair<string, string>("Horde.Storage:ReplicationLogWriterImplementation", HordeStorageSettings.ReplicationLogWriterImplementations.Memory.ToString()),
+                new KeyValuePair<string, string>("Horde_Storage:ReferencesDbImplementation", HordeStorageSettings.ReferencesDbImplementations.Memory.ToString()),
+                new KeyValuePair<string, string>("Horde_Storage:ContentIdStoreImplementation", HordeStorageSettings.ContentIdStoreImplementations.Memory.ToString()),
+                new KeyValuePair<string, string>("Horde_Storage:ReplicationLogWriterImplementation", HordeStorageSettings.ReplicationLogWriterImplementations.Memory.ToString()),
             };
         }
 
diff --git a/Engine/Source/Programs/Horde/HordeStorage/Horde.Storage/Tests/Functional/Objects/ReplicationTests.cs b/Engine/Source/Programs/Horde/HordeStorage/Horde.Storage/Tests/Functional/Objects/ReplicationTests.cs
--- a/Engine/Source/Programs/Horde/HordeStorage/Horde.Storage/Tests/Functional/Objects/ReplicationTests.cs
+++ b/Engine/Source/Programs/Horde/HordeStorage/Horde.Storage/Tests/Functional/Objects/ReplicationTests.cs
@@ -36,8 +36,8 @@ namespace Horde.Storage.FunctionalTests.References
         {
             return new[]
             {
-                new KeyValuePair<string, string>("Horde.Storage:ReferencesDbImplementation", HordeStorageSettings.ReferencesDbImplementations.Scylla.ToString()),
-                new KeyValuePair<string, string>("Horde.Storage:ReplicationLogWriterImplementation", HordeStorageSettings.ReplicationLogWriterImplementations.Scylla.ToString()),
+                new KeyValuePair<string, string>("Horde_Storage:ReferencesDbImplementation", HordeStorageSettings.ReferencesDbImplementations.Scylla.ToString()),
+                new KeyValuePair<string, string>("Horde_Storage:ReplicationLogWriterImplementation", HordeStorageSettings.ReplicationLogWriterImplementations.Scylla.ToString()),
             };
         }
 
@@ -79,8 +79,8 @@ namespace Horde.Storage.FunctionalTests.References
         {
             return new[]
             {
-                new KeyValuePair<string, string>("Horde.Storage:ReferencesDbImplementation", HordeStorageSettings.ReferencesDbImplementations.Memory.ToString()),
-                new KeyValuePair<string, string>("Horde.Storage:ReplicationLogWriterImplementation", HordeStorageSettings.ReplicationLogWriterImplementations.Memory.ToString()),
+                new KeyValuePair<string, string>("Horde_Storage:ReferencesDbImplementation", HordeStorageSettings.ReferencesDbImplementations.Memory.ToString()),
+                new KeyValuePair<string, string>("Horde_Storage:ReplicationLogWriterImplementation", HordeStorageSettings.ReplicationLogWriterImplementations.Memory.ToString()),
             };
         }
 
diff --git a/Engine/Source/Programs/Horde/HordeStorage/Horde.Storage/Tests/Functional/RefTests.cs b/Engine/Source/Programs/Horde/HordeStorage/Horde.Storage/Tests/Functional/RefTests.cs
--- a/Engine/Source/Programs/Horde/HordeStorage/Horde.Storage/Tests/Functional/RefTests.cs
+++ b/Engine/Source/Programs/Horde/HordeStorage/Horde.Storage/Tests/Functional/RefTests.cs
@@ -84,7 +84,7 @@ namespace Horde.Storage.FunctionalTests.Ref
         }
         protected override IEnumerable<KeyValuePair<string, string>> GetSettings()
         {
-            return new[] {new KeyValuePair<string, string>("Horde.Storage:RefDbImplementation", HordeStorageSettings.RefDbImplementations.Mongo.ToString())};
+            return new[] {new KeyValuePair<string, string>("Horde_Storage:RefDbImplementation", HordeStorageSettings.RefDbImplementations.Mongo.ToString())};
         }
     }
 
@@ -97,7 +97,7 @@ namespace Horde.Storage.FunctionalTests.Ref
 
         protected override IEnumerable<KeyValuePair<string, string>> GetSettings()
         {
-            return new[] {new KeyValuePair<string, string>("Horde.Storage:RefDbImplementation", HordeStorageSettings.RefDbImplementations.DynamoDb.ToString())};
+            return new[] {new KeyValuePair<string, string>("Horde_Storage:RefDbImplementation", HordeStorageSettings.RefDbImplementations.DynamoDb.ToString())};
         }
 
         protected override async Task SeedDb(IServiceProvider provider)
diff --git a/Engine/Source/Programs/Horde/HordeStorage/Horde.Storage/Tests/Functional/StorageTests.cs b/Engine/Source/Programs/Horde/HordeStorage/Horde.Storage/Tests/Functional/StorageTests.cs
--- a/Engine/Source/Programs/Horde/HordeStorage/Horde.Storage/Tests/Functional/StorageTests.cs
+++ b/Engine/Source/Programs/Horde/HordeStorage/Horde.Storage/Tests/Functional/StorageTests.cs
@@ -43,7 +43,7 @@ namespace Horde.Storage.FunctionalTests.Storage
         {
             return new[]
             {
-                new KeyValuePair<string, string>("Horde.Storage:StorageImplementations:0", HordeStorageSettings.StorageBackendImplementations.S3.ToString()),
+                new KeyValuePair<string, string>("Horde_Storage:StorageImplementations:0", HordeStorageSettings.StorageBackendImplementations.S3.ToString()),
                 new KeyValuePair<string, string>("S3:BucketName", $"tests-{TestNamespaceName}")
             };
         }
@@ -85,7 +85,7 @@ namespace Horde.Storage.FunctionalTests.Storage
 
         protected override IEnumerable<KeyValuePair<string, string>> GetSettings()
         {
-            return new[] {new KeyValuePair<string, string>("Horde.Storage:StorageImplementations:0", HordeStorageSettings.StorageBackendImplementations.Azure.ToString())};
+            return new[] {new KeyValuePair<string, string>("Horde_Storage:StorageImplementations:0", HordeStorageSettings.StorageBackendImplementations.Azure.ToString())};
         }
 
         protected override async Task Seed(IServiceProvider provider)
@@ -146,7 +146,7 @@ namespace Horde.Storage.FunctionalTests.Storage
         protected override IEnumerable<KeyValuePair<string, string>> GetSettings()
         {
             return new[] { 
-                new KeyValuePair<string, string>("Horde.Storage:StorageImplementations:0", HordeStorageSettings.StorageBackendImplementations.FileSystem.ToString()),
+                new KeyValuePair<string, string>("Horde_Storage:StorageImplementations:0", HordeStorageSettings.StorageBackendImplementations.FileSystem.ToString()),
                 new KeyValuePair<string, string>("Filesystem:RootDir", _localTestDir)
             };
         }
diff --git a/Engine/Source/Programs/Horde/HordeStorage/Horde.Storage/Tests/UnitTests/HordeStorageStartupTest.cs b/Engine/Source/Programs/Horde/HordeStorage/Horde.Storage/Tests/UnitTests/HordeStorageStartupTest.cs
--- a/Engine/Source/Programs/Horde/HordeStorage/Horde.Storage/Tests/UnitTests/HordeStorageStartupTest.cs
+++ b/Engine/Source/Programs/Horde/HordeStorage/Horde.Storage/Tests/UnitTests/HordeStorageStartupTest.cs
@@ -31,18 +31,18 @@ namespace Horde.Storage.UnitTests
             Assert.IsTrue(blobService.BlobStore.Single() is MemoryCacheBlobStore);
             
             // A single blob store configuration should yield the store itself without a hierarchical wrapper store
-            blobService = GetBlobServiceForConfig(new Dictionary<string, string> {{"Horde.Storage:StorageImplementations:0", "FileSystem"}});
+            blobService = GetBlobServiceForConfig(new Dictionary<string, string> {{"Horde_Storage:StorageImplementations:0", "FileSystem"}});
             Assert.IsTrue(blobService.BlobStore.Single() is FileSystemStore);
             
             // Should not be case-sensitive
-            blobService = GetBlobServiceForConfig(new Dictionary<string, string> {{"Horde.Storage:StorageImplementations:0", "FiLeSYsTEm"}});
+            blobService = GetBlobServiceForConfig(new Dictionary<string, string> {{"Horde_Storage:StorageImplementations:0", "FiLeSYsTEm"}});
             Assert.IsTrue(blobService.BlobStore.Single() is FileSystemStore);
             
             // Two or more blob stores returns a hierarchical store
             blobService = GetBlobServiceForConfig(new Dictionary<string, string>
             {
-                {"Horde.Storage:StorageImplementations:0", "FileSystem"},
-                {"Horde.Storage:StorageImplementations:1", "Memory"},
+                {"Horde_Storage:StorageImplementations:0", "FileSystem"},
+                {"Horde_Storage:StorageImplementations:1", "Memory"},
             });
 
 
diff --git a/Engine/Source/Programs/Horde/HordeStorage/Horde.Storage/appsettings.Development.json b/Engine/Source/Programs/Horde/HordeStorage/Horde.Storage/appsettings.Development.json
--- a/Engine/Source/Programs/Horde/HordeStorage/Horde.Storage/appsettings.Development.json
+++ b/Engine/Source/Programs/Horde/HordeStorage/Horde.Storage/appsettings.Development.json
@@ -17,7 +17,7 @@
     }
   },
 
-  "Horde.Storage": {
+  "Horde_Storage": {
     "refDBImplementation": "mongo",
 
     "TransactionLogWriterImplementation": "memory",
diff --git a/Engine/Source/Programs/Horde/HordeStorage/Horde.Storage/appsettings.Production.json b/Engine/Source/Programs/Horde/HordeStorage/Horde.Storage/appsettings.Production.json
--- a/Engine/Source/Programs/Horde/HordeStorage/Horde.Storage/appsettings.Production.json
+++ b/Engine/Source/Programs/Horde/HordeStorage/Horde.Storage/appsettings.Production.json
@@ -27,7 +27,7 @@
     "AutoscaleCapacity": true
   },
 
-  "Horde.Storage": {
+  "Horde_Storage": {
     "TransactionLogWriterImplementation": "callisto"
   },
 
diff --git a/Engine/Source/Programs/Horde/HordeStorage/Horde.Storage/appsettings.Testing.json b/Engine/Source/Programs/Horde/HordeStorage/Horde.Storage/appsettings.Testing.json
--- a/Engine/Source/Programs/Horde/HordeStorage/Horde.Storage/appsettings.Testing.json
+++ b/Engine/Source/Programs/Horde/HordeStorage/Horde.Storage/appsettings.Testing.json
@@ -32,7 +32,7 @@
     "UseGlobalTables": false
   },
 
-  "Horde.Storage": {
+  "Horde_Storage": {
     "RefDBImplementation": "mongo",
     "StorageImplementations": [ "memory" ],
     "TransactionLogWriterImplementation": "memory"
