Horde: Fix "drain" lease causing exception when another lease is already executing.

diff --git a/Engine/Source/Programs/Horde/Horde.Build/Agents/Impl/AgentService.cs b/Engine/Source/Programs/Horde/Horde.Build/Agents/Impl/AgentService.cs
--- a/Engine/Source/Programs/Horde/Horde.Build/Agents/Impl/AgentService.cs
+++ b/Engine/Source/Programs/Horde/Horde.Build/Agents/Impl/AgentService.cs
@@ -411,7 +411,7 @@ namespace HordeServer.Services
 								Result = TaskResult;
 								CancellationSource.Cancel();
 							}
-							else
+							else if (TaskResult.Value.Item2 != AgentLease.Drain)
 							{
 								(ITaskSource TaskSource, AgentLease TaskLease) = TaskResult.Value;
 								await TaskSource.CancelLeaseAsync(Agent, TaskLease.Id, Any.Parser.ParseFrom(TaskLease.Payload));
