Horde Storage
Fixed recurrsion on startup due to cyclic DI dependencies between the blob cleanup and the filesystem store

diff --git a/Engine/Source/Programs/Horde/HordeStorage/Horde.Storage/Implementation/Blob/FileSystemStore.cs b/Engine/Source/Programs/Horde/HordeStorage/Horde.Storage/Implementation/Blob/FileSystemStore.cs
--- a/Engine/Source/Programs/Horde/HordeStorage/Horde.Storage/Implementation/Blob/FileSystemStore.cs
+++ b/Engine/Source/Programs/Horde/HordeStorage/Horde.Storage/Implementation/Blob/FileSystemStore.cs
@@ -7,7 +7,6 @@ using System.Linq;
 using System.Threading;
 using System.Threading.Tasks;
 using Dasync.Collections;
-using Jupiter;
 using Jupiter.Implementation;
 using Microsoft.Extensions.Options;
 using Serilog;
@@ -21,10 +20,9 @@ namespace Horde.Storage.Implementation
 
         internal const int DefaultBufferSize = 4096;
 
-        public FileSystemStore(IOptionsMonitor<FilesystemSettings> settings, BlobCleanupService cleanupService)
+        public FileSystemStore(IOptionsMonitor<FilesystemSettings> settings)
         {
             _settings = settings;
-            cleanupService.RegisterCleanup(this);
         }
 
         public static FileInfo GetFilesystemPath(string rootDir, NamespaceId ns, BlobIdentifier blob)
diff --git a/Engine/Source/Programs/Horde/HordeStorage/Horde.Storage/Implementation/GC/BlobCleanupService.cs b/Engine/Source/Programs/Horde/HordeStorage/Horde.Storage/Implementation/GC/BlobCleanupService.cs
--- a/Engine/Source/Programs/Horde/HordeStorage/Horde.Storage/Implementation/GC/BlobCleanupService.cs
+++ b/Engine/Source/Programs/Horde/HordeStorage/Horde.Storage/Implementation/GC/BlobCleanupService.cs
@@ -47,7 +47,12 @@ namespace Horde.Storage.Implementation
                 OrphanBlobCleanupRefs orphanBlobCleanupRefs = provider.GetService<OrphanBlobCleanupRefs>()!;
                 RegisterCleanup(orphanBlobCleanupRefs);
             }
-            
+
+            FileSystemStore? fileSystemStore = provider.GetService<FileSystemStore>();
+            if (fileSystemStore != null)
+            {
+                RegisterCleanup(fileSystemStore);
+            }
         }
 
         public void RegisterCleanup(IBlobCleanup cleanup)
