Horde: Catch additional gRPC exception types, and annotate log messages with the agent id triggering the request.

diff --git a/Engine/Source/Programs/Horde/Horde.Build/Acls/AclService.cs b/Engine/Source/Programs/Horde/Horde.Build/Acls/AclService.cs
--- a/Engine/Source/Programs/Horde/Horde.Build/Acls/AclService.cs
+++ b/Engine/Source/Programs/Horde/Horde.Build/Acls/AclService.cs
@@ -67,7 +67,7 @@ namespace HordeServer.Services
 		/// <summary>
 		/// Role for all agents
 		/// </summary>
-		public static AclClaim AgentClaim { get; } = new AclClaim(HordeClaimTypes.Role, "agent");
+		public static AclClaim AgentRoleClaim { get; } = new AclClaim(HordeClaimTypes.Role, "agent");
 
 		/// <summary>
 		/// The default permissions
@@ -110,7 +110,7 @@ namespace HordeServer.Services
 			DefaultAcl = new Acl();
 			DefaultAcl.Entries.Add(new AclEntry(new AclClaim(ClaimTypes.Role, "internal:AgentRegistration"), new[] { AclAction.CreateAgent, AclAction.CreateSession }));
 			DefaultAcl.Entries.Add(new AclEntry(AgentRegistrationClaim, new[] { AclAction.CreateAgent, AclAction.CreateSession, AclAction.UpdateAgent, AclAction.DownloadSoftware }));
-			DefaultAcl.Entries.Add(new AclEntry(AgentClaim, new[] { AclAction.ViewProject, AclAction.ViewStream, AclAction.CreateEvent, AclAction.DownloadSoftware }));
+			DefaultAcl.Entries.Add(new AclEntry(AgentRoleClaim, new[] { AclAction.ViewProject, AclAction.ViewStream, AclAction.CreateEvent, AclAction.DownloadSoftware }));
 			DefaultAcl.Entries.Add(new AclEntry(DownloadSoftwareClaim, new[] { AclAction.DownloadSoftware }));
 			DefaultAcl.Entries.Add(new AclEntry(UploadSoftwareClaim, new[] { AclAction.UploadSoftware }));
 			DefaultAcl.Entries.Add(new AclEntry(ConfigureProjectsClaim, new[] { AclAction.CreateProject, AclAction.UpdateProject, AclAction.ViewProject, AclAction.CreateStream, AclAction.UpdateStream, AclAction.ViewStream, AclAction.ChangePermissions }));
@@ -213,7 +213,35 @@ namespace HordeServer.Services
 		}
 
 		/// <summary>
-		/// Gets the role for a specific user
+		/// Gets the agent id associated with a particular user
+		/// </summary>
+		/// <param name="User"></param>
+		/// <returns></returns>
+		public static AgentId? GetAgentId(ClaimsPrincipal User)
+		{
+			Claim? Claim = User.Claims.FirstOrDefault(x => x.Type == HordeClaimTypes.AgentId);
+			if (Claim == null)
+			{
+				return null;
+			}
+			else
+			{
+				return new AgentId(Claim.Value);
+			}
+		}
+
+		/// <summary>
+		/// Gets the role for a specific agent
+		/// </summary>
+		/// <param name="AgentId">The session id</param>
+		/// <returns>New claim instance</returns>
+		public static AclClaim GetAgentClaim(AgentId AgentId)
+		{
+			return new AclClaim(HordeClaimTypes.AgentId, AgentId.ToString());
+		}
+
+		/// <summary>
+		/// Gets the role for a specific agent session
 		/// </summary>
 		/// <param name="SessionId">The session id</param>
 		/// <returns>New claim instance</returns>
diff --git a/Engine/Source/Programs/Horde/Horde.Build/Agents/Impl/AgentService.cs b/Engine/Source/Programs/Horde/Horde.Build/Agents/Impl/AgentService.cs
--- a/Engine/Source/Programs/Horde/Horde.Build/Agents/Impl/AgentService.cs
+++ b/Engine/Source/Programs/Horde/Horde.Build/Agents/Impl/AgentService.cs
@@ -144,12 +144,14 @@ namespace HordeServer.Services
 		/// <summary>
 		/// Issues a bearer token for the given session id
 		/// </summary>
+		/// <param name="AgentId">The agent id</param>
 		/// <param name="SessionId">The session id</param>
 		/// <returns>Bearer token for the agent</returns>
-		public string IssueSessionToken(SessionId SessionId)
+		public string IssueSessionToken(AgentId AgentId, SessionId SessionId)
 		{
 			List<AclClaim> Claims = new List<AclClaim>();
-			Claims.Add(AclService.AgentClaim);
+			Claims.Add(AclService.AgentRoleClaim);
+			Claims.Add(AclService.GetAgentClaim(AgentId));
 			Claims.Add(AclService.GetSessionClaim(SessionId));
 			return AclService.IssueBearerToken(Claims, null);
 		}
diff --git a/Engine/Source/Programs/Horde/Horde.Build/Fleet/Impl/Managers/Fleet/AwsFleetManager.cs b/Engine/Source/Programs/Horde/Horde.Build/Fleet/Impl/Managers/Fleet/AwsFleetManager.cs
--- a/Engine/Source/Programs/Horde/Horde.Build/Fleet/Impl/Managers/Fleet/AwsFleetManager.cs
+++ b/Engine/Source/Programs/Horde/Horde.Build/Fleet/Impl/Managers/Fleet/AwsFleetManager.cs
@@ -92,7 +92,7 @@ namespace HordeServer.Services.Impl
 					{
 						if (await AgentCollection.TryUpdateSettingsAsync(NewAgent, bRequestShutdown: true) != null)
 						{
-							AgentCollection.GetLogger(Agent.Id).LogInformation("Shutting down due to autoscaler");
+							AgentCollection.GetLogger(Agent.Id).LogInformation("Marked for shutdown due to autoscalar (currently {NumLeases} leases outstanding)", NewAgent.Leases.Count);
 							Count--;
 							break;
 						}
diff --git a/Engine/Source/Programs/Horde/Horde.Build/Services/RpcService.cs b/Engine/Source/Programs/Horde/Horde.Build/Services/RpcService.cs
--- a/Engine/Source/Programs/Horde/Horde.Build/Services/RpcService.cs
+++ b/Engine/Source/Programs/Horde/Horde.Build/Services/RpcService.cs
@@ -369,7 +369,7 @@ namespace HordeServer.Services
 			Response.AgentId = Agent.Id.ToString();
 			Response.SessionId = Agent.SessionId.ToString();
 			Response.ExpiryTime = Timestamp.FromDateTime(Agent.SessionExpiresAt!.Value);
-			Response.Token = AgentService.IssueSessionToken(Agent.SessionId!.Value);
+			Response.Token = AgentService.IssueSessionToken(Agent.Id, Agent.SessionId!.Value);
 			return Response;
 		}
 
diff --git a/Engine/Source/Programs/Horde/Horde.Build/Startup.cs b/Engine/Source/Programs/Horde/Horde.Build/Startup.cs
--- a/Engine/Source/Programs/Horde/Horde.Build/Startup.cs
+++ b/Engine/Source/Programs/Horde/Horde.Build/Startup.cs
@@ -107,107 +107,64 @@ namespace HordeServer
 			}
 
 			[SuppressMessage("Design", "CA1031:Do not catch general exception types", Justification = "<Pending>")]
-			IDisposable? RequestContext<TRequest>(TRequest Request)
+			public override Task<TResponse> UnaryServerHandler<TRequest, TResponse>(TRequest Request, ServerCallContext Context, UnaryServerMethod<TRequest, TResponse> Continuation)
 			{
-				try
-				{
-					return Logger.BeginScope("Request: {Request}", JsonSerializer.Serialize<TRequest>(Request));
-				}
-				catch
-				{
-					return null;
-				}
+				return Guard(Context, Context.Method, () => base.UnaryServerHandler(Request, Context, Continuation));
 			}
 
-			[SuppressMessage("Design", "CA1031:Do not catch general exception types", Justification = "<Pending>")]
-			IDisposable? ConnectionContext(ConnectionInfo Connection)
+			public override Task<TResponse> ClientStreamingServerHandler<TRequest, TResponse>(IAsyncStreamReader<TRequest> RequestStream, ServerCallContext Context, ClientStreamingServerMethod<TRequest, TResponse> Continuation) where TRequest : class where TResponse : class
 			{
-				try
-				{
-					return Logger.BeginScope("Client: {ClientIp}", Connection.RemoteIpAddress);
-				}
-				catch
-				{
-					return null;
-				}
+				return Guard(Context, Context.Method, () => base.ClientStreamingServerHandler(RequestStream, Context, Continuation));
 			}
 
-			public override AsyncClientStreamingCall<TRequest, TResponse> AsyncClientStreamingCall<TRequest, TResponse>(ClientInterceptorContext<TRequest, TResponse> Context, AsyncClientStreamingCallContinuation<TRequest, TResponse> Continuation)
+			public override Task ServerStreamingServerHandler<TRequest, TResponse>(TRequest Request, IServerStreamWriter<TResponse> ResponseStream, ServerCallContext Context, ServerStreamingServerMethod<TRequest, TResponse> Continuation) where TRequest : class where TResponse : class
 			{
-				try
-				{
-					return base.AsyncClientStreamingCall(Context, Continuation);
-				}
-				catch (StructuredRpcException Ex)
-				{
-					Logger.LogError(Ex, Ex.Format, Ex.Args);
-					throw;
-				}
-				catch (Exception Ex)
-				{
-					Logger.LogError(Ex, "Exception in call to {Method}", Context.Method);
-					throw new RpcException(new Status(StatusCode.Internal, $"An exception was thrown on the server: {Ex}"));
-				}
+				return Guard(Context, Context.Method, () => base.ServerStreamingServerHandler(Request, ResponseStream, Context, Continuation));
 			}
 
-			public override AsyncDuplexStreamingCall<TRequest, TResponse> AsyncDuplexStreamingCall<TRequest, TResponse>(ClientInterceptorContext<TRequest, TResponse> Context, AsyncDuplexStreamingCallContinuation<TRequest, TResponse> Continuation)
+			public override Task DuplexStreamingServerHandler<TRequest, TResponse>(IAsyncStreamReader<TRequest> RequestStream, IServerStreamWriter<TResponse> ResponseStream, ServerCallContext Context, DuplexStreamingServerMethod<TRequest, TResponse> Continuation) where TRequest : class where TResponse : class
 			{
-				try
-				{
-					return base.AsyncDuplexStreamingCall(Context, Continuation);
-				}
-				catch (StructuredRpcException Ex)
-				{
-					Logger.LogError(Ex, Ex.Format, Ex.Args);
-					throw;
-				}
-				catch (Exception Ex)
-				{
-					Logger.LogError(Ex, "Exception in call to {Method}", Context.Method);
-					throw new RpcException(new Status(StatusCode.Internal, $"An exception was thrown on the server: {Ex}"));
-				}
+				return Guard(Context, Context.Method, () => base.DuplexStreamingServerHandler(RequestStream, ResponseStream, Context, Continuation));
 			}
 
-			[SuppressMessage("Design", "CA1031:Do not catch general exception types", Justification = "<Pending>")]
-			public override AsyncServerStreamingCall<TResponse> AsyncServerStreamingCall<TRequest, TResponse>(TRequest Request, ClientInterceptorContext<TRequest, TResponse> Context, AsyncServerStreamingCallContinuation<TRequest, TResponse> Continuation)
+			async Task<T> Guard<T>(ServerCallContext Context, string Method, Func<Task<T>> CallFunc) where T : class
 			{
-				try
-				{
-					return base.AsyncServerStreamingCall(Request, Context, Continuation);
-				}
-				catch (StructuredRpcException Ex)
+				T Result = null!;
+				await Guard(Context, Method, async () => { Result = await CallFunc(); });
+				return Result;
+			}
+
+			async Task Guard(ServerCallContext Context, string Method, Func<Task> CallFunc)
+			{
+				HttpContext HttpContext = Context.GetHttpContext();
+
+				AgentId? AgentId = AclService.GetAgentId(HttpContext.User);
+				if (AgentId != null)
 				{
-					using IDisposable? Scope = RequestContext(Request); 
-					Logger.LogError(Ex, Ex.Format, Ex.Args);
-					throw;
+					using IDisposable Scope = Logger.BeginScope("Agent: {AgentId}, RemoteIP: {RemoteIP}, Method: {Method}", AgentId.Value, HttpContext.Connection.RemoteIpAddress, Method);
+					await GuardInner(Method, CallFunc);
 				}
-				catch (Exception Ex)
+				else
 				{
-					using IDisposable? Scope = RequestContext(Request);
-					Logger.LogError(Ex, "Exception in call to {Method}", Context.Method);
-					throw new RpcException(new Status(StatusCode.Internal, $"An exception was thrown on the server: {Ex}"));
+					using IDisposable Scope = Logger.BeginScope("RemoteIP: {RemoteIP}, Method: {Method}", HttpContext.Connection.RemoteIpAddress, Method);
+					await GuardInner(Method, CallFunc);
 				}
 			}
 
-			[SuppressMessage("Design", "CA1031:Do not catch general exception types", Justification = "<Pending>")]
-			public override async Task<TResponse> UnaryServerHandler<TRequest, TResponse>(TRequest Request, ServerCallContext Context, UnaryServerMethod<TRequest, TResponse> Continuation)
+			async Task GuardInner(string Method, Func<Task> CallFunc)
 			{
 				try
 				{
-					return await base.UnaryServerHandler(Request, Context, Continuation);
+					await CallFunc();
 				}
 				catch (StructuredRpcException Ex)
 				{
-					using IDisposable? ConnectionScope = ConnectionContext(Context.GetHttpContext().Connection);
-					using IDisposable? Scope = RequestContext(Request);
 					Logger.LogError(Ex, Ex.Format, Ex.Args);
 					throw;
 				}
 				catch (Exception Ex)
 				{
-					using IDisposable? ConnectionScope = ConnectionContext(Context.GetHttpContext().Connection);
-					using IDisposable? Scope = RequestContext(Request);
-					Logger.LogError(Ex, "Exception in call to {Method}", Context.Method);
+					Logger.LogError(Ex, "Exception in call to {Method}", Method);
 					throw new RpcException(new Status(StatusCode.Internal, $"An exception was thrown on the server: {Ex}"));
 				}
 			}
diff --git a/Engine/Source/Programs/Horde/Horde.Build/Utilities/HordeClaimTypes.cs b/Engine/Source/Programs/Horde/Horde.Build/Utilities/HordeClaimTypes.cs
--- a/Engine/Source/Programs/Horde/Horde.Build/Utilities/HordeClaimTypes.cs
+++ b/Engine/Source/Programs/Horde/Horde.Build/Utilities/HordeClaimTypes.cs
@@ -33,6 +33,11 @@ namespace HordeServer.Utilities
 		/// </summary>
 		public const string InternalRole = Prefix + "internal-role";
 
+		/// <summary>
+		/// Claim for a particular session
+		/// </summary>
+		public const string AgentId = Prefix + "agent";
+
 		/// <summary>
 		/// Claim for a particular session
 		/// </summary>
