Horde: Add shutdown reason to agents
#rnx

diff --git a/Engine/Source/Programs/Horde/Horde.Build/Agents/IAgent.cs b/Engine/Source/Programs/Horde/Horde.Build/Agents/IAgent.cs
--- a/Engine/Source/Programs/Horde/Horde.Build/Agents/IAgent.cs
+++ b/Engine/Source/Programs/Horde/Horde.Build/Agents/IAgent.cs
@@ -491,6 +491,11 @@ namespace HordeServer.Models
 		/// </summary>
 		public bool RequestShutdown { get; }
 
+		/// <summary>
+		/// The reason for the last agent shutdown
+		/// </summary>
+		public string? LastShutdownReason { get; }
+
 		/// <summary>
 		/// List of workspaces currently synced to this machine
 		/// </summary>
diff --git a/Engine/Source/Programs/Horde/Horde.Build/Agents/IAgentCollection.cs b/Engine/Source/Programs/Horde/Horde.Build/Agents/IAgentCollection.cs
--- a/Engine/Source/Programs/Horde/Horde.Build/Agents/IAgentCollection.cs
+++ b/Engine/Source/Programs/Horde/Horde.Build/Agents/IAgentCollection.cs
@@ -86,12 +86,13 @@ namespace HordeServer.Collections
 		/// <param name="bRequestFullConform">Whether to request a full conform job be run</param>
 		/// <param name="bRequestRestart">Whether to request the machine be restarted</param>
 		/// <param name="bRequestShutdown">Whether to request the machine be shut down</param>
+		/// <param name="ShutdownReason">The reason for shutting down agent, ex. Autoscaler/Manual/Unexpected</param>
 		/// <param name="Channel">Override for the desired software channel</param>
 		/// <param name="Pools">List of pools for the agent</param>
 		/// <param name="Acl">New ACL for this agent</param>
 		/// <param name="Comment">New comment</param>
 		/// <returns>Version of the software that needs to be installed on the agent. Null if the agent is running the correct version.</returns>
-		Task<IAgent?> TryUpdateSettingsAsync(IAgent Agent, bool? bEnabled = null, bool? bRequestConform = null, bool? bRequestFullConform = null, bool? bRequestRestart = null, bool? bRequestShutdown = null, AgentSoftwareChannelName? Channel = null, List<PoolId>? Pools = null, Acl? Acl = null, string? Comment = null);
+		Task<IAgent?> TryUpdateSettingsAsync(IAgent Agent, bool? bEnabled = null, bool? bRequestConform = null, bool? bRequestFullConform = null, bool? bRequestRestart = null, bool? bRequestShutdown = null, string? ShutdownReason = null, AgentSoftwareChannelName? Channel = null, List<PoolId>? Pools = null, Acl? Acl = null, string? Comment = null);
 
 		/// <summary>
 		/// Update the current workspaces for an agent.
diff --git a/Engine/Source/Programs/Horde/Horde.Build/Agents/Impl/AgentCollection.cs b/Engine/Source/Programs/Horde/Horde.Build/Agents/Impl/AgentCollection.cs
--- a/Engine/Source/Programs/Horde/Horde.Build/Agents/Impl/AgentCollection.cs
+++ b/Engine/Source/Programs/Horde/Horde.Build/Agents/Impl/AgentCollection.cs
@@ -104,6 +104,9 @@ namespace HordeServer.Collections.Impl
 			[BsonIgnoreIfDefault, BsonDefaultValue(false)]
 			public bool RequestShutdown { get; set; }
 
+			[BsonIgnoreIfNull]
+			public string? LastShutdownReason { get; set; }
+
 			public List<AgentWorkspace> Workspaces { get; set; } = new List<AgentWorkspace>();
 			public DateTime LastConformTime { get; set; }
 
@@ -254,7 +257,7 @@ namespace HordeServer.Collections.Impl
 		}
 
 		/// <inheritdoc/>
-		public async Task<IAgent?> TryUpdateSettingsAsync(IAgent AgentInterface, bool? Enabled = null, bool? RequestConform = null, bool? RequestFullConform = null, bool? RequestRestart = null, bool? RequestShutdown = null, AgentSoftwareChannelName? Channel = null, List<PoolId>? Pools = null, Acl? Acl = null, string? Comment = null)
+		public async Task<IAgent?> TryUpdateSettingsAsync(IAgent AgentInterface, bool? Enabled = null, bool? RequestConform = null, bool? RequestFullConform = null, bool? RequestRestart = null, bool? RequestShutdown = null, string? ShutdownReason = null, AgentSoftwareChannelName? Channel = null, List<PoolId>? Pools = null, Acl? Acl = null, string? Comment = null)
 		{
 			AgentDocument Agent = (AgentDocument)AgentInterface;
 
@@ -302,6 +305,12 @@ namespace HordeServer.Collections.Impl
 					Updates.Add(UpdateBuilder.Unset(x => x.RequestShutdown));
 				}
 			}
+
+			if (ShutdownReason != null)
+			{
+				Updates.Add(UpdateBuilder.Set(x => x.LastShutdownReason, ShutdownReason));
+			}
+
 			if (Channel != null)
 			{
 				if (Channel.Value == AgentSoftwareService.DefaultChannelName)
@@ -463,6 +472,7 @@ namespace HordeServer.Collections.Impl
 			Updates.Add(UpdateBuilder.Set(x => x.Version, Version));
 			Updates.Add(UpdateBuilder.Unset(x => x.RequestRestart));
 			Updates.Add(UpdateBuilder.Unset(x => x.RequestShutdown));
+			Updates.Add(UpdateBuilder.Set(x => x.LastShutdownReason, "Unexpected"));
 
 			// Apply the update
 			return await TryUpdateAsync(Agent, UpdateBuilder.Combine(Updates));
diff --git a/Engine/Source/Programs/Horde/Horde.Build/Api/Agents.cs b/Engine/Source/Programs/Horde/Horde.Build/Api/Agents.cs
--- a/Engine/Source/Programs/Horde/Horde.Build/Api/Agents.cs
+++ b/Engine/Source/Programs/Horde/Horde.Build/Api/Agents.cs
@@ -362,6 +362,11 @@ namespace HordeServer.Api
 		/// </summary>
 		public bool PendingShutdown { get; set; }
 
+		/// <summary>
+		/// The reason for the last shutdown
+		/// </summary>
+		public string LastShutdownReason { get; set; }
+
 		/// <summary>
 		/// Last time a conform was attempted
 		/// </summary>
@@ -450,6 +455,7 @@ namespace HordeServer.Api
 			this.PendingConform = Agent.RequestConform;
 			this.PendingRestart = Agent.RequestRestart;
 			this.PendingShutdown = Agent.RequestShutdown;
+			this.LastShutdownReason = Agent.LastShutdownReason ?? "Unknown";
 			this.LastConformTime = Agent.LastConformTime;
 			this.NextConformTime =   Agent.LastConformTime;
 			this.ConformAttemptCount = Agent.ConformAttemptCount;
diff --git a/Engine/Source/Programs/Horde/Horde.Build/Controllers/AgentsController.cs b/Engine/Source/Programs/Horde/Horde.Build/Controllers/AgentsController.cs
--- a/Engine/Source/Programs/Horde/Horde.Build/Controllers/AgentsController.cs
+++ b/Engine/Source/Programs/Horde/Horde.Build/Controllers/AgentsController.cs
@@ -159,7 +159,7 @@ namespace HordeServer.Controllers
 					return Forbid(AclAction.ChangePermissions, AgentId);
 				}
 
-				IAgent? NewAgent = await AgentService.Agents.TryUpdateSettingsAsync(Agent, Update.Enabled, Update.RequestConform, Update.RequestFullConform, Update.RequestRestart, Update.RequestShutdown, Channel, Update.Pools?.ConvertAll(x => new PoolId(x)), Acl.Merge(Agent.Acl, Update.Acl), Update.Comment);
+				IAgent? NewAgent = await AgentService.Agents.TryUpdateSettingsAsync(Agent, Update.Enabled, Update.RequestConform, Update.RequestFullConform, Update.RequestRestart, Update.RequestShutdown, $"Manual ({UserName})", Channel, Update.Pools?.ConvertAll(x => new PoolId(x)), Acl.Merge(Agent.Acl, Update.Acl), Update.Comment);
 				if (NewAgent == null)
 				{
 					continue;
diff --git a/Engine/Source/Programs/Horde/Horde.Build/Fleet/Impl/Managers/Fleet/AwsFleetManager.cs b/Engine/Source/Programs/Horde/Horde.Build/Fleet/Impl/Managers/Fleet/AwsFleetManager.cs
--- a/Engine/Source/Programs/Horde/Horde.Build/Fleet/Impl/Managers/Fleet/AwsFleetManager.cs
+++ b/Engine/Source/Programs/Horde/Horde.Build/Fleet/Impl/Managers/Fleet/AwsFleetManager.cs
@@ -90,7 +90,7 @@ namespace HordeServer.Services.Impl
 				{
 					for (IAgent? NewAgent = Agent; NewAgent != null; NewAgent = await AgentCollection.GetAsync(Agent.Id))
 					{
-						if (await AgentCollection.TryUpdateSettingsAsync(NewAgent, bRequestShutdown: true) != null)
+						if (await AgentCollection.TryUpdateSettingsAsync(NewAgent, bRequestShutdown: true, ShutdownReason: "Autoscaler") != null)
 						{
 							AgentCollection.GetLogger(Agent.Id).LogInformation("Marked for shutdown due to autoscalar (currently {NumLeases} leases outstanding)", NewAgent.Leases.Count);
 							Count--;
