Horde: Add an additional message for when a lease is aborted.

diff --git a/Engine/Source/Programs/Horde/Horde.Agent/Services/WorkerService.cs b/Engine/Source/Programs/Horde/Horde.Agent/Services/WorkerService.cs
--- a/Engine/Source/Programs/Horde/Horde.Agent/Services/WorkerService.cs
+++ b/Engine/Source/Programs/Horde/Horde.Agent/Services/WorkerService.cs
@@ -945,6 +945,7 @@ namespace HordeAgent.Services
 			{
 				if (CancellationToken.IsCancellationRequested && IsCancellationException(Ex))
 				{
+					Logger.LogError("Step was aborted");
 					throw;
 				}
 				else
