Horde: Use {TemplateId} consistently for referencing templates.

diff --git a/Engine/Source/Programs/Horde/Horde.Build/Controllers/JobsController.cs b/Engine/Source/Programs/Horde/Horde.Build/Controllers/JobsController.cs
--- a/Engine/Source/Programs/Horde/Horde.Build/Controllers/JobsController.cs
+++ b/Engine/Source/Programs/Horde/Horde.Build/Controllers/JobsController.cs
@@ -219,12 +219,12 @@ namespace HordeServer.Controllers
 			IList<IJob> Jobs = await JobService.FindJobsAsync(StreamId: Stream.Id, Templates: new[] { TemplateId }, Target: Target, State: new[] { JobStepState.Completed }, Outcome: Outcomes.ToArray(), Count: 1);
 			if (Jobs.Count == 0)
 			{
-				Logger.LogInformation("Unable to find successful build of {TemplateRefId} target {Target}. Using latest change instead", TemplateId, Target);
+				Logger.LogInformation("Unable to find successful build of {TemplateId} target {Target}. Using latest change instead", TemplateId, Target);
 				return await Perforce.GetLatestChangeAsync(Stream.ClusterName, Stream.Name, null);
 			}
 			else
 			{
-				Logger.LogInformation("Last successful build of {TemplateRefId} target {Target} was job {JobId} at change {Change}", TemplateId, Target, Jobs[0].Id, Jobs[0].Change);
+				Logger.LogInformation("Last successful build of {TemplateId} target {Target} was job {JobId} at change {Change}", TemplateId, Target, Jobs[0].Id, Jobs[0].Change);
 				return Jobs[0].Change;
 			}
 		}
diff --git a/Engine/Source/Programs/Horde/Horde.Build/Controllers/StreamsController.cs b/Engine/Source/Programs/Horde/Horde.Build/Controllers/StreamsController.cs
--- a/Engine/Source/Programs/Horde/Horde.Build/Controllers/StreamsController.cs
+++ b/Engine/Source/Programs/Horde/Horde.Build/Controllers/StreamsController.cs
@@ -268,13 +268,13 @@ namespace HordeServer.Controllers
 		/// Gets all the templates for a stream
 		/// </summary>
 		/// <param name="StreamId">Unique id of the stream to query</param>
-		/// <param name="TemplateRefId">Unique id of the template to query</param>
+		/// <param name="TemplateId">Unique id of the template to query</param>
 		/// <param name="Filter">Filter for properties to return</param>
 		/// <returns>Information about all the templates</returns>
 		[HttpGet]
-		[Route("/api/v1/streams/{StreamId}/templates/{TemplateRefId}")]
+		[Route("/api/v1/streams/{StreamId}/templates/{TemplateId}")]
 		[ProducesResponseType(typeof(List<GetTemplateResponse>), 200)]
-		public async Task<ActionResult<object>> GetTemplateAsync(StreamId StreamId, TemplateRefId TemplateRefId, [FromQuery] PropertyFilter? Filter = null)
+		public async Task<ActionResult<object>> GetTemplateAsync(StreamId StreamId, TemplateRefId TemplateId, [FromQuery] PropertyFilter? Filter = null)
 		{
 			IStream? Stream = await StreamService.GetStreamAsync(StreamId);
 			if (Stream == null)
@@ -283,9 +283,9 @@ namespace HordeServer.Controllers
 			}
 
 			TemplateRef? TemplateRef;
-			if (!Stream.Templates.TryGetValue(TemplateRefId, out TemplateRef))
+			if (!Stream.Templates.TryGetValue(TemplateId, out TemplateRef))
 			{
-				return NotFound(StreamId, TemplateRefId);
+				return NotFound(StreamId, TemplateId);
 			}
 			if (!await StreamService.AuthorizeAsync(Stream, TemplateRef, AclAction.ViewTemplate, User, null))
 			{
diff --git a/Engine/Source/Programs/Horde/Horde.Build/Jobs/Impl/JobService.cs b/Engine/Source/Programs/Horde/Horde.Build/Jobs/Impl/JobService.cs
--- a/Engine/Source/Programs/Horde/Horde.Build/Jobs/Impl/JobService.cs
+++ b/Engine/Source/Programs/Horde/Horde.Build/Jobs/Impl/JobService.cs
@@ -1239,7 +1239,7 @@ namespace HordeServer.Services
 					TemplateRef? TemplateRef;
 					if (!Stream.Templates.TryGetValue(JobTrigger.TemplateRefId, out TemplateRef))
 					{
-						Logger.LogWarning("Cannot find template {TemplateRefId} in stream {StreamId}", JobTrigger.TemplateRefId, NewJob.StreamId);
+						Logger.LogWarning("Cannot find template {TemplateId} in stream {StreamId}", JobTrigger.TemplateRefId, NewJob.StreamId);
 						break;
 					}
 
diff --git a/Engine/Source/Programs/Horde/Horde.Build/Projects/Impl/StreamCollection.cs b/Engine/Source/Programs/Horde/Horde.Build/Projects/Impl/StreamCollection.cs
--- a/Engine/Source/Programs/Horde/Horde.Build/Projects/Impl/StreamCollection.cs
+++ b/Engine/Source/Programs/Horde/Horde.Build/Projects/Impl/StreamCollection.cs
@@ -329,28 +329,28 @@ namespace HordeServer.Collections.Impl
 		}
 
 		/// <inheritdoc/>
-		public async Task<IStream?> TryUpdateScheduleTriggerAsync(IStream StreamInterface, TemplateRefId TemplateRefId, DateTime? LastTriggerTimeUtc, int? LastTriggerChange, List<JobId> NewActiveJobs)
+		public async Task<IStream?> TryUpdateScheduleTriggerAsync(IStream StreamInterface, TemplateRefId TemplateId, DateTime? LastTriggerTimeUtc, int? LastTriggerChange, List<JobId> NewActiveJobs)
 		{
 			StreamDocument Stream = (StreamDocument)StreamInterface;
-			Schedule Schedule = Stream.Templates[TemplateRefId].Schedule!;
+			Schedule Schedule = Stream.Templates[TemplateId].Schedule!;
 
 			// Build the updates. MongoDB driver cannot parse TemplateRefId in expression tree; need to specify field name explicitly
 			List<UpdateDefinition<StreamDocument>> Updates = new List<UpdateDefinition<StreamDocument>>();
 			if (LastTriggerTimeUtc.HasValue && LastTriggerTimeUtc.Value != Schedule.LastTriggerTime)
 			{
-				FieldDefinition<StreamDocument, DateTimeOffset> LastTriggerTimeField = $"{nameof(Stream.Templates)}.{TemplateRefId}.{nameof(Schedule)}.{nameof(Schedule.LastTriggerTime)}";
+				FieldDefinition<StreamDocument, DateTimeOffset> LastTriggerTimeField = $"{nameof(Stream.Templates)}.{TemplateId}.{nameof(Schedule)}.{nameof(Schedule.LastTriggerTime)}";
 				Updates.Add(Builders<StreamDocument>.Update.Set(LastTriggerTimeField, LastTriggerTimeUtc.Value));
 				Schedule.LastTriggerTimeUtc = LastTriggerTimeUtc.Value;
 			}
 			if (LastTriggerChange.HasValue && LastTriggerChange.Value > Schedule.LastTriggerChange)
 			{
-				FieldDefinition<StreamDocument, int> LastTriggerChangeField = $"{nameof(Stream.Templates)}.{TemplateRefId}.{nameof(Schedule)}.{nameof(Schedule.LastTriggerChange)}";
+				FieldDefinition<StreamDocument, int> LastTriggerChangeField = $"{nameof(Stream.Templates)}.{TemplateId}.{nameof(Schedule)}.{nameof(Schedule.LastTriggerChange)}";
 				Updates.Add(Builders<StreamDocument>.Update.Set(LastTriggerChangeField, LastTriggerChange.Value));
 				Schedule.LastTriggerChange = LastTriggerChange.Value;
 			}
 			if (NewActiveJobs != null)
 			{
-				FieldDefinition<StreamDocument, List<JobId>> Field = $"{nameof(Stream.Templates)}.{TemplateRefId}.{nameof(Schedule)}.{nameof(Schedule.ActiveJobs)}";
+				FieldDefinition<StreamDocument, List<JobId>> Field = $"{nameof(Stream.Templates)}.{TemplateId}.{nameof(Schedule)}.{nameof(Schedule.ActiveJobs)}";
 				Updates.Add(Builders<StreamDocument>.Update.Set(Field, NewActiveJobs));
 				Schedule.ActiveJobs = NewActiveJobs;
 			}
diff --git a/Engine/Source/Programs/Horde/Horde.Build/Projects/ScheduleService.cs b/Engine/Source/Programs/Horde/Horde.Build/Projects/ScheduleService.cs
--- a/Engine/Source/Programs/Horde/Horde.Build/Projects/ScheduleService.cs
+++ b/Engine/Source/Programs/Horde/Horde.Build/Projects/ScheduleService.cs
@@ -206,7 +206,7 @@ namespace HordeServer.Services
 				// Try to prevent any updates being done for 5 minutes, but keep extending that again every minute. This
 				// creates a reservation for the stream update that will expire if the pod is terminated.
 				TimeSpan ReservationTime = TimeSpan.FromMinutes(5.0);
-				if (!await TrySetNextUpdateTime(Entry, DateTime.UtcNow + ReservationTime))
+				if (!await TrySetNextUpdateTime(Entry, Clock.UtcNow + ReservationTime))
 				{
 					continue;
 				}
@@ -352,7 +352,7 @@ namespace HordeServer.Services
 			}
 			catch (Exception Ex)
 			{
-				Logger.LogError(Ex, "Failed to start schedule {StreamId}/{TemplateRefId}", Stream.Id, TemplateId);
+				Logger.LogError(Ex, "Failed to start schedule {StreamId}/{TemplateId}", Stream.Id, TemplateId);
 			}
 
 			// Print some timing info
@@ -364,21 +364,21 @@ namespace HordeServer.Services
 		/// Trigger a schedule to run
 		/// </summary>
 		/// <param name="Stream">Stream for the schedule</param>
-		/// <param name="TemplateRefId"></param>
+		/// <param name="TemplateId"></param>
 		/// <param name="TemplateRef"></param>
 		/// <param name="Schedule"></param>
 		/// <param name="NumActiveJobs"></param>
 		/// <param name="UtcNow">The current time</param>
 		/// <returns>Async task</returns>
-		private async Task TriggerAsync(IStream Stream, TemplateRefId TemplateRefId, TemplateRef TemplateRef, Schedule Schedule, int NumActiveJobs, DateTime UtcNow)
+		private async Task TriggerAsync(IStream Stream, TemplateRefId TemplateId, TemplateRef TemplateRef, Schedule Schedule, int NumActiveJobs, DateTime UtcNow)
 		{
 			// Check we're not already at the maximum number of allowed jobs
 			if (Schedule.MaxActive != 0 && NumActiveJobs >= Schedule.MaxActive)
 			{
-				Logger.LogInformation("Skipping trigger of {StreamId} template {TemplateRefId} - already have maximum number of jobs running ({NumJobs})", Stream.Id, TemplateRefId, Schedule.MaxActive);
+				Logger.LogInformation("Skipping trigger of {StreamId} template {TemplateId} - already have maximum number of jobs running ({NumJobs})", Stream.Id, TemplateId, Schedule.MaxActive);
 				foreach (JobId JobId in Schedule.ActiveJobs)
 				{
-					Logger.LogInformation("Active job for {StreamId} template {TemplateRefId}: {JobId}", Stream.Id, TemplateRefId, JobId);
+					Logger.LogInformation("Active job for {StreamId} template {TemplateId}: {JobId}", Stream.Id, TemplateId, JobId);
 				}
 				return;
 			}
@@ -423,7 +423,7 @@ namespace HordeServer.Services
 				int Change, CodeChange;
 				if (Schedule.Gate != null)
 				{
-					(Change, CodeChange) = await GetNextChangeForGateAsync(Stream.Id, TemplateRefId, Schedule.Gate, MinChangeNumber, MaxChangeNumber);
+					(Change, CodeChange) = await GetNextChangeForGateAsync(Stream.Id, TemplateId, Schedule.Gate, MinChangeNumber, MaxChangeNumber);
 				}
 				else
 				{
@@ -467,7 +467,7 @@ namespace HordeServer.Services
 			// Early out if there's nothing to do
 			if (TriggerChanges.Count == 0)
 			{
-				Logger.LogInformation("Skipping trigger of {StreamName} template {TemplateRefId} - no candidate changes after CL {LastTriggerChange}", Stream.Id, TemplateRefId, Schedule.LastTriggerChange);
+				Logger.LogInformation("Skipping trigger of {StreamName} template {TemplateId} - no candidate changes after CL {LastTriggerChange}", Stream.Id, TemplateId, Schedule.LastTriggerChange);
 				return;
 			}
 
@@ -475,7 +475,7 @@ namespace HordeServer.Services
 			ITemplate? Template = await TemplateCollection.GetAsync(TemplateRef.Hash);
 			if (Template == null)
 			{
-				Logger.LogWarning("Unable to find template '{TemplateHash}' for '{TemplateRefId}'", TemplateRef.Hash, TemplateRefId);
+				Logger.LogWarning("Unable to find template '{TemplateHash}' for '{TemplateId}'", TemplateRef.Hash, TemplateId);
 				return;
 			}
 
@@ -494,9 +494,9 @@ namespace HordeServer.Services
 			foreach ((int Change, int CodeChange) in TriggerChanges.OrderBy(x => x.Change))
 			{
 				List<string> DefaultArguments = Template.GetDefaultArguments();
-				IJob NewJob = await JobService.CreateJobAsync(null, Stream, TemplateRefId, Template.Id, Graph, Template.Name, Change, CodeChange, null, null, null, Template.Priority, null, null, TemplateRef.ChainedJobs, TemplateRef.ShowUgsBadges, TemplateRef.ShowUgsAlerts, TemplateRef.NotificationChannel, TemplateRef.NotificationChannelFilter, DefaultArguments);
+				IJob NewJob = await JobService.CreateJobAsync(null, Stream, TemplateId, Template.Id, Graph, Template.Name, Change, CodeChange, null, null, null, Template.Priority, null, null, TemplateRef.ChainedJobs, TemplateRef.ShowUgsBadges, TemplateRef.ShowUgsAlerts, TemplateRef.NotificationChannel, TemplateRef.NotificationChannelFilter, DefaultArguments);
 				Logger.LogInformation("Started new job for {StreamName} template {TemplateName} at CL {Change} (Code CL {CodeChange}): {JobId}", Stream.Id, TemplateRef.Name, Change, CodeChange, NewJob.Id);
-				await StreamService.UpdateScheduleTriggerAsync(Stream, TemplateRefId, UtcNow, Change, new List<JobId> { NewJob.Id }, new List<JobId>());
+				await StreamService.UpdateScheduleTriggerAsync(Stream, TemplateId, UtcNow, Change, new List<JobId> { NewJob.Id }, new List<JobId>());
 			}
 		}
 
@@ -578,7 +578,7 @@ namespace HordeServer.Services
 						{
 							return (Job.Change, Job.CodeChange);
 						}
-						Logger.LogInformation("Skipping trigger of {StreamName} template {TemplateRefId} - last {OtherTemplateRefId} job ({JobId}) ended with errors", StreamId, TemplateRefId, Gate.TemplateRefId, Job.Id);
+						Logger.LogInformation("Skipping trigger of {StreamName} template {TemplateId} - last {OtherTemplateRefId} job ({JobId}) ended with errors", StreamId, TemplateRefId, Gate.TemplateRefId, Job.Id);
 					}
 				}
 
