Horde: Fix log message referencing a template name rather than id.

diff --git a/Engine/Source/Programs/Horde/Horde.Build/Projects/ScheduleService.cs b/Engine/Source/Programs/Horde/Horde.Build/Projects/ScheduleService.cs
--- a/Engine/Source/Programs/Horde/Horde.Build/Projects/ScheduleService.cs
+++ b/Engine/Source/Programs/Horde/Horde.Build/Projects/ScheduleService.cs
@@ -495,7 +495,7 @@ namespace HordeServer.Services
 			{
 				List<string> DefaultArguments = Template.GetDefaultArguments();
 				IJob NewJob = await JobService.CreateJobAsync(null, Stream, TemplateId, Template.Id, Graph, Template.Name, Change, CodeChange, null, null, null, Template.Priority, null, null, TemplateRef.ChainedJobs, TemplateRef.ShowUgsBadges, TemplateRef.ShowUgsAlerts, TemplateRef.NotificationChannel, TemplateRef.NotificationChannelFilter, DefaultArguments);
-				Logger.LogInformation("Started new job for {StreamName} template {TemplateName} at CL {Change} (Code CL {CodeChange}): {JobId}", Stream.Id, TemplateRef.Name, Change, CodeChange, NewJob.Id);
+				Logger.LogInformation("Started new job for {StreamName} template {TemplateId} at CL {Change} (Code CL {CodeChange}): {JobId}", Stream.Id, Template.Id, Change, CodeChange, NewJob.Id);
 				await StreamService.UpdateScheduleTriggerAsync(Stream, TemplateId, UtcNow, Change, new List<JobId> { NewJob.Id }, new List<JobId>());
 			}
 		}
