Horde: Add a separate set of locks for schedule updates to prevent races on queue updates.

diff --git a/Engine/Source/Programs/Horde/Horde.Build/Projects/ScheduleService.cs b/Engine/Source/Programs/Horde/Horde.Build/Projects/ScheduleService.cs
--- a/Engine/Source/Programs/Horde/Horde.Build/Projects/ScheduleService.cs
+++ b/Engine/Source/Programs/Horde/Horde.Build/Projects/ScheduleService.cs
@@ -75,7 +75,8 @@ namespace HordeServer.Services
 		IClock Clock;
 		ILogger Logger;
 		RedisService Redis;
-		RedisSortedSet<QueueItem> Queue;
+		RedisSortedSet<QueueItem> Queue; // Next tick time for each queue item
+		RedisSortedSet<QueueItem> Locks; // Currently locked queue items (locks expire after a period of time)
 		BackgroundTask BackgroundTask;
 
 		/// <summary>
@@ -110,6 +111,7 @@ namespace HordeServer.Services
 
 			this.Redis = Redis;
 			this.Queue = new RedisSortedSet<QueueItem>(Redis.Database, "scheduler/queue");
+			this.Locks = new RedisSortedSet<QueueItem>(Redis.Database, "scheduler/locks");
 			this.BackgroundTask = new BackgroundTask(RunAsync);
 		}
 
@@ -151,6 +153,9 @@ namespace HordeServer.Services
 		{
 			DateTime UtcNow = Clock.UtcNow;
 
+			// Remove any expired locks
+			await Locks.RemoveRangeByScoreAsync(-double.NegativeInfinity, QueueItem.GetScoreFromTime(UtcNow));
+
 			// Keep updating schedules
 			while (!CancellationToken.IsCancellationRequested)
 			{
@@ -162,6 +167,7 @@ namespace HordeServer.Services
 					return TimeSpan.FromSeconds(10.0);
 				}
 
+				// Deconstruct the item
 				SortedSetEntry<QueueItem> Entry = Entries[0];
 				QueueItem Item = Entry.Element;
 				DateTime ScheduledTime = QueueItem.GetTimeFromScore(Entry.Score);
@@ -180,51 +186,28 @@ namespace HordeServer.Services
 					continue;
 				}
 
-				// Get the stream and template
-				IStream? Stream = await StreamService.GetStreamAsync(Item.StreamId);
-				if (Stream == null || !Stream.Templates.TryGetValue(Item.TemplateId, out TemplateRef? TemplateRef) || TemplateRef.Schedule == null)
-				{
-					await Queue.RemoveAsync(Item);
-					continue;
-				}
-
-				// Calculate the next trigger time for this schedule; it may have changed since the queue item was created.
-				DateTime? NextUpdateTimeUtc = TemplateRef.Schedule.GetNextTriggerTimeUtc(TimeZone);
-				if (NextUpdateTimeUtc == null)
+				// Try to create a lock on this item before we consider updating it. If it fails, just update the schedule time to a while from now.
+				double LockTime = QueueItem.GetScoreFromTime(Clock.UtcNow + TimeSpan.FromMinutes(5.0));
+				if (!await Locks.AddAsync(Item, LockTime, When.NotExists))
 				{
-					await Queue.RemoveAsync(Item);
-					continue;
-				}
-
-				// Check if we're ready to trigger now
-				if (UtcNow < NextUpdateTimeUtc)
-				{
-					await TrySetNextUpdateTime(Entry, NextUpdateTimeUtc.Value);
+					await Queue.AddAsync(Item, QueueItem.GetScoreFromTime(Clock.UtcNow + TimeSpan.FromSeconds(20.0)));
 					continue;
 				}
 
-				// Try to prevent any updates being done for 5 minutes, but keep extending that again every minute. This
-				// creates a reservation for the stream update that will expire if the pod is terminated.
-				TimeSpan ReservationTime = TimeSpan.FromMinutes(5.0);
-				if (!await TrySetNextUpdateTime(Entry, Clock.UtcNow + ReservationTime))
-				{
-					continue;
-				}
-
-				// Try to do the update
+				// Spawn the task to tick the template while simultaneously extending the lock time
+				Task TickTemplateTask = Task.Run(() => TickTemplateAsync(Item, UtcNow, CancellationToken));
 				try
 				{
-					Task TriggerTask = Task.Run(() => TriggerAsync(Stream, Item.TemplateId, TemplateRef, UtcNow));
 					for (; ; )
 					{
 						Task DelayTask = Task.Delay(TimeSpan.FromMinutes(1.0));
-						if (await Task.WhenAny(DelayTask, TriggerTask) == TriggerTask)
+						if (await Task.WhenAny(TickTemplateTask, DelayTask) == TickTemplateTask)
 						{
+							await TickTemplateTask;
 							break;
 						}
-						await SetNextUpdateTime(Item, Clock.UtcNow + ReservationTime);
+						await Locks.AddAsync(Item, QueueItem.GetScoreFromTime(Clock.UtcNow + TimeSpan.FromMinutes(5.0)));
 					}
-					await TriggerTask;
 				}
 				catch (OperationCanceledException)
 				{
@@ -235,6 +218,45 @@ namespace HordeServer.Services
 					Logger.LogError(Ex, "Error while updating schedule for {StreamId}/{TemplateId}", Item.StreamId, Item.TemplateId);
 				}
 
+				// Remove the lock
+				await Locks.RemoveAsync(Item);
+			}
+
+			return TimeSpan.Zero;
+		}
+
+		async Task TickTemplateAsync(QueueItem Item, DateTime UtcNow, CancellationToken CancellationToken)
+		{
+			// Get the stream and template
+			IStream? Stream = await StreamService.GetStreamAsync(Item.StreamId);
+			if (Stream == null || !Stream.Templates.TryGetValue(Item.TemplateId, out TemplateRef? TemplateRef) || TemplateRef.Schedule == null)
+			{
+				await Queue.RemoveAsync(Item);
+				return;
+			}
+
+			// Calculate the next trigger time for this schedule; it may have changed since the queue item was created.
+			DateTime? NextUpdateTimeUtc = TemplateRef.Schedule.GetNextTriggerTimeUtc(TimeZone);
+			if (NextUpdateTimeUtc == null)
+			{
+				await Queue.RemoveAsync(Item);
+				return;
+			}
+
+			// Check if we're ready to trigger now
+			if (UtcNow < NextUpdateTimeUtc)
+			{
+				await Queue.AddAsync(Item, QueueItem.GetScoreFromTime(NextUpdateTimeUtc.Value));
+				return;
+			}
+
+			// Set the update time on this stream. If it fails, just return and try again.
+			Stream = await StreamService.UpdateScheduleTriggerAsync(Stream, Item.TemplateId, LastTriggerTimeUtc: Clock.UtcNow);
+			if (Stream != null)
+			{
+				// Try to do the update
+				await TriggerAsync(Stream, Item.TemplateId, TemplateRef, UtcNow, CancellationToken);
+
 				// Update the next trigger time
 				DateTime? NextTime = TemplateRef.Schedule.GetNextTriggerTimeUtc(UtcNow, TimeZone);
 				if (NextTime == null)
@@ -243,15 +265,15 @@ namespace HordeServer.Services
 				}
 				else
 				{
-					await SetNextUpdateTime(Item, NextTime.Value);
+					await Queue.AddAsync(Item, QueueItem.GetScoreFromTime(NextTime.Value));
 				}
 			}
-			return TimeSpan.Zero;
 		}
 
 		internal async Task ResetAsync()
 		{
 			await Redis.Database.KeyDeleteAsync(Queue.Key);
+			await Redis.Database.KeyDeleteAsync(Locks.Key);
 		}
 
 		internal async Task TickForTestingAsync()
@@ -260,19 +282,6 @@ namespace HordeServer.Services
 			await TickAsync(CancellationToken.None);
 		}
 
-		Task<bool> TrySetNextUpdateTime(SortedSetEntry<QueueItem> Entry, DateTime NextTimeUtc)
-		{
-			ITransaction Transaction = Redis.Database.CreateTransaction();
-			Transaction.AddCondition(Condition.SortedSetEqual(Queue.Key, Entry.ElementValue, Entry.Score));
-			_ = Transaction.With(Queue).AddAsync(Entry.Element, QueueItem.GetScoreFromTime(NextTimeUtc));
-			return Transaction.ExecuteAsync();
-		}
-
-		async Task SetNextUpdateTime(QueueItem Item, DateTime NextTimeUtc)
-		{
-			await Queue.AddAsync(Item, QueueItem.GetScoreFromTime(NextTimeUtc));
-		}
-
 		/// <summary>
 		/// Get the current set of streams and ensure there's an entry for each item
 		/// </summary>
@@ -308,8 +317,9 @@ namespace HordeServer.Services
 		/// <param name="TemplateId"></param>
 		/// <param name="TemplateRef"></param>
 		/// <param name="UtcNow"></param>
+		/// <param name="CancellationToken"></param>
 		/// <returns>Async task</returns>
-		private async Task TriggerAsync(IStream Stream, TemplateRefId TemplateId, TemplateRef TemplateRef, DateTime UtcNow)
+		private async Task TriggerAsync(IStream Stream, TemplateRefId TemplateId, TemplateRef TemplateRef, DateTime UtcNow, CancellationToken CancellationToken)
 		{
 			using IScope Scope = GlobalTracer.Instance.BuildSpan("ScheduleService.TriggerAsync").StartActive();
 			Scope.Span.SetTag("StreamId", Stream.Id);
@@ -339,7 +349,7 @@ namespace HordeServer.Services
 			await StreamService.UpdateScheduleTriggerAsync(Stream, TemplateId, RemoveJobs: RemoveJobIds);
 
 			// If the stream is paused, bail out
-			if (Stream.IsPaused(Clock.UtcNow))
+			if (Stream.IsPaused(UtcNow))
 			{
 				Logger.LogDebug("Skipping schedule update for stream {StreamId}. It has been paused until {PausedUntil} with comment '{PauseComment}'.", Stream.Id, Stream.PausedUntil, Stream.PauseComment);
 				return;
@@ -348,7 +358,7 @@ namespace HordeServer.Services
 			// Trigger this schedule
 			try
 			{
-				await TriggerAsync(Stream, TemplateId, TemplateRef, Schedule, Schedule.ActiveJobs.Count - RemoveJobIds.Count, UtcNow);
+				await TriggerAsync(Stream, TemplateId, TemplateRef, Schedule, Schedule.ActiveJobs.Count - RemoveJobIds.Count, UtcNow, CancellationToken);
 			}
 			catch (Exception Ex)
 			{
@@ -369,8 +379,9 @@ namespace HordeServer.Services
 		/// <param name="Schedule"></param>
 		/// <param name="NumActiveJobs"></param>
 		/// <param name="UtcNow">The current time</param>
+		/// <param name="CancellationToken"></param>
 		/// <returns>Async task</returns>
-		private async Task TriggerAsync(IStream Stream, TemplateRefId TemplateId, TemplateRef TemplateRef, Schedule Schedule, int NumActiveJobs, DateTime UtcNow)
+		private async Task TriggerAsync(IStream Stream, TemplateRefId TemplateId, TemplateRef TemplateRef, Schedule Schedule, int NumActiveJobs, DateTime UtcNow, CancellationToken CancellationToken)
 		{
 			// Check we're not already at the maximum number of allowed jobs
 			if (Schedule.MaxActive != 0 && NumActiveJobs >= Schedule.MaxActive)
@@ -419,11 +430,13 @@ namespace HordeServer.Services
 			List<(int Change, int CodeChange)> TriggerChanges = new List<(int, int)>();
 			while (TriggerChanges.Count < MaxNewChanges)
 			{
+				CancellationToken.ThrowIfCancellationRequested();
+
 				// Get the next valid change
 				int Change, CodeChange;
 				if (Schedule.Gate != null)
 				{
-					(Change, CodeChange) = await GetNextChangeForGateAsync(Stream.Id, TemplateId, Schedule.Gate, MinChangeNumber, MaxChangeNumber);
+					(Change, CodeChange) = await GetNextChangeForGateAsync(Stream.Id, TemplateId, Schedule.Gate, MinChangeNumber, MaxChangeNumber, CancellationToken);
 				}
 				else
 				{
@@ -493,6 +506,7 @@ namespace HordeServer.Services
 			// Try to start all the new jobs
 			foreach ((int Change, int CodeChange) in TriggerChanges.OrderBy(x => x.Change))
 			{
+				CancellationToken.ThrowIfCancellationRequested();
 				List<string> DefaultArguments = Template.GetDefaultArguments();
 				IJob NewJob = await JobService.CreateJobAsync(null, Stream, TemplateId, Template.Id, Graph, Template.Name, Change, CodeChange, null, null, null, Template.Priority, null, null, TemplateRef.ChainedJobs, TemplateRef.ShowUgsBadges, TemplateRef.ShowUgsAlerts, TemplateRef.NotificationChannel, TemplateRef.NotificationChannelFilter, DefaultArguments);
 				Logger.LogInformation("Started new job for {StreamName} template {TemplateId} at CL {Change} (Code CL {CodeChange}): {JobId}", Stream.Id, Template.Id, Change, CodeChange, NewJob.Id);
@@ -555,10 +569,12 @@ namespace HordeServer.Services
 		/// Gets the next change to build for a schedule on a gate
 		/// </summary>
 		/// <returns></returns>
-		private async Task<(int Change, int CodeChange)> GetNextChangeForGateAsync(StreamId StreamId, TemplateRefId TemplateRefId, ScheduleGate Gate, int? MinChange, int? MaxChange)
+		private async Task<(int Change, int CodeChange)> GetNextChangeForGateAsync(StreamId StreamId, TemplateRefId TemplateRefId, ScheduleGate Gate, int? MinChange, int? MaxChange, CancellationToken CancellationToken)
 		{
 			for (; ; )
 			{
+				CancellationToken.ThrowIfCancellationRequested();
+
 				List<IJob> Jobs = await JobCollection.FindAsync(StreamId: StreamId, Templates: new[] { Gate.TemplateRefId }, MinChange: MinChange, MaxChange: MaxChange, Count: 1);
 				if (Jobs.Count == 0)
 				{
