Horde: Add additional status update filtering and agent shutdown controls
#rnx

diff --git a/Engine/Source/Programs/Horde/HordeDashboard/src/backend/Api.ts b/Engine/Source/Programs/Horde/HordeDashboard/src/backend/Api.ts
--- a/Engine/Source/Programs/Horde/HordeDashboard/src/backend/Api.ts
+++ b/Engine/Source/Programs/Horde/HordeDashboard/src/backend/Api.ts
@@ -393,6 +393,8 @@ export type UpdateAgentRequest = {
 
 	requestRestart?: boolean;
 
+	requestShutdown?: boolean;
+
 	/**Per-agent override for the desired client version */
 	forceVersion?: string;
 
@@ -647,6 +649,12 @@ export type GetAgentResponse = {
 	/** Next conform attempt time */
 	nextConformTime?: Date | string;
 
+	/** The reason for the last shutdown */
+	lastShutdownReason: string;
+
+	/** Whether a shutdown is pending */
+	pendingShutdown: boolean;
+
 	/** agent workspaces */
 	workspaces: GetAgentWorkspaceResponse[];
 
