Horde Storage
Fixed issue in the orphan blob cleanup were it wouldnt check other namespaces for references to blobs before deleting.

Now checks all namespaces that share a storage pool

diff --git a/Engine/Source/Programs/Horde/HordeStorage/Horde.Storage/Implementation/GC/OrphanBlobCleanupRefs.cs b/Engine/Source/Programs/Horde/HordeStorage/Horde.Storage/Implementation/GC/OrphanBlobCleanupRefs.cs
--- a/Engine/Source/Programs/Horde/HordeStorage/Horde.Storage/Implementation/GC/OrphanBlobCleanupRefs.cs
+++ b/Engine/Source/Programs/Horde/HordeStorage/Horde.Storage/Implementation/GC/OrphanBlobCleanupRefs.cs
@@ -2,12 +2,14 @@
 
 using System;
 using System.Collections.Generic;
+using System.Linq;
 using System.Threading;
 using System.Threading.Tasks;
 using async_enumerable_dotnet;
 using Dasync.Collections;
 using Horde.Storage.Controllers;
 using Horde.Storage.Implementation.Blob;
+using Jupiter;
 using Jupiter.Implementation;
 using Microsoft.Extensions.Options;
 using Serilog;
@@ -21,16 +23,18 @@ namespace Horde.Storage.Implementation
         private readonly IBlobIndex _blobIndex;
         private readonly ILeaderElection _leaderElection;
         private readonly IOptionsMonitor<GCSettings> _gcSettings;
+        private readonly IOptionsMonitor<NamespaceSettings> _namespaceSettings;
         private readonly ILogger _logger = Log.ForContext<OrphanBlobCleanupRefs>();
 
         // ReSharper disable once UnusedMember.Global
-        public OrphanBlobCleanupRefs(IBlobService blobService, IObjectService objectService, IBlobIndex blobIndex, ILeaderElection leaderElection, IOptionsMonitor<GCSettings> gcSettings)
+        public OrphanBlobCleanupRefs(IBlobService blobService, IObjectService objectService, IBlobIndex blobIndex, ILeaderElection leaderElection, IOptionsMonitor<GCSettings> gcSettings, IOptionsMonitor<NamespaceSettings> namespaceSettings)
         {
             _blobService = blobService;
             _objectService = objectService;
             _blobIndex = blobIndex;
             _leaderElection = leaderElection;
             _gcSettings = gcSettings;
+            _namespaceSettings = namespaceSettings;
         }
 
         public async Task<List<BlobIdentifier>> Cleanup(CancellationToken cancellationToken)
@@ -61,38 +65,40 @@ namespace Horde.Storage.Implementation
                     if (cancellationToken.IsCancellationRequested)
                         break;
 
-                    IBlobIndex.BlobInfo? blobIndex = await _blobIndex.GetBlobInfo(@namespace, blob);
-
-                    bool shouldDelete = blobIndex == null;
-                    if (blobIndex != null)
+                    bool found = false;
+                    NamespaceSettings.PerNamespaceSettings policy = _namespaceSettings.CurrentValue.GetPoliciesForNs(@namespace);
+                    
+                    // check all other namespaces that share the same storage pool for presence of the blob
+                    foreach (NamespaceId blobNamespace in namespaces.Where(ns => _namespaceSettings.CurrentValue.GetPoliciesForNs(ns).StoragePool == policy.StoragePool))
                     {
-                        bool found = false;
-                        foreach ((BucketId bucket, IoHashKey key) in blobIndex.References)
-                        {
-                            if (cancellationToken.IsCancellationRequested)
-                                break;
+                        IBlobIndex.BlobInfo? blobIndex = await _blobIndex.GetBlobInfo(blobNamespace, blob);
 
-                            try
-                            {
-                                (ObjectRecord, BlobContents) _ = await _objectService.Get(@namespace, bucket, key, Array.Empty<string>());
-                                found = true;
-                            }
-                            catch (ObjectNotFoundException)
-                            {
-                                // this is not a valid reference so we should delete
-                            }
-                            catch (MissingBlobsException)
+                        if (blobIndex != null)
+                        {
+                            foreach ((BucketId bucket, IoHashKey key) in blobIndex.References)
                             {
-                                // we do not care if there are missing blobs, as long as the record is valid we keep this blob around
-                                found = true;
+                                if (cancellationToken.IsCancellationRequested)
+                                    break;
+
+                                try
+                                {
+                                    (ObjectRecord, BlobContents) _ = await _objectService.Get(blobNamespace, bucket, key, Array.Empty<string>());
+                                    found = true;
+                                }
+                                catch (ObjectNotFoundException)
+                                {
+                                    // this is not a valid reference so we should delete
+                                }
+                                catch (MissingBlobsException)
+                                {
+                                    // we do not care if there are missing blobs, as long as the record is valid we keep this blob around
+                                    found = true;
+                                }
                             }
                         }
-
-                        if (!found)
-                            shouldDelete = true;
                     }
-                    
-                    if (shouldDelete)
+
+                    if (!found)
                     {
                         await RemoveBlob(@namespace, blob);
                         removedBlobs.Add(blob);
