Horde Storage

Added last access tracker for compact binary references
Last access service will now also read that those tracked references and updated the db
Ref cleanup will now also clean namespaces from compact binary references and delete old references.

diff --git a/Engine/Source/Programs/Horde/HordeStorage/Horde.Storage/Horde.Storage.Startup.cs b/Engine/Source/Programs/Horde/HordeStorage/Horde.Storage/Horde.Storage.Startup.cs
--- a/Engine/Source/Programs/Horde/HordeStorage/Horde.Storage/Horde.Storage.Startup.cs
+++ b/Engine/Source/Programs/Horde/HordeStorage/Horde.Storage/Horde.Storage.Startup.cs
@@ -228,6 +228,10 @@ namespace Horde.Storage
             services.AddSingleton(serviceType: typeof(ILastAccessCache<RefRecord>), lastAccessTracker);
             services.AddSingleton(serviceType: typeof(ILastAccessTracker<RefRecord>), lastAccessTracker);
 
+            LastAccessTrackerReference lastAccessTrackerReferences = new LastAccessTrackerReference();
+            services.AddSingleton(serviceType: typeof(ILastAccessCache<LastAccessRecord>), lastAccessTrackerReferences);
+            services.AddSingleton(serviceType: typeof(ILastAccessTracker<LastAccessRecord>), lastAccessTrackerReferences);
+
             services.AddSingleton(serviceType: typeof(ILeaderElection), CreateLeaderElection);
 
             services.AddSingleton(Configuration);
diff --git a/Engine/Source/Programs/Horde/HordeStorage/Horde.Storage/Implementation/GC/CacheCleanup.cs b/Engine/Source/Programs/Horde/HordeStorage/Horde.Storage/Implementation/GC/CacheCleanup.cs
--- a/Engine/Source/Programs/Horde/HordeStorage/Horde.Storage/Implementation/GC/CacheCleanup.cs
+++ b/Engine/Source/Programs/Horde/HordeStorage/Horde.Storage/Implementation/GC/CacheCleanup.cs
@@ -20,20 +20,70 @@ namespace Horde.Storage.Implementation
     {
         private readonly IOptionsMonitor<GCSettings> _settings;
         private readonly IRefsStore _refs;
+        private readonly IReferencesStore _referencesStore;
+        private readonly IReplicationLog _replicationLog;
         private readonly ITransactionLogWriter _transactionLog;
         private readonly ILogger _logger = Log.ForContext<RefCleanup>();
 
-        public RefCleanup(IOptionsMonitor<GCSettings> settings, IRefsStore refs, ITransactionLogWriter transactionLog)
+        public RefCleanup(IOptionsMonitor<GCSettings> settings, IRefsStore refs, ITransactionLogWriter transactionLog, IReferencesStore referencesStore, IReplicationLog replicationLog)
         {
             _settings = settings;
             _refs = refs;
             _transactionLog = transactionLog;
+            _referencesStore = referencesStore;
+            _replicationLog = replicationLog;
         }
 
-        public async Task<List<OldRecord>> Cleanup(NamespaceId ns, CancellationToken cancellationToken)
+        public Task<List<OldRecord>> Cleanup(NamespaceId ns, CancellationToken cancellationToken)
         {
             using Scope scope = Tracer.Instance.StartActive("gc.refs.namespace");
             scope.Span.ResourceName = ns.ToString();
+            if (_settings.CurrentValue.CleanNamespacesV1.Contains(ns.ToString()))
+            {
+                return CleanNamespaceV1(ns, cancellationToken);
+            }
+            else if (_settings.CurrentValue.CleanNamespaces.Contains(ns.ToString()))
+            {
+                return CleanNamespace(ns, cancellationToken);
+            }
+            else
+                throw new NotImplementedException(
+                    $"Namespace {ns} not present in CleanNamespaces or CleanNamespacesV1 lists, unable to clean it as we do not know which method to use.");
+           
+        }
+
+        private async Task<List<OldRecord>> CleanNamespace(NamespaceId ns, CancellationToken cancellationToken)
+        {
+            List<OldRecord> deletedRecords = new List<OldRecord>();
+            DateTime cutoffTime = DateTime.Now.Add(_settings.CurrentValue.LastAccessCutoff);
+            await foreach (ObjectRecord record in _referencesStore.GetOldestRecords(ns).WithCancellation(cancellationToken))
+            {
+                // the record are returned oldest first, so once we find a record that is newer then our cutoff time we can stop iterating
+                if (record.LastAccess > cutoffTime)
+                    break;
+
+                // delete the old record from the ref refs
+                Task storeDelete = _referencesStore.Delete(record.Namespace, record.Bucket, record.Name);
+                // insert a delete event into the transaction log
+                Task<(string, Guid)> transactionLogDelete = _replicationLog.InsertDeleteEvent(record.Namespace, record.Bucket, record.Name, record.BlobIdentifier);
+
+                try
+                {
+                    Task.WaitAll(storeDelete, transactionLogDelete);
+                }
+                catch (Exception e)
+                {
+                    _logger.Warning(e, "Exception when attempting to delete record {Record} in {Namespace}", record, ns);
+                }
+                // we convert the ObjectRecords key (which is a iohash) to a keyid (which is a generic string) not the prettiest but this result is only used for debugging purposes anyway
+                deletedRecords.Add(new OldRecord(record.Namespace, record.Bucket, new KeyId(record.Name.ToString())));
+            }
+
+            return deletedRecords;
+        }
+
+        private async Task<List<OldRecord>> CleanNamespaceV1(NamespaceId ns, CancellationToken cancellationToken)
+        {
             List<OldRecord> deletedRecords = new List<OldRecord>();
             await foreach (OldRecord record in _refs.GetOldRecords(ns, _settings.CurrentValue.LastAccessCutoff).WithCancellation(cancellationToken))
             {
diff --git a/Engine/Source/Programs/Horde/HordeStorage/Horde.Storage/Implementation/GC/RefCleanupService.cs b/Engine/Source/Programs/Horde/HordeStorage/Horde.Storage/Implementation/GC/RefCleanupService.cs
--- a/Engine/Source/Programs/Horde/HordeStorage/Horde.Storage/Implementation/GC/RefCleanupService.cs
+++ b/Engine/Source/Programs/Horde/HordeStorage/Horde.Storage/Implementation/GC/RefCleanupService.cs
@@ -16,6 +16,7 @@ namespace Horde.Storage.Implementation
     {
         private readonly IOptionsMonitor<GCSettings> _settings;
         private readonly ILeaderElection _leaderElection;
+        private readonly IReferencesStore _referencesStore;
         private volatile bool _alreadyPolling;
         private readonly ILogger _logger = Log.ForContext<RefCleanupService>();
 
@@ -32,11 +33,12 @@ namespace Horde.Storage.Implementation
             public IRefCleanup RefCleanup { get; }
         }
 
-        public RefCleanupService(IOptionsMonitor<GCSettings> settings, IRefsStore store, IRefCleanup refCleanup, ILeaderElection leaderElection) : 
+        public RefCleanupService(IOptionsMonitor<GCSettings> settings, IRefsStore store, IRefCleanup refCleanup, ILeaderElection leaderElection, IReferencesStore referencesStore) : 
             base(serviceName: nameof(RefCleanupService), settings.CurrentValue.RefCleanupPollFrequency, new RefCleanupState(store, refCleanup))
         {
             _settings = settings;
             _leaderElection = leaderElection;
+            _referencesStore = referencesStore;
         }
 
         public override bool ShouldStartPolling()
@@ -74,6 +76,23 @@ namespace Horde.Storage.Implementation
                     }
                 }
 
+                await foreach (NamespaceId ns in _referencesStore.GetNamespaces().WithCancellation(cancellationToken))
+                {
+                    using Scope scope = Tracer.Instance.StartActive("gc.refs");
+                    scope.Span.ResourceName = ns.ToString();
+
+                    _logger.Information("Attempting to run Refs Cleanup of {Namespace}. ", ns);
+                    try
+                    {
+                        List<OldRecord> oldRecords = await state.RefCleanup.Cleanup(ns, cancellationToken);
+                        _logger.Information("Ran Refs Cleanup of {Namespace}. Deleted {CountRefRecords}", ns, oldRecords.Count);
+                    }
+                    catch (Exception e)
+                    {
+                        _logger.Error("Error running Refs Cleanup of {Namespace}. {Exception}", ns, e);
+                    }
+                }
+
                 return true;
 
             }
diff --git a/Engine/Source/Programs/Horde/HordeStorage/Horde.Storage/Implementation/LastAccess/LastAccessServiceReferences.cs b/Engine/Source/Programs/Horde/HordeStorage/Horde.Storage/Implementation/LastAccess/LastAccessServiceReferences.cs
new file mode 100644
--- /dev/null
+++ b/Engine/Source/Programs/Horde/HordeStorage/Horde.Storage/Implementation/LastAccess/LastAccessServiceReferences.cs
@@ -0,0 +1,83 @@
+// Copyright Epic Games, Inc. All Rights Reserved.
+
+using System;
+using System.Collections.Generic;
+using System.Threading;
+using System.Threading.Tasks;
+using Datadog.Trace;
+using Microsoft.Extensions.Hosting;
+using Microsoft.Extensions.Options;
+using Serilog;
+
+namespace Horde.Storage.Implementation
+{
+    // ReSharper disable once ClassNeverInstantiated.Global
+    public class LastAccessServiceReferences : IHostedService, IDisposable
+    {
+        private readonly ILastAccessCache<LastAccessRecord> _lastAccessCacheRecord;
+        private readonly IReferencesStore _referencesStore;
+        private readonly ILogger _logger = Log.ForContext<LastAccessServiceReferences>();
+        private Timer? _timer;
+        private readonly HordeStorageSettings _settings;
+        
+        public bool Running { get; private set; }
+
+        public LastAccessServiceReferences(IOptionsMonitor<HordeStorageSettings> settings, ILastAccessCache<LastAccessRecord> lastAccessCache, IReferencesStore referencesStore)
+        {
+            _lastAccessCacheRecord = lastAccessCache;
+            _referencesStore = referencesStore;
+            _settings = settings.CurrentValue;
+        }
+
+        public Task StartAsync(CancellationToken cancellationToken)
+        {
+            _logger.Information("Last Access Aggregation service starting.");
+
+            _timer = new Timer(OnUpdate, null, TimeSpan.Zero,
+                period: TimeSpan.FromSeconds(_settings.LastAccessRollupFrequencySeconds));
+            Running = true;
+
+            return Task.CompletedTask;
+        }
+
+        public async Task StopAsync(CancellationToken cancellationToken)
+        {
+            _logger.Information("Last Access Aggregation service stopping.");
+
+            _timer?.Change(Timeout.Infinite, 0);
+            Running = false;
+
+            // process the last records we have built up
+            await ProcessLastAccessRecords();
+        }
+
+        private void OnUpdate(object? state)
+        {
+            // call results to make sure we join the task
+            Task.WaitAll(
+                ProcessLastAccessRecords()
+            );
+        }
+
+        internal async Task<List<(LastAccessRecord, DateTime)>> ProcessLastAccessRecords()
+        {
+            using Scope _ = Tracer.Instance.StartActive("lastAccessRefs.update");
+            _logger.Information("Running Last Access Aggregation for refs");
+            List<(LastAccessRecord, DateTime)> records = await _lastAccessCacheRecord.GetLastAccessedRecords();
+            foreach ((LastAccessRecord record, DateTime lastAccessTime) in records)
+            {
+                using Scope scope = Tracer.Instance.StartActive("lastAccess.update.record");
+                scope.Span.ResourceName = $"{record.Namespace}:{record.Bucket}.{record.Key}";
+                _logger.Debug("Updating last access time to {LastAccessTime} for {Record}", lastAccessTime, record);
+                await _referencesStore.UpdateLastAccessTime(record.Namespace, record.Bucket, record.Key, lastAccessTime);
+            }
+
+            return records;
+        }
+
+        public void Dispose()
+        {
+            _timer?.Dispose();
+        }
+    }
+}
diff --git a/Engine/Source/Programs/Horde/HordeStorage/Horde.Storage/Implementation/LastAccess/LastAccessTracker.cs b/Engine/Source/Programs/Horde/HordeStorage/Horde.Storage/Implementation/LastAccess/LastAccessTracker.cs
--- a/Engine/Source/Programs/Horde/HordeStorage/Horde.Storage/Implementation/LastAccess/LastAccessTracker.cs
+++ b/Engine/Source/Programs/Horde/HordeStorage/Horde.Storage/Implementation/LastAccess/LastAccessTracker.cs
@@ -6,6 +6,7 @@ using System.Collections.Generic;
 using System.Linq;
 using System.Threading;
 using System.Threading.Tasks;
+using Jupiter.Implementation;
 using Serilog;
 
 namespace Horde.Storage.Implementation
@@ -17,6 +18,29 @@ namespace Horde.Storage.Implementation
             return $"{record.Namespace}.{record.Bucket}.{record.RefName}";
         }
     }
+
+    public class LastAccessRecord
+    {
+        public LastAccessRecord(NamespaceId ns, BucketId bucket, IoHashKey key)
+        {
+            Namespace = ns;
+            Bucket = bucket;
+            Key = key;
+        }
+
+        public NamespaceId Namespace { get; set; }
+        public BucketId Bucket { get; set; }
+        public IoHashKey Key { get; set; }
+    }
+
+    public class LastAccessTrackerReference : LastAccessTracker<LastAccessRecord>
+    {
+        protected override string BuildCacheKey(LastAccessRecord record)
+        {
+            return $"{record.Namespace}.{record.Bucket}.{record.Key}";
+        }
+    }
+
     public abstract class LastAccessTracker<T> : ILastAccessTracker<T>, ILastAccessCache<T>
     {
         private readonly ILogger _logger = Log.ForContext<LastAccessTracker<T>>();
diff --git a/Engine/Source/Programs/Horde/HordeStorage/Horde.Storage/Implementation/ObjectService.cs b/Engine/Source/Programs/Horde/HordeStorage/Horde.Storage/Implementation/ObjectService.cs
--- a/Engine/Source/Programs/Horde/HordeStorage/Horde.Storage/Implementation/ObjectService.cs
+++ b/Engine/Source/Programs/Horde/HordeStorage/Horde.Storage/Implementation/ObjectService.cs
@@ -33,18 +33,23 @@ namespace Horde.Storage.Implementation
         private readonly IReferenceResolver _referenceResolver;
         private readonly IReplicationLog _replicationLog;
         private readonly IBlobIndex _blobIndex;
+        private readonly ILastAccessTracker<LastAccessRecord> _lastAccessTracker;
 
-        public ObjectService(IReferencesStore referencesStore, IBlobService blobService, IReferenceResolver referenceResolver, IReplicationLog replicationLog, IBlobIndex blobIndex)
+        public ObjectService(IReferencesStore referencesStore, IBlobService blobService, IReferenceResolver referenceResolver, IReplicationLog replicationLog, IBlobIndex blobIndex, ILastAccessTracker<LastAccessRecord> lastAccessTracker)
         {
             _referencesStore = referencesStore;
             _blobService = blobService;
             _referenceResolver = referenceResolver;
             _replicationLog = replicationLog;
             _blobIndex = blobIndex;
+            _lastAccessTracker = lastAccessTracker;
         }
 
         public async Task<(ObjectRecord, BlobContents)> Get(NamespaceId ns, BucketId bucket, IoHashKey key, string[]? fields = null)
         {
+            // we do not wait for the last access tracking as it does not matter when it completes
+            Task lastAccessTask = _lastAccessTracker.TrackUsed(new LastAccessRecord(ns, bucket, key));
+
             ObjectRecord o = await _referencesStore.Get(ns, bucket, key);
 
             BlobContents blobContents;
diff --git a/Engine/Source/Programs/Horde/HordeStorage/Horde.Storage/Tests/Functional/GC/GCRefTests.cs b/Engine/Source/Programs/Horde/HordeStorage/Horde.Storage/Tests/Functional/GC/GCRefTests.cs
--- a/Engine/Source/Programs/Horde/HordeStorage/Horde.Storage/Tests/Functional/GC/GCRefTests.cs
+++ b/Engine/Source/Programs/Horde/HordeStorage/Horde.Storage/Tests/Functional/GC/GCRefTests.cs
@@ -38,6 +38,10 @@ namespace Horde.Storage.FunctionalTests.GC
                             // we are not reading the base appSettings here as we want exact control over what runs in the tests
                             .AddJsonFile("appsettings.Testing.json", false)
                             .AddEnvironmentVariables()
+                            .AddInMemoryCollection(new List<KeyValuePair<string, string>>()
+                            {
+                                new KeyValuePair<string, string>("GC:CleanNamespacesV1:0", TestNamespace.ToString()),
+                            })
                             .Build();
 
             Logger logger = new LoggerConfiguration()
