Horde: Simpler version of scheduling service. Now has one pod queue up all work to process (and update stream data), and other pods can just deque items.

diff --git a/Engine/Source/Programs/Horde/Horde.Build/Projects/ScheduleService.cs b/Engine/Source/Programs/Horde/Horde.Build/Projects/ScheduleService.cs
--- a/Engine/Source/Programs/Horde/Horde.Build/Projects/ScheduleService.cs
+++ b/Engine/Source/Programs/Horde/Horde.Build/Projects/ScheduleService.cs
@@ -29,6 +29,7 @@ using StackExchange.Redis;
 using EpicGames.Redis;
 using EpicGames.Serialization;
 using System.Text.RegularExpressions;
+using EpicGames.Redis.Utility;
 
 namespace HordeServer.Services
 {
@@ -75,8 +76,9 @@ namespace HordeServer.Services
 		IClock Clock;
 		ILogger Logger;
 		RedisService Redis;
-		RedisSortedSet<QueueItem> Queue; // Next tick time for each queue item
-		RedisSortedSet<QueueItem> Locks; // Currently locked queue items (locks expire after a period of time)
+		RedisKey BaseLockKey;
+		RedisKey TickLockKey; // Lock to tick the queue
+		RedisSortedSet<QueueItem> Queue; // Items to tick, ordered by time
 		BackgroundTask BackgroundTask;
 
 		/// <summary>
@@ -110,15 +112,16 @@ namespace HordeServer.Services
 			this.Logger = Logger;
 
 			this.Redis = Redis;
+			this.BaseLockKey = "scheduler/locks";
+			this.TickLockKey = BaseLockKey.Append("/tick");
 			this.Queue = new RedisSortedSet<QueueItem>(Redis.Database, "scheduler/queue");
-			this.Locks = new RedisSortedSet<QueueItem>(Redis.Database, "scheduler/locks");
 			this.BackgroundTask = new BackgroundTask(RunAsync);
 		}
 
 		/// <inheritdoc/>
 		public async Task StartAsync(CancellationToken CancellationToken)
 		{
-			await RefreshSchedulesAsync(Clock.UtcNow);
+			await UpdateQueueAsync(Clock.UtcNow);
 			BackgroundTask.Start();
 		}
 
@@ -136,136 +139,78 @@ namespace HordeServer.Services
 
 		async Task RunAsync(CancellationToken CancellationToken)
 		{
-			await RefreshSchedulesAsync(Clock.UtcNow);
+			await UpdateQueueAsync(Clock.UtcNow);
 
 			for(; ;)
 			{
-				TimeSpan DelayTime = await TickAsync(CancellationToken);
-				if (CancellationToken.IsCancellationRequested)
+				Stopwatch Timer = Stopwatch.StartNew();
+				await TickAsync(CancellationToken);
+
+				TimeSpan DelayTime = TimeSpan.FromMinutes(1.0) - Timer.Elapsed;
+				if(DelayTime > TimeSpan.Zero)
 				{
-					break;
+					await Task.Delay(DelayTime, CancellationToken);
 				}
-				await Task.Delay(DelayTime, CancellationToken);
 			}
 		}
 
-		async Task<TimeSpan> TickAsync(CancellationToken CancellationToken)
+		async Task TickAsync(CancellationToken CancellationToken)
 		{
 			DateTime UtcNow = Clock.UtcNow;
 
-			// Remove any expired locks
-			await Locks.RemoveRangeByScoreAsync(-double.NegativeInfinity, QueueItem.GetScoreFromTime(UtcNow));
+			// Update the current queue
+			await using (RedisLock Lock = new RedisLock(Redis.Database, TickLockKey))
+			{
+				if (await Lock.AcquireAsync(TimeSpan.FromMinutes(1.0), false))
+				{
+					await UpdateQueueAsync(UtcNow);
+				}
+			}
 
 			// Keep updating schedules
 			while (!CancellationToken.IsCancellationRequested)
 			{
 				// Get the item with the lowest score (ie. the one that hasn't been updated in the longest time)
-				SortedSetEntry<QueueItem>[] Entries = await Queue.RangeByScoreWithScoresAsync(Take: 1);
-				if (Entries.Length == 0)
-				{
-					await RefreshSchedulesAsync(UtcNow);
-					return TimeSpan.FromSeconds(10.0);
-				}
-
-				// Deconstruct the item
-				SortedSetEntry<QueueItem> Entry = Entries[0];
-				QueueItem Item = Entry.Element;
-				DateTime ScheduledTime = QueueItem.GetTimeFromScore(Entry.Score);
-
-				// Check if the first entry is due to run
-				if (UtcNow < ScheduledTime)
-				{
-					TimeSpan Delay = (ScheduledTime - UtcNow).Add(TimeSpan.FromSeconds(1.0));
-					return Delay;
-				}
-
-				// If the stream id is null, it's a request to update the schedules
-				if (Item.StreamId == StreamId.Empty)
-				{
-					await RefreshSchedulesAsync(UtcNow);
-					continue;
-				}
-
-				// Try to create a lock on this item before we consider updating it. If it fails, just update the schedule time to a while from now.
-				double LockTime = QueueItem.GetScoreFromTime(Clock.UtcNow + TimeSpan.FromMinutes(5.0));
-				if (!await Locks.AddAsync(Item, LockTime, When.NotExists))
+				QueueItem? Item = await PopQueueItemAsync();
+				if (Item == null)
 				{
-					await Queue.AddAsync(Item, QueueItem.GetScoreFromTime(Clock.UtcNow + TimeSpan.FromSeconds(20.0)));
-					continue;
+					break;
 				}
 
-				// Spawn the task to tick the template while simultaneously extending the lock time
-				Task TickTemplateTask = Task.Run(() => TickTemplateAsync(Item, UtcNow, CancellationToken));
-				try
+				// Acquire the lock for this schedule and update it
+				await using (RedisLock Lock = new RedisLock<QueueItem>(Redis.Database, BaseLockKey, Item))
 				{
-					for (; ; )
+					if (await Lock.AcquireAsync(TimeSpan.FromMinutes(1.0)))
 					{
-						Task DelayTask = Task.Delay(TimeSpan.FromMinutes(1.0));
-						if (await Task.WhenAny(TickTemplateTask, DelayTask) == TickTemplateTask)
+						try
+						{
+							await TriggerAsync(Item.StreamId, Item.TemplateId, UtcNow, CancellationToken);
+						}
+						catch (OperationCanceledException)
+						{
+							throw;
+						}
+						catch (Exception Ex)
 						{
-							await TickTemplateTask;
-							break;
+							Logger.LogError(Ex, "Error while updating schedule for {StreamId}/{TemplateId}", Item.StreamId, Item.TemplateId);
 						}
-						await Locks.AddAsync(Item, QueueItem.GetScoreFromTime(Clock.UtcNow + TimeSpan.FromMinutes(5.0)));
 					}
 				}
-				catch (OperationCanceledException)
-				{
-					throw;
-				}
-				catch (Exception Ex)
-				{
-					Logger.LogError(Ex, "Error while updating schedule for {StreamId}/{TemplateId}", Item.StreamId, Item.TemplateId);
-				}
-
-				// Remove the lock
-				await Locks.RemoveAsync(Item);
 			}
-
-			return TimeSpan.Zero;
 		}
 
-		async Task TickTemplateAsync(QueueItem Item, DateTime UtcNow, CancellationToken CancellationToken)
+		async Task<QueueItem?> PopQueueItemAsync()
 		{
-			// Get the stream and template
-			IStream? Stream = await StreamService.GetStreamAsync(Item.StreamId);
-			if (Stream == null || !Stream.Templates.TryGetValue(Item.TemplateId, out TemplateRef? TemplateRef) || TemplateRef.Schedule == null)
-			{
-				await Queue.RemoveAsync(Item);
-				return;
-			}
-
-			// Calculate the next trigger time for this schedule; it may have changed since the queue item was created.
-			DateTime? NextUpdateTimeUtc = TemplateRef.Schedule.GetNextTriggerTimeUtc(TimeZone);
-			if (NextUpdateTimeUtc == null)
-			{
-				await Queue.RemoveAsync(Item);
-				return;
-			}
-
-			// Check if we're ready to trigger now
-			if (UtcNow < NextUpdateTimeUtc)
-			{
-				await Queue.AddAsync(Item, QueueItem.GetScoreFromTime(NextUpdateTimeUtc.Value));
-				return;
-			}
-
-			// Set the update time on this stream. If it fails, just return and try again.
-			Stream = await StreamService.UpdateScheduleTriggerAsync(Stream, Item.TemplateId, LastTriggerTimeUtc: Clock.UtcNow);
-			if (Stream != null)
+			for (; ; )
 			{
-				// Try to do the update
-				await TriggerAsync(Stream, Item.TemplateId, TemplateRef, UtcNow, CancellationToken);
-
-				// Update the next trigger time
-				DateTime? NextTime = TemplateRef.Schedule.GetNextTriggerTimeUtc(UtcNow, TimeZone);
-				if (NextTime == null)
+				QueueItem[] Items = await Queue.RangeByRankAsync(0, 0);
+				if (Items.Length == 0)
 				{
-					await Queue.RemoveAsync(Item);
+					return null;
 				}
-				else
+				if (await Queue.RemoveAsync(Items[0]))
 				{
-					await Queue.AddAsync(Item, QueueItem.GetScoreFromTime(NextTime.Value));
+					return Items[0];
 				}
 			}
 		}
@@ -273,19 +218,19 @@ namespace HordeServer.Services
 		internal async Task ResetAsync()
 		{
 			await Redis.Database.KeyDeleteAsync(Queue.Key);
-			await Redis.Database.KeyDeleteAsync(Locks.Key);
+			await Redis.Database.KeyDeleteAsync(TickLockKey);
 		}
 
 		internal async Task TickForTestingAsync()
 		{
-			await RefreshSchedulesAsync(Clock.UtcNow);
+			await UpdateQueueAsync(Clock.UtcNow);
 			await TickAsync(CancellationToken.None);
 		}
 
 		/// <summary>
 		/// Get the current set of streams and ensure there's an entry for each item
 		/// </summary>
-		public async Task RefreshSchedulesAsync(DateTime UtcNow)
+		public async Task UpdateQueueAsync(DateTime UtcNow)
 		{
 			List<SortedSetEntry<QueueItem>> QueueItems = new List<SortedSetEntry<QueueItem>>();
 
@@ -299,13 +244,17 @@ namespace HordeServer.Services
 						DateTime? NextTriggerTimeUtc = TemplateRef.Schedule.GetNextTriggerTimeUtc(TimeZone);
 						if (NextTriggerTimeUtc != null)
 						{
-							double Score = QueueItem.GetScoreFromTime(NextTriggerTimeUtc.Value);
-							QueueItems.Add(new SortedSetEntry<QueueItem>(new QueueItem(Stream.Id, TemplateId), Score));
+							if (UtcNow > NextTriggerTimeUtc.Value)
+							{
+								double Score = QueueItem.GetScoreFromTime(NextTriggerTimeUtc.Value);
+								QueueItems.Add(new SortedSetEntry<QueueItem>(new QueueItem(Stream.Id, TemplateId), Score));
+
+								await StreamService.UpdateScheduleTriggerAsync(Stream, TemplateId, UtcNow);
+							}
 						}
 					}
 				}
 			}
-			QueueItems.Add(new SortedSetEntry<QueueItem>(new QueueItem(StreamId.Empty, TemplateRefId.Empty), QueueItem.GetScoreFromTime(UtcNow.AddMinutes(1.0))));
 
 			await Queue.AddAsync(QueueItems.ToArray());
 		}
@@ -313,14 +262,25 @@ namespace HordeServer.Services
 		/// <summary>
 		/// Trigger a schedule to run
 		/// </summary>
-		/// <param name="Stream">Stream for the schedule</param>
+		/// <param name="StreamId">Stream for the schedule</param>
 		/// <param name="TemplateId"></param>
-		/// <param name="TemplateRef"></param>
 		/// <param name="UtcNow"></param>
 		/// <param name="CancellationToken"></param>
 		/// <returns>Async task</returns>
-		private async Task TriggerAsync(IStream Stream, TemplateRefId TemplateId, TemplateRef TemplateRef, DateTime UtcNow, CancellationToken CancellationToken)
+		private async Task<bool> TriggerAsync(StreamId StreamId, TemplateRefId TemplateId, DateTime UtcNow, CancellationToken CancellationToken)
 		{
+			IStream? Stream = await StreamService.GetStreamAsync(StreamId);
+			if (Stream == null || !Stream.Templates.TryGetValue(TemplateId, out TemplateRef? TemplateRef))
+			{
+				return false;
+			}
+
+			Schedule? Schedule = TemplateRef.Schedule;
+			if (Schedule == null)
+			{
+				return false;
+			}
+
 			using IScope Scope = GlobalTracer.Instance.BuildSpan("ScheduleService.TriggerAsync").StartActive();
 			Scope.Span.SetTag("StreamId", Stream.Id);
 			Scope.Span.SetTag("TemplateId", TemplateId);
@@ -328,13 +288,6 @@ namespace HordeServer.Services
 			Stopwatch Stopwatch = Stopwatch.StartNew();
 			Logger.LogInformation("Updating schedule for {StreamId} template {TemplateId}", Stream.Id, TemplateId);
 
-			// Make sure the schedule is valid
-			Schedule? Schedule = TemplateRef.Schedule;
-			if(Schedule == null)
-			{
-				return;
-			}
-
 			// Get a list of jobs that we need to remove
 			List<JobId> RemoveJobIds = new List<JobId>();
 			foreach (JobId ActiveJobId in Schedule.ActiveJobs)
@@ -352,7 +305,7 @@ namespace HordeServer.Services
 			if (Stream.IsPaused(UtcNow))
 			{
 				Logger.LogDebug("Skipping schedule update for stream {StreamId}. It has been paused until {PausedUntil} with comment '{PauseComment}'.", Stream.Id, Stream.PausedUntil, Stream.PauseComment);
-				return;
+				return false;
 			}
 
 			// Trigger this schedule
@@ -368,6 +321,7 @@ namespace HordeServer.Services
 			// Print some timing info
 			Stopwatch.Stop();
 			Logger.LogInformation("Updated schedule for {StreamId} template {TemplateId} in {TimeSeconds}ms", Stream.Id, TemplateId, (long)Stopwatch.Elapsed.TotalMilliseconds);
+			return true;
 		}
 
 		/// <summary>
diff --git a/Engine/Source/Programs/Shared/EpicGames.Redis/Utility/RedisLock.cs b/Engine/Source/Programs/Shared/EpicGames.Redis/Utility/RedisLock.cs
--- a/Engine/Source/Programs/Shared/EpicGames.Redis/Utility/RedisLock.cs
+++ b/Engine/Source/Programs/Shared/EpicGames.Redis/Utility/RedisLock.cs
@@ -3,6 +3,7 @@
 using StackExchange.Redis;
 using System;
 using System.Collections.Generic;
+using System.Diagnostics;
 using System.Text;
 using System.Threading;
 using System.Threading.Tasks;
@@ -52,13 +53,14 @@ namespace EpicGames.Redis.Utility
 		/// Attempts to acquire the lock for the given period of time. The lock will be renewed once half of this interval has elapsed.
 		/// </summary>
 		/// <param name="Duration">Time after which the lock expires</param>
+		/// <param name="ReleaseOnDispose">Whether the lock should be released when disposed. If false, the lock will be held for the given time, but not renewed once disposed.</param>
 		/// <returns>True if the lock was acquired, false if another service already has it</returns>
-		public async ValueTask<bool> AcquireAsync(TimeSpan Duration)
+		public async ValueTask<bool> AcquireAsync(TimeSpan Duration, bool ReleaseOnDispose = true)
 		{
 			if (await Database.StringSetAsync(Key, RedisValue.EmptyString, Duration, When.NotExists))
 			{
 				CancellationSource = new CancellationTokenSource();
-				BackgroundTask = Task.Run(() => RenewAsync(Duration, CancellationSource.Token));
+				BackgroundTask = Task.Run(() => RenewAsync(Duration, ReleaseOnDispose, CancellationSource.Token));
 				return true;
 			}
 			return false;
@@ -68,16 +70,26 @@ namespace EpicGames.Redis.Utility
 		/// Background task which renews the lock while the service is running
 		/// </summary>
 		/// <param name="Duration"></param>
+		/// <param name="ReleaseOnDispose">Whether the lock should be released when the cancellation token is fired</param>
 		/// <param name="CancellationToken"></param>
 		/// <returns></returns>
-		async Task RenewAsync(TimeSpan Duration, CancellationToken CancellationToken)
+		async Task RenewAsync(TimeSpan Duration, bool ReleaseOnDispose, CancellationToken CancellationToken)
 		{
+			Stopwatch Timer = Stopwatch.StartNew();
 			for (; ; )
 			{
 				await Task.Delay(Duration / 2, CancellationToken).ContinueWith(x => { }); // Do not throw
 				if (CancellationToken.IsCancellationRequested)
 				{
-					await Database.StringSetAsync(Key, RedisValue.Null);
+					Timer.Stop();
+					if (ReleaseOnDispose || Timer.Elapsed > Duration)
+					{
+						await Database.StringSetAsync(Key, RedisValue.Null);
+					}
+					else
+					{
+						await Database.StringSetAsync(Key, RedisValue.EmptyString, Duration - Timer.Elapsed, When.Exists);
+					}
 					break;
 				}
 				if (!await Database.StringSetAsync(Key, RedisValue.EmptyString, Duration, When.Exists))
