Horde Storage
Added dd scope for last access tracking to see that it runs and how long it takes

diff --git a/Engine/Source/Programs/Horde/HordeStorage/Horde.Storage/Implementation/LastAccess/LastAccessTracker.cs b/Engine/Source/Programs/Horde/HordeStorage/Horde.Storage/Implementation/LastAccess/LastAccessTracker.cs
--- a/Engine/Source/Programs/Horde/HordeStorage/Horde.Storage/Implementation/LastAccess/LastAccessTracker.cs
+++ b/Engine/Source/Programs/Horde/HordeStorage/Horde.Storage/Implementation/LastAccess/LastAccessTracker.cs
@@ -6,6 +6,7 @@ using System.Collections.Generic;
 using System.Linq;
 using System.Threading;
 using System.Threading.Tasks;
+using Datadog.Trace;
 using Jupiter.Implementation;
 using Serilog;
 
@@ -57,6 +58,7 @@ namespace Horde.Storage.Implementation
         {
             return Task.Run(() =>
             {
+                using Scope _ = Tracer.Instance.StartActive("lastAccessTracker.track");
                 try
                 {
                     _rwLock.AcquireReaderLock(-1);
