Horde Storage

Blob cleanup health checks are now registred if we expect the service to exist, as it will exist for filesystem cleanup even if the orphan blob clean is disabled.

Also added running of the LastAccessServiceReferences background service

diff --git a/Engine/Source/Programs/Horde/HordeStorage/Horde.Storage/Horde.Storage.Startup.cs b/Engine/Source/Programs/Horde/HordeStorage/Horde.Storage/Horde.Storage.Startup.cs
--- a/Engine/Source/Programs/Horde/HordeStorage/Horde.Storage/Horde.Storage.Startup.cs
+++ b/Engine/Source/Programs/Horde/HordeStorage/Horde.Storage/Horde.Storage.Startup.cs
@@ -246,6 +246,9 @@ namespace Horde.Storage
             services.AddSingleton<LastAccessService>();
             services.AddHostedService<LastAccessService>(p => p.GetService<LastAccessService>()!);
 
+            services.AddSingleton<LastAccessServiceReferences>();
+            services.AddHostedService<LastAccessServiceReferences>(p => p.GetService<LastAccessServiceReferences>()!);
+
             services.AddSingleton<IServiceCredentials, ServiceCredentials>(p => ActivatorUtilities.CreateInstance<ServiceCredentials>(p));
             
             services.AddSingleton<ReplicationService>();
@@ -635,7 +638,7 @@ namespace Horde.Storage
                 healthChecks.AddCheck<RefCleanupServiceCheck>("RefCleanupCheck", tags: new[] {"services"});
             }
 
-            if (gcSettings.CleanOldBlobs)
+            if (gcSettings.BlobCleanupServiceEnabled)
             {
                 healthChecks.AddCheck<BlobCleanupServiceCheck>("BlobStoreCheck", tags: new[] {"services"});
             }
