Horde: Additional logging for scheduler.

diff --git a/Engine/Source/Programs/Horde/Horde.Build/Projects/ScheduleService.cs b/Engine/Source/Programs/Horde/Horde.Build/Projects/ScheduleService.cs
--- a/Engine/Source/Programs/Horde/Horde.Build/Projects/ScheduleService.cs
+++ b/Engine/Source/Programs/Horde/Horde.Build/Projects/ScheduleService.cs
@@ -458,12 +458,13 @@ namespace HordeServer.Services
 			}
 
 			// Try to start all the new jobs
+			Logger.LogInformation("Starting {NumJobs} new jobs for {StreamId} template {TemplateId} (active: {NumActive}, max new: {MaxNewJobs})", TriggerChanges.Count, Stream.Id, TemplateId, NumActiveJobs, MaxNewChanges);
 			foreach ((int Change, int CodeChange) in TriggerChanges.OrderBy(x => x.Change))
 			{
 				CancellationToken.ThrowIfCancellationRequested();
 				List<string> DefaultArguments = Template.GetDefaultArguments();
 				IJob NewJob = await JobService.CreateJobAsync(null, Stream, TemplateId, Template.Id, Graph, Template.Name, Change, CodeChange, null, null, null, Template.Priority, null, null, TemplateRef.ChainedJobs, TemplateRef.ShowUgsBadges, TemplateRef.ShowUgsAlerts, TemplateRef.NotificationChannel, TemplateRef.NotificationChannelFilter, DefaultArguments);
-				Logger.LogInformation("Started new job for {StreamName} template {TemplateId} at CL {Change} (Code CL {CodeChange}): {JobId}", Stream.Id, Template.Id, Change, CodeChange, NewJob.Id);
+				Logger.LogInformation("Started new job for {StreamId} template {TemplateId} at CL {Change} (Code CL {CodeChange}): {JobId}", Stream.Id, TemplateId, Change, CodeChange, NewJob.Id);
 				await StreamService.UpdateScheduleTriggerAsync(Stream, TemplateId, UtcNow, Change, new List<JobId> { NewJob.Id }, new List<JobId>());
 			}
 		}
