Horde: Ignore server errors that are caused by exceeding gRPC deadlines or client terminations.

diff --git a/Engine/Source/Programs/Horde/Horde.Build/Services/RpcService.cs b/Engine/Source/Programs/Horde/Horde.Build/Services/RpcService.cs
--- a/Engine/Source/Programs/Horde/Horde.Build/Services/RpcService.cs
+++ b/Engine/Source/Programs/Horde/Horde.Build/Services/RpcService.cs
@@ -396,7 +396,7 @@ namespace HordeServer.Services
 				}
 
 				// Get a task for moving to the next item. This will only complete once the call has closed.
-				using CancellationTokenSource CancellationSource = new CancellationTokenSource();
+				using CancellationTokenSource CancellationSource = CancellationTokenSource.CreateLinkedTokenSource(Context.CancellationToken);
 				NextRequestTask = Reader.MoveNext();
 				NextRequestTask = NextRequestTask.ContinueWith(Task => { CancellationSource.Cancel(); return Task.Result; }, TaskScheduler.Current);
 
@@ -429,10 +429,13 @@ namespace HordeServer.Services
 				}
 
 				// Create the new session info
-				UpdateSessionResponse Response = new UpdateSessionResponse();
-				Response.Leases.Add(Agent.Leases.Select(x => x.ToRpcMessage()));
-				Response.ExpiryTime = (Agent.SessionExpiresAt == null) ? new Timestamp() : Timestamp.FromDateTime(Agent.SessionExpiresAt.Value);
-				await Writer.WriteAsync(Response);
+				if (!Context.CancellationToken.IsCancellationRequested)
+				{
+					UpdateSessionResponse Response = new UpdateSessionResponse();
+					Response.Leases.Add(Agent.Leases.Select(x => x.ToRpcMessage()));
+					Response.ExpiryTime = (Agent.SessionExpiresAt == null) ? new Timestamp() : Timestamp.FromDateTime(Agent.SessionExpiresAt.Value);
+					await Writer.WriteAsync(Response);
+				}
 
 				// Wait for the client to close the stream
 				while (await NextRequestTask)
diff --git a/Engine/Source/Programs/Horde/Horde.Build/Startup.cs b/Engine/Source/Programs/Horde/Horde.Build/Startup.cs
--- a/Engine/Source/Programs/Horde/Horde.Build/Startup.cs
+++ b/Engine/Source/Programs/Horde/Horde.Build/Startup.cs
@@ -109,49 +109,49 @@ namespace HordeServer
 			[SuppressMessage("Design", "CA1031:Do not catch general exception types", Justification = "<Pending>")]
 			public override Task<TResponse> UnaryServerHandler<TRequest, TResponse>(TRequest Request, ServerCallContext Context, UnaryServerMethod<TRequest, TResponse> Continuation)
 			{
-				return Guard(Context, Context.Method, () => base.UnaryServerHandler(Request, Context, Continuation));
+				return Guard(Context, () => base.UnaryServerHandler(Request, Context, Continuation));
 			}
 
 			public override Task<TResponse> ClientStreamingServerHandler<TRequest, TResponse>(IAsyncStreamReader<TRequest> RequestStream, ServerCallContext Context, ClientStreamingServerMethod<TRequest, TResponse> Continuation) where TRequest : class where TResponse : class
 			{
-				return Guard(Context, Context.Method, () => base.ClientStreamingServerHandler(RequestStream, Context, Continuation));
+				return Guard(Context, () => base.ClientStreamingServerHandler(RequestStream, Context, Continuation));
 			}
 
 			public override Task ServerStreamingServerHandler<TRequest, TResponse>(TRequest Request, IServerStreamWriter<TResponse> ResponseStream, ServerCallContext Context, ServerStreamingServerMethod<TRequest, TResponse> Continuation) where TRequest : class where TResponse : class
 			{
-				return Guard(Context, Context.Method, () => base.ServerStreamingServerHandler(Request, ResponseStream, Context, Continuation));
+				return Guard(Context, () => base.ServerStreamingServerHandler(Request, ResponseStream, Context, Continuation));
 			}
 
 			public override Task DuplexStreamingServerHandler<TRequest, TResponse>(IAsyncStreamReader<TRequest> RequestStream, IServerStreamWriter<TResponse> ResponseStream, ServerCallContext Context, DuplexStreamingServerMethod<TRequest, TResponse> Continuation) where TRequest : class where TResponse : class
 			{
-				return Guard(Context, Context.Method, () => base.DuplexStreamingServerHandler(RequestStream, ResponseStream, Context, Continuation));
+				return Guard(Context, () => base.DuplexStreamingServerHandler(RequestStream, ResponseStream, Context, Continuation));
 			}
 
-			async Task<T> Guard<T>(ServerCallContext Context, string Method, Func<Task<T>> CallFunc) where T : class
+			async Task<T> Guard<T>(ServerCallContext Context, Func<Task<T>> CallFunc) where T : class
 			{
 				T Result = null!;
-				await Guard(Context, Method, async () => { Result = await CallFunc(); });
+				await Guard(Context, async () => { Result = await CallFunc(); });
 				return Result;
 			}
 
-			async Task Guard(ServerCallContext Context, string Method, Func<Task> CallFunc)
+			async Task Guard(ServerCallContext Context, Func<Task> CallFunc)
 			{
 				HttpContext HttpContext = Context.GetHttpContext();
 
 				AgentId? AgentId = AclService.GetAgentId(HttpContext.User);
 				if (AgentId != null)
 				{
-					using IDisposable Scope = Logger.BeginScope("Agent: {AgentId}, RemoteIP: {RemoteIP}, Method: {Method}", AgentId.Value, HttpContext.Connection.RemoteIpAddress, Method);
-					await GuardInner(Method, CallFunc);
+					using IDisposable Scope = Logger.BeginScope("Agent: {AgentId}, RemoteIP: {RemoteIP}, Method: {Method}", AgentId.Value, HttpContext.Connection.RemoteIpAddress, Context.Method);
+					await GuardInner(Context, CallFunc);
 				}
 				else
 				{
-					using IDisposable Scope = Logger.BeginScope("RemoteIP: {RemoteIP}, Method: {Method}", HttpContext.Connection.RemoteIpAddress, Method);
-					await GuardInner(Method, CallFunc);
+					using IDisposable Scope = Logger.BeginScope("RemoteIP: {RemoteIP}, Method: {Method}", HttpContext.Connection.RemoteIpAddress, Context.Method);
+					await GuardInner(Context, CallFunc);
 				}
 			}
 
-			async Task GuardInner(string Method, Func<Task> CallFunc)
+			async Task GuardInner(ServerCallContext Context, Func<Task> CallFunc)
 			{
 				try
 				{
@@ -164,8 +164,16 @@ namespace HordeServer
 				}
 				catch (Exception Ex)
 				{
-					Logger.LogError(Ex, "Exception in call to {Method}", Method);
-					throw new RpcException(new Status(StatusCode.Internal, $"An exception was thrown on the server: {Ex}"));
+					if (Context.CancellationToken.IsCancellationRequested)
+					{
+						Logger.LogInformation(Ex, "Call to method {Method} was cancelled", Context.Method);
+						throw;
+					}
+					else
+					{
+						Logger.LogError(Ex, "Exception in call to {Method}", Context.Method);
+						throw new RpcException(new Status(StatusCode.Internal, $"An exception was thrown on the server: {Ex}"));
+					}
 				}
 			}
 		}
