Horde: Reduce number of P4 calls required to query recent changes to build when considering multiple changes in a single run.

diff --git a/Engine/Source/Programs/Horde/Horde.Build/Projects/ScheduleService.cs b/Engine/Source/Programs/Horde/Horde.Build/Projects/ScheduleService.cs
--- a/Engine/Source/Programs/Horde/Horde.Build/Projects/ScheduleService.cs
+++ b/Engine/Source/Programs/Horde/Horde.Build/Projects/ScheduleService.cs
@@ -65,6 +65,76 @@ namespace HordeServer.Services
 			public static double GetScoreFromTime(DateTime Time) => (Time.ToUniversalTime() - DateTime.UnixEpoch).TotalSeconds;
 		}
 
+		class PerforceHistory
+		{
+			IPerforceService Perforce;
+			string ClusterName;
+			string StreamName;
+			int MaxResults;
+			List<ChangeDetails> Changes = new List<ChangeDetails>();
+			int NextIndex;
+
+			public PerforceHistory(IPerforceService Perforce, string ClusterName, string StreamName)
+			{
+				this.Perforce = Perforce;
+				this.ClusterName = ClusterName;
+				this.StreamName = StreamName;
+				this.MaxResults = 10;
+			}
+
+			async ValueTask<ChangeDetails?> GetChangeAtIndexAsync(int Index)
+			{
+				while (Index >= Changes.Count)
+				{
+					List<ChangeSummary> NewChanges = await Perforce.GetChangesAsync(ClusterName, null, null, MaxResults);
+					if (Changes.Count > 0)
+					{
+						NewChanges.RemoveAll(x => x.Number >= Changes[Changes.Count - 1].Number);
+					}
+					if (NewChanges.Count > 0)
+					{
+						List<ChangeDetails> NewChangeDetails = await Perforce.GetChangeDetailsAsync(ClusterName, StreamName, NewChanges.ConvertAll(x => x.Number), null);
+						Changes.AddRange(NewChangeDetails.OrderByDescending(x => x.Number));
+					}
+					MaxResults += 10;
+				}
+				return Changes[Index];
+			}
+
+			public async ValueTask<ChangeDetails?> GetNextChangeAsync()
+			{
+				ChangeDetails? Details = await GetChangeAtIndexAsync(NextIndex);
+				if (Details != null)
+				{
+					NextIndex++;
+				}
+				return Details;
+			}
+
+			public async ValueTask<int> GetCodeChange(int Change)
+			{
+				int Index = Changes.BinarySearch(x => -x.Number, -Change);
+				if (Index < 0)
+				{
+					Index = ~Index;
+				}
+
+				for (; ; )
+				{
+					ChangeDetails? Details = await GetChangeAtIndexAsync(Index);
+					if (Details == null)
+					{
+						return 0;
+					}
+					if ((Details.GetContentFlags() & ChangeContentFlags.ContainsCode) != 0)
+					{
+						return Details.Number;
+					}
+					Index++;
+				}
+			}
+		}
+
 		IDowntimeService DowntimeService;
 		IGraphCollection Graphs;
 		IPerforceService Perforce;
@@ -380,6 +450,9 @@ namespace HordeServer.Services
 				FileFilter = new FileFilter(Schedule.Files);
 			}
 
+			// Cache the Perforce history as we're iterating through changes to improve query performance
+			PerforceHistory History = new PerforceHistory(Perforce, Stream.ClusterName, Stream.Name);
+			
 			// Start as many jobs as possible
 			List<(int Change, int CodeChange)> TriggerChanges = new List<(int, int)>();
 			while (TriggerChanges.Count < MaxNewChanges)
@@ -387,29 +460,40 @@ namespace HordeServer.Services
 				CancellationToken.ThrowIfCancellationRequested();
 
 				// Get the next valid change
-				int Change, CodeChange;
+				ChangeDetails? ChangeDetails = null;
 				if (Schedule.Gate != null)
 				{
-					(Change, CodeChange) = await GetNextChangeForGateAsync(Stream.Id, TemplateId, Schedule.Gate, MinChangeNumber, MaxChangeNumber, CancellationToken);
+					ChangeDetails = await GetNextChangeForGateAsync(Stream, TemplateId, Schedule.Gate, MinChangeNumber, MaxChangeNumber, CancellationToken);
 				}
 				else
 				{
-					(Change, CodeChange) = await GetNextChangeAsync(Stream.ClusterName, Stream.Name, MinChangeNumber, MaxChangeNumber);
+					ChangeDetails = await History.GetNextChangeAsync();
 				}
 
 				// Quit if we didn't find anything
-				if (Change < MinChangeNumber)
+				if (ChangeDetails == null)
 				{
 					break;
 				}
-				if (Change == MinChangeNumber && (Schedule.RequireSubmittedChange || TriggerChanges.Count > 0))
+				if (ChangeDetails.Number < MinChangeNumber)
+				{
+					break;
+				}
+				if (ChangeDetails.Number == MinChangeNumber && (Schedule.RequireSubmittedChange || TriggerChanges.Count > 0))
 				{
 					break;
 				}
 
 				// Adjust the changelist for the desired filter
-				if (await ShouldBuildChangeAsync(Stream.ClusterName, Stream.Name, Change, FilterFlags, FileFilter))
+				int Change = ChangeDetails.Number;
+				if (ShouldBuildChangeAsync(ChangeDetails, FilterFlags, FileFilter))
 				{
+					int CodeChange = await History.GetCodeChange(Change);
+					if (CodeChange == -1)
+					{
+						Logger.LogWarning("Unable to find code change for CL {Change}", Change);
+						CodeChange = Change;
+					}
 					TriggerChanges.Add((Change, CodeChange));
 				}
 
@@ -419,16 +503,6 @@ namespace HordeServer.Services
 					Logger.LogError("Querying for changes to trigger has taken {Time}. Aborting.", Timer.Elapsed);
 					break;
 				}
-
-				// Start the next build before this change
-				if (FilterFlags != null && (FilterFlags.Value & ChangeContentFlags.ContainsContent) == 0)
-				{
-					MaxChangeNumber = Math.Min(CodeChange, Change - 1);
-				}
-				else
-				{
-					MaxChangeNumber = Change - 1;
-				}
 			}
 
 			// Early out if there's nothing to do
@@ -472,15 +546,12 @@ namespace HordeServer.Services
 		/// <summary>
 		/// Tests whether a schedule should build a particular change, based on its requested change filters
 		/// </summary>
-		/// <param name="ClusterName"></param>
-		/// <param name="StreamName"></param>
-		/// <param name="Change"></param>
+		/// <param name="Details">The change details</param>
 		/// <param name="FilterFlags"></param>
 		/// <param name="FileFilter">Filter for the files to trigger a build</param>
 		/// <returns></returns>
-		private async Task<bool> ShouldBuildChangeAsync(string ClusterName, string StreamName, int Change, ChangeContentFlags? FilterFlags, FileFilter? FileFilter)
+		private bool ShouldBuildChangeAsync(ChangeDetails Details, ChangeContentFlags? FilterFlags, FileFilter? FileFilter)
 		{
-			ChangeDetails Details = await Perforce.GetChangeDetailsAsync(ClusterName, StreamName, Change, null);
 			if (Regex.IsMatch(Details.Description, @"^\s*#\s*skipci", RegexOptions.Multiline))
 			{
 				return false;
@@ -489,7 +560,7 @@ namespace HordeServer.Services
 			{
 				if ((Details.GetContentFlags() & FilterFlags.Value) == 0)
 				{
-					Logger.LogDebug("Not building change {Change} ({ChangeFlags}) due to filter flags ({FilterFlags})", Change, Details.GetContentFlags().ToString(), FilterFlags.Value.ToString());
+					Logger.LogDebug("Not building change {Change} ({ChangeFlags}) due to filter flags ({FilterFlags})", Details.Number, Details.GetContentFlags().ToString(), FilterFlags.Value.ToString());
 					return false;
 				}
 			}
@@ -497,7 +568,7 @@ namespace HordeServer.Services
 			{
 				if (!Details.Files.Any(x => FileFilter.Matches(x.Path)))
 				{
-					Logger.LogDebug("Not building change {Change} due to file filter", Change);
+					Logger.LogDebug("Not building change {Change} due to file filter", Details.Number);
 					return false;
 				}
 			}
@@ -524,16 +595,16 @@ namespace HordeServer.Services
 		/// Gets the next change to build for a schedule on a gate
 		/// </summary>
 		/// <returns></returns>
-		private async Task<(int Change, int CodeChange)> GetNextChangeForGateAsync(StreamId StreamId, TemplateRefId TemplateRefId, ScheduleGate Gate, int? MinChange, int? MaxChange, CancellationToken CancellationToken)
+		private async Task<ChangeDetails?> GetNextChangeForGateAsync(IStream Stream, TemplateRefId TemplateRefId, ScheduleGate Gate, int? MinChange, int? MaxChange, CancellationToken CancellationToken)
 		{
 			for (; ; )
 			{
 				CancellationToken.ThrowIfCancellationRequested();
 
-				List<IJob> Jobs = await JobCollection.FindAsync(StreamId: StreamId, Templates: new[] { Gate.TemplateRefId }, MinChange: MinChange, MaxChange: MaxChange, Count: 1);
+				List<IJob> Jobs = await JobCollection.FindAsync(StreamId: Stream.Id, Templates: new[] { Gate.TemplateRefId }, MinChange: MinChange, MaxChange: MaxChange, Count: 1);
 				if (Jobs.Count == 0)
 				{
-					return default;
+					return null;
 				}
 
 				IJob Job = Jobs[0];
@@ -547,9 +618,9 @@ namespace HordeServer.Services
 						JobStepOutcome Outcome = State.Value.Item2;
 						if (Outcome == JobStepOutcome.Success || Outcome == JobStepOutcome.Warnings)
 						{
-							return (Job.Change, Job.CodeChange);
+							return await Perforce.GetChangeDetailsAsync(Stream.ClusterName, Stream.Name, Job.Change);
 						}
-						Logger.LogInformation("Skipping trigger of {StreamName} template {TemplateId} - last {OtherTemplateRefId} job ({JobId}) ended with errors", StreamId, TemplateRefId, Gate.TemplateRefId, Job.Id);
+						Logger.LogInformation("Skipping trigger of {StreamName} template {TemplateId} - last {OtherTemplateRefId} job ({JobId}) ended with errors", Stream.Id, TemplateRefId, Gate.TemplateRefId, Job.Id);
 					}
 				}
 
