Horde: Fix for failing tests.

diff --git a/Engine/Source/Programs/Horde/Horde.Build/Projects/ScheduleService.cs b/Engine/Source/Programs/Horde/Horde.Build/Projects/ScheduleService.cs
--- a/Engine/Source/Programs/Horde/Horde.Build/Projects/ScheduleService.cs
+++ b/Engine/Source/Programs/Horde/Horde.Build/Projects/ScheduleService.cs
@@ -86,15 +86,20 @@ namespace HordeServer.Services
 			{
 				while (Index >= Changes.Count)
 				{
-					List<ChangeSummary> NewChanges = await Perforce.GetChangesAsync(ClusterName, null, null, MaxResults);
+					List<ChangeSummary> NewChanges = await Perforce.GetChangesAsync(ClusterName, StreamName, null, null, MaxResults, null);
+
+					int NumResults = NewChanges.Count;
 					if (Changes.Count > 0)
 					{
 						NewChanges.RemoveAll(x => x.Number >= Changes[Changes.Count - 1].Number);
 					}
-					if (NewChanges.Count > 0)
+					if (NewChanges.Count == 0 && NumResults < MaxResults)
+					{
+						return null;
+					}
+					if(NewChanges.Count > 0)
 					{
-						List<ChangeDetails> NewChangeDetails = await Perforce.GetChangeDetailsAsync(ClusterName, StreamName, NewChanges.ConvertAll(x => x.Number), null);
-						Changes.AddRange(NewChangeDetails.OrderByDescending(x => x.Number));
+						Changes.AddRange((await Perforce.GetChangeDetailsAsync(ClusterName, StreamName, NewChanges.ConvertAll(x => x.Number), null)).OrderByDescending(x => x.Number));
 					}
 					MaxResults += 10;
 				}
