Horde: Don't sync the AutoSDK workspace if no files have been submitted.

diff --git a/Engine/Source/Programs/Horde/Horde.Agent/Execution/PerforceExecutor.cs b/Engine/Source/Programs/Horde/Horde.Agent/Execution/PerforceExecutor.cs
--- a/Engine/Source/Programs/Horde/Horde.Agent/Execution/PerforceExecutor.cs
+++ b/Engine/Source/Programs/Horde/Horde.Agent/Execution/PerforceExecutor.cs
@@ -80,9 +80,19 @@ namespace HordeAgent.Execution
 
 					int AutoSdkChangeNumber = await AutoSdkWorkspace.GetLatestChangeAsync(CancellationToken);
 
-					FileReference AutoSdkCacheFile = FileReference.Combine(AutoSdkWorkspace.MetadataDir, "Contents.dat");
-					await AutoSdkWorkspace.UpdateLocalCacheMarker(AutoSdkCacheFile, AutoSdkChangeNumber, -1);
-					await AutoSdkWorkspace.SyncAsync(AutoSdkChangeNumber, -1, AutoSdkCacheFile, Logger, CancellationToken);
+					string SyncText = $"Synced to CL {AutoSdkChangeNumber}";
+
+					FileReference SyncFile = FileReference.Combine(AutoSdkWorkspace.MetadataDir, "Synced.txt");
+					if (!FileReference.Exists(SyncFile) || FileReference.ReadAllText(SyncFile) != SyncText)
+					{
+						FileReference.Delete(SyncFile);
+
+						FileReference AutoSdkCacheFile = FileReference.Combine(AutoSdkWorkspace.MetadataDir, "Contents.dat");
+						await AutoSdkWorkspace.UpdateLocalCacheMarker(AutoSdkCacheFile, AutoSdkChangeNumber, -1);
+						await AutoSdkWorkspace.SyncAsync(AutoSdkChangeNumber, -1, AutoSdkCacheFile, Logger, CancellationToken);
+
+						FileReference.WriteAllText(SyncFile, SyncText);
+					}
 				}
 			}
 
