Horde: Add environment variables indicating the last built changelist (UE_HORDE_LAST_CL), last successfully build changelist (UE_HORDE_LAST_SUCCESS_CL) and last changelist built that succeeded or completed with warnings (UE_HORDE_LAST_WARNING_CL). Being able to query this information within builds allows us to modify behavior accordingly (eg. to build only things that have changed since the last green build).

diff --git a/Engine/Source/Programs/Horde/Horde.Agent/Execution/BuildGraphExecutor.cs b/Engine/Source/Programs/Horde/Horde.Agent/Execution/BuildGraphExecutor.cs
--- a/Engine/Source/Programs/Horde/Horde.Agent/Execution/BuildGraphExecutor.cs
+++ b/Engine/Source/Programs/Horde/Horde.Agent/Execution/BuildGraphExecutor.cs
@@ -595,6 +595,16 @@ namespace HordeAgent.Execution
 				}
 			}
 
+			if (Step.EnvVars != null && Step.EnvVars.Count > 0)
+			{
+				Dictionary<string, string> NewEnvVars = new Dictionary<string, string>(EnvVars);
+				foreach (KeyValuePair<string, string> EnvVar in Step.EnvVars)
+				{
+					NewEnvVars[EnvVar.Key] = EnvVar.Value;
+				}
+				EnvVars = NewEnvVars;
+			}
+
 			return await ExecuteAutomationToolAsync(Step, WorkspaceDir, Arguments.ToString(), EnvVars, Step.Credentials, Logger, CancellationToken) == 0;
 		}
 
diff --git a/Engine/Source/Programs/Horde/Horde.Build.Tests/IssueServiceTests.cs b/Engine/Source/Programs/Horde/Horde.Build.Tests/IssueServiceTests.cs
--- a/Engine/Source/Programs/Horde/Horde.Build.Tests/IssueServiceTests.cs
+++ b/Engine/Source/Programs/Horde/Horde.Build.Tests/IssueServiceTests.cs
@@ -230,7 +230,7 @@ namespace HordeServerTests
 
 			JobStepRefId JobStepRefId = new JobStepRefId(Job.Id, Batch.Id, Step.Id);
 			string NodeName = Graph.Groups[Batch.GroupIdx].Nodes[Step.NodeIdx].Name;
-			await JobStepRefCollection.InsertOrReplaceAsync(JobStepRefId, "TestJob", NodeName, Job.StreamId, Job.TemplateId, Job.Change, Step.LogId, null, null, Outcome, 0.0f, 0.0f, Step.StartTimeUtc!.Value, Step.StartTimeUtc);
+			await JobStepRefCollection.InsertOrReplaceAsync(JobStepRefId, "TestJob", NodeName, Job.StreamId, Job.TemplateId, Job.Change, Step.LogId, null, null, Outcome, null, null, 0.0f, 0.0f, Step.StartTimeUtc!.Value, Step.StartTimeUtc);
 		}
 
 		async Task AddEvent(IJob Job, int BatchIdx, int StepIdx, object Data, EventSeverity Severity = EventSeverity.Error)
diff --git a/Engine/Source/Programs/Horde/Horde.Build.Tests/Stubs/Collections/JobStepRefCollectionStub.cs b/Engine/Source/Programs/Horde/Horde.Build.Tests/Stubs/Collections/JobStepRefCollectionStub.cs
--- a/Engine/Source/Programs/Horde/Horde.Build.Tests/Stubs/Collections/JobStepRefCollectionStub.cs
+++ b/Engine/Source/Programs/Horde/Horde.Build.Tests/Stubs/Collections/JobStepRefCollectionStub.cs
@@ -32,6 +32,8 @@ namespace HordeServerTests.Stubs.Collections
 		public PoolId? PoolId { get; set; }
 		public AgentId? AgentId { get; set; }
 		public JobStepOutcome? Outcome { get; set; }
+		public int? LastSuccess { get; set; }
+		public int? LastWarning { get; set; }
 
 		public virtual float BatchWaitTime => throw new NotImplementedException();
 		public virtual float BatchInitTime => throw new NotImplementedException();
@@ -96,7 +98,7 @@ namespace HordeServerTests.Stubs.Collections
 			throw new NotImplementedException();
 		}
 
-		Task<IJobStepRef> IJobStepRefCollection.InsertOrReplaceAsync(JobStepRefId Id, string JobName, string NodeName, StreamId StreamId, TemplateRefId TemplateId, int Change, LogId? LogId, PoolId? PoolId, AgentId? AgentId, JobStepOutcome? Outcome, float WaitTime, float InitTime, DateTime StartTime, DateTime? FinishTime)
+		Task<IJobStepRef> IJobStepRefCollection.InsertOrReplaceAsync(JobStepRefId Id, string JobName, string NodeName, StreamId StreamId, TemplateRefId TemplateId, int Change, LogId? LogId, PoolId? PoolId, AgentId? AgentId, JobStepOutcome? Outcome, int? LastSuccess, int? LastWarning, float WaitTime, float InitTime, DateTime StartTime, DateTime? FinishTime)
 		{
 			throw new NotImplementedException();
 		}
diff --git a/Engine/Source/Programs/Horde/Horde.Build/Jobs/IJobStepRef.cs b/Engine/Source/Programs/Horde/Horde.Build/Jobs/IJobStepRef.cs
--- a/Engine/Source/Programs/Horde/Horde.Build/Jobs/IJobStepRef.cs
+++ b/Engine/Source/Programs/Horde/Horde.Build/Jobs/IJobStepRef.cs
@@ -164,6 +164,16 @@ namespace HordeServer.Models
 		/// </summary>
 		public JobStepOutcome? Outcome { get; }
 
+		/// <summary>
+		/// The last change that succeeded. Note that this is only set when the ref is updated; it is not necessarily consistent with steps run later.
+		/// </summary>
+		public int? LastSuccess { get; }
+
+		/// <summary>
+		/// The last change that succeeded, or completed a warning. See <see cref="LastSuccess"/>.
+		/// </summary>
+		public int? LastWarning { get; }
+
 		/// <summary>
 		/// Time taken for the batch containing this batch to start after it became ready
 		/// </summary>
diff --git a/Engine/Source/Programs/Horde/Horde.Build/Jobs/IJobStepRefCollection.cs b/Engine/Source/Programs/Horde/Horde.Build/Jobs/IJobStepRefCollection.cs
--- a/Engine/Source/Programs/Horde/Horde.Build/Jobs/IJobStepRefCollection.cs
+++ b/Engine/Source/Programs/Horde/Horde.Build/Jobs/IJobStepRefCollection.cs
@@ -37,11 +37,13 @@ namespace HordeServer.Collections
 		/// <param name="PoolId">The pool id</param>
 		/// <param name="AgentId">The agent id</param>
 		/// <param name="Outcome">Outcome of this step, if known</param>
+		/// <param name="LastSuccess">The last change that completed with success</param>
+		/// <param name="LastWarning">The last change that completed with a warning (or success)</param>
 		/// <param name="WaitTime">Time taken for the batch containing this step to start</param>
 		/// <param name="InitTime">Time taken for the batch containing this step to initializer</param>
 		/// <param name="StartTimeUtc">Start time</param>
 		/// <param name="FinishTimeUtc">Finish time for the step, if known</param>
-		Task<IJobStepRef> InsertOrReplaceAsync(JobStepRefId Id, string JobName, string StepName, StreamId StreamId, TemplateRefId TemplateId, int Change, LogId? LogId, PoolId? PoolId, AgentId? AgentId, JobStepOutcome? Outcome, float WaitTime, float InitTime, DateTime StartTimeUtc, DateTime? FinishTimeUtc);
+		Task<IJobStepRef> InsertOrReplaceAsync(JobStepRefId Id, string JobName, string StepName, StreamId StreamId, TemplateRefId TemplateId, int Change, LogId? LogId, PoolId? PoolId, AgentId? AgentId, JobStepOutcome? Outcome, int? LastSuccess, int? LastWarning, float WaitTime, float InitTime, DateTime StartTimeUtc, DateTime? FinishTimeUtc);
 
 		/// <summary>
 		/// Gets the history of a given node
@@ -78,7 +80,7 @@ namespace HordeServer.Collections
 
 	static class JobStepRefCollectionExtensions
 	{
-		public static Task UpdateAsync(this IJobStepRefCollection JobStepRefs, IJob Job, IJobStepBatch Batch, IJobStep Step, IGraph Graph)
+		public static async Task UpdateAsync(this IJobStepRefCollection JobStepRefs, IJob Job, IJobStepBatch Batch, IJobStep Step, IGraph Graph)
 		{
 			if (Job.PreflightChange == 0)
 			{
@@ -88,11 +90,22 @@ namespace HordeServer.Collections
 				string NodeName = Graph.Groups[Batch.GroupIdx].Nodes[Step.NodeIdx].Name;
 				JobStepOutcome? Outcome = Step.IsPending() ? (JobStepOutcome?)null : Step.Outcome;
 
-				return JobStepRefs.InsertOrReplaceAsync(new JobStepRefId(Job.Id, Batch.Id, Step.Id), Job.Name, NodeName, Job.StreamId, Job.TemplateId, Job.Change, Step.LogId, Batch.PoolId, Batch.AgentId, Outcome, WaitTime, InitTime, Step.StartTimeUtc ?? DateTime.UtcNow, Step.FinishTimeUtc);
-			}
-			else
-			{
-				return Task.CompletedTask;
+				int? LastSuccess = null;
+				int? LastWarning = null;
+				if (Outcome != JobStepOutcome.Success)
+				{
+					IJobStepRef? PrevStep = await JobStepRefs.GetPrevStepForNodeAsync(Job.StreamId, Job.TemplateId, NodeName, Job.Change);
+					if (PrevStep != null)
+					{
+						LastSuccess = PrevStep.LastSuccess;
+						if (Outcome != JobStepOutcome.Warnings)
+						{
+							LastWarning = PrevStep.LastWarning;
+						}
+					}
+				}
+
+				await JobStepRefs.InsertOrReplaceAsync(new JobStepRefId(Job.Id, Batch.Id, Step.Id), Job.Name, NodeName, Job.StreamId, Job.TemplateId, Job.Change, Step.LogId, Batch.PoolId, Batch.AgentId, Outcome, LastSuccess, LastWarning, WaitTime, InitTime, Step.StartTimeUtc ?? DateTime.UtcNow, Step.FinishTimeUtc);
 			}
 		}
 	}
diff --git a/Engine/Source/Programs/Horde/Horde.Build/Services/RpcService.cs b/Engine/Source/Programs/Horde/Horde.Build/Services/RpcService.cs
--- a/Engine/Source/Programs/Horde/Horde.Build/Services/RpcService.cs
+++ b/Engine/Source/Programs/Horde/Horde.Build/Services/RpcService.cs
@@ -48,100 +48,27 @@ namespace HordeServer.Services
 		/// </summary>
 		internal TimeSpan LongPollTimeout = TimeSpan.FromMinutes(9);
 
-		/// <summary>
-		/// Instance of the DatabaseService singleton
-		/// </summary>
 		DatabaseService DatabaseService;
-
-		/// <summary>
-		/// Instance of the AclService singleton
-		/// </summary>
 		AclService AclService;
-
-		/// <summary>
-		/// Instance of the AgentService singleton
-		/// </summary>
 		AgentService AgentService;
-
-		/// <summary>
-		/// The stream service instance
-		/// </summary>
 		StreamService StreamService;
-
-		/// <summary>
-		/// The job service instance
-		/// </summary>
 		JobService JobService;
-
-		/// <summary>
-		/// The software service instance
-		/// </summary>
 		AgentSoftwareService AgentSoftwareService;
-
-		/// <summary>
-		/// The artifact service instance
-		/// </summary>
 		IArtifactCollection ArtifactCollection;
-
-		/// <summary>
-		/// The log file service instance
-		/// </summary>
 		ILogFileService LogFileService;
-
-		/// <summary>
-		/// The credential service instance
-		/// </summary>
 		CredentialService CredentialService;
-
-		/// <summary>
-		/// The pool service instance
-		/// </summary>
 		PoolService PoolService;
-
-		/// <summary>
-		/// The application lifetime interface
-		/// </summary>
 		LifetimeService LifetimeService;
-
-		/// <summary>
-		/// Collection of graph documents
-		/// </summary>
 		IGraphCollection Graphs;
-
-		/// <summary>
-		/// Collection of testdata documents
-		/// </summary>
 		ITestDataCollection TestData;
-
-		/// <summary>
-		/// The conform task source
-		/// </summary>
+		IJobStepRefCollection JobStepRefCollection;
 		ConformTaskSource ConformTaskSource;
-
-		/// <summary>
-		/// Writer for log output
-		/// </summary>
 		ILogger<RpcService> Logger;
 
 		/// <summary>
 		/// Constructor
 		/// </summary>
-		/// <param name="DatabaseService">Instance of the DatabaseService singleton</param>
-		/// <param name="AclService">Instance of the AclService singleton</param>
-		/// <param name="AgentService">Instance of the AgentService singleton</param>
-		/// <param name="JobService">Instance of the JobService singleton</param>
-		/// <param name="StreamService">Instance of the StreamService singleton</param>
-		/// <param name="LogFileService">Instance of the LogFileService singleton</param>
-		/// <param name="AgentSoftwareService">Instance of the JobService singleton</param>
-		/// <param name="ArtifactCollection">Instance of the ArtifactService singleton</param>
-		/// <param name="CredentialService">Instance of the CredentialsService singleton</param>
-		/// <param name="PoolService">Instance of the PoolService singleton</param>
-		/// <param name="LifetimeService">The application lifetime</param>
-		/// <param name="Graphs">Collection of graph documents</param>
-		/// <param name="TestData">Collection of testdata</param>
-		/// <param name="ConformTaskSource"></param>
-		/// <param name="Logger">Log writer</param>
-		public RpcService(DatabaseService DatabaseService, AclService AclService, AgentService AgentService, StreamService StreamService, JobService JobService, AgentSoftwareService AgentSoftwareService, IArtifactCollection ArtifactCollection, ILogFileService LogFileService, CredentialService CredentialService, PoolService PoolService, LifetimeService LifetimeService, IGraphCollection Graphs, ITestDataCollection TestData, ConformTaskSource ConformTaskSource, ILogger<RpcService> Logger)
+		public RpcService(DatabaseService DatabaseService, AclService AclService, AgentService AgentService, StreamService StreamService, JobService JobService, AgentSoftwareService AgentSoftwareService, IArtifactCollection ArtifactCollection, ILogFileService LogFileService, CredentialService CredentialService, PoolService PoolService, LifetimeService LifetimeService, IGraphCollection Graphs, ITestDataCollection TestData, IJobStepRefCollection JobStepRefCollection, ConformTaskSource ConformTaskSource, ILogger<RpcService> Logger)
 		{
 			this.DatabaseService = DatabaseService;
 			this.AclService = AclService;
@@ -156,6 +83,7 @@ namespace HordeServer.Services
 			this.LifetimeService = LifetimeService;
 			this.Graphs = Graphs;
 			this.TestData = TestData;
+			this.JobStepRefCollection = JobStepRefCollection;
 			this.ConformTaskSource = ConformTaskSource;
 			this.Logger = Logger;
 		}
@@ -662,6 +590,31 @@ namespace HordeServer.Services
 				Response.StepId = Step.Id.ToString();
 				Response.Name = Node.Name;
 				Response.Credentials.Add(Credentials);
+
+				IJobStepRef? LastStep = await JobStepRefCollection.GetPrevStepForNodeAsync(Job.StreamId, Job.TemplateId, Node.Name, Job.Change);
+				if (LastStep != null)
+				{
+					Response.EnvVars.Add("UE_HORDE_LAST_CL", LastStep.Change.ToString(CultureInfo.InvariantCulture));
+
+					if (LastStep.Outcome == JobStepOutcome.Success)
+					{
+						Response.EnvVars.Add("UE_HORDE_LAST_SUCCESS_CL", LastStep.Change.ToString(CultureInfo.InvariantCulture));
+					}
+					else if (LastStep.LastSuccess != null)
+					{
+						Response.EnvVars.Add("UE_HORDE_LAST_SUCCESS_CL", LastStep.LastSuccess.Value.ToString(CultureInfo.InvariantCulture));
+					}
+
+					if (LastStep.Outcome == JobStepOutcome.Success || LastStep.Outcome == JobStepOutcome.Warnings)
+					{
+						Response.EnvVars.Add("UE_HORDE_LAST_WARNING_CL", LastStep.Change.ToString(CultureInfo.InvariantCulture));
+					}
+					else if (LastStep.LastWarning != null)
+					{
+						Response.EnvVars.Add("UE_HORDE_LAST_WARNING_CL", LastStep.LastWarning.Value.ToString(CultureInfo.InvariantCulture));
+					}
+				}
+
 				if (Node.Properties != null)
 				{
 					Response.Properties.Add(Node.Properties);
