Horde: Missing files.

#skipci

diff --git a/Engine/Source/Programs/Horde/Horde.Build/Jobs/Impl/JobStepRefCollection.cs b/Engine/Source/Programs/Horde/Horde.Build/Jobs/Impl/JobStepRefCollection.cs
--- a/Engine/Source/Programs/Horde/Horde.Build/Jobs/Impl/JobStepRefCollection.cs
+++ b/Engine/Source/Programs/Horde/Horde.Build/Jobs/Impl/JobStepRefCollection.cs
@@ -40,6 +40,13 @@ namespace HordeServer.Collections.Impl
 			public PoolId? PoolId { get; set; }
 			public AgentId? AgentId { get; set; }
 			public JobStepOutcome? Outcome { get; set; }
+
+			[BsonIgnoreIfNull]
+			public int? LastSuccess { get; }
+
+			[BsonIgnoreIfNull]
+			public int? LastWarning { get; }
+
 			public float BatchWaitTime { get; set; }
 			public float BatchInitTime { get; set; }
 
@@ -59,7 +66,7 @@ namespace HordeServer.Collections.Impl
 			DateTime? IJobStepRef.FinishTimeUtc => FinishTimeUtc ?? FinishTime?.UtcDateTime;
 			string IJobStepRef.NodeName => Name;
 
-			public JobStepRef(JobStepRefId Id, string JobName, string NodeName, StreamId StreamId, TemplateRefId TemplateId, int Change, LogId? LogId, PoolId? PoolId, AgentId? AgentId, JobStepOutcome? Outcome, float BatchWaitTime, float BatchInitTime, DateTime StartTimeUtc, DateTime? FinishTimeUtc)
+			public JobStepRef(JobStepRefId Id, string JobName, string NodeName, StreamId StreamId, TemplateRefId TemplateId, int Change, LogId? LogId, PoolId? PoolId, AgentId? AgentId, JobStepOutcome? Outcome, int? LastSuccess, int? LastWarning, float BatchWaitTime, float BatchInitTime, DateTime StartTimeUtc, DateTime? FinishTimeUtc)
 			{
 				this.Id = Id;
 				this.JobName = JobName;
@@ -71,6 +78,8 @@ namespace HordeServer.Collections.Impl
 				this.PoolId = PoolId;
 				this.AgentId = AgentId;
 				this.Outcome = Outcome;
+				this.LastSuccess = LastSuccess;
+				this.LastWarning = LastWarning;
 				this.BatchWaitTime = BatchWaitTime;
 				this.BatchInitTime = BatchInitTime;
 				this.StartTime = StartTimeUtc;
@@ -97,9 +106,9 @@ namespace HordeServer.Collections.Impl
 		}
 
 		/// <inheritdoc/>
-		public async Task<IJobStepRef> InsertOrReplaceAsync(JobStepRefId Id, string JobName, string StepName, StreamId StreamId, TemplateRefId TemplateId, int Change, LogId? LogId, PoolId? PoolId, AgentId? AgentId, JobStepOutcome? Outcome, float WaitTime, float InitTime, DateTime StartTimeUtc, DateTime? FinishTimeUtc)
+		public async Task<IJobStepRef> InsertOrReplaceAsync(JobStepRefId Id, string JobName, string StepName, StreamId StreamId, TemplateRefId TemplateId, int Change, LogId? LogId, PoolId? PoolId, AgentId? AgentId, JobStepOutcome? Outcome, int? LastSuccess, int? LastWarning, float WaitTime, float InitTime, DateTime StartTimeUtc, DateTime? FinishTimeUtc)
 		{
-			JobStepRef NewJobStepRef = new JobStepRef(Id, JobName, StepName, StreamId, TemplateId, Change, LogId, PoolId, AgentId, Outcome, WaitTime, InitTime, StartTimeUtc, FinishTimeUtc);
+			JobStepRef NewJobStepRef = new JobStepRef(Id, JobName, StepName, StreamId, TemplateId, Change, LogId, PoolId, AgentId, Outcome, LastSuccess, LastWarning, WaitTime, InitTime, StartTimeUtc, FinishTimeUtc);
 			await JobStepRefs.ReplaceOneAsync(Builders<JobStepRef>.Filter.Eq(x => x.Id, NewJobStepRef.Id), NewJobStepRef, new ReplaceOptions { IsUpsert = true });
 			return NewJobStepRef;
 		}
