Horde Storage - Tweaked the Gc ref cleanup in scylla, adding more scopes and just generally trying to make it clearer were we are spending the time.

diff --git a/Engine/Source/Programs/Horde/HordeStorage/Horde.Storage/Implementation/GC/CacheCleanup.cs b/Engine/Source/Programs/Horde/HordeStorage/Horde.Storage/Implementation/GC/CacheCleanup.cs
--- a/Engine/Source/Programs/Horde/HordeStorage/Horde.Storage/Implementation/GC/CacheCleanup.cs
+++ b/Engine/Source/Programs/Horde/HordeStorage/Horde.Storage/Implementation/GC/CacheCleanup.cs
@@ -62,6 +62,8 @@ namespace Horde.Storage.Implementation
                 if (record.LastAccess > cutoffTime)
                     break;
 
+                using Scope scope = Tracer.Instance.StartActive("refCleanup.delete_record");
+                scope.Span.ResourceName = $"{record.Namespace}:{record.Bucket}.{record.Name}";
                 // delete the old record from the ref refs
                 Task storeDelete = _referencesStore.Delete(record.Namespace, record.Bucket, record.Name);
                 // insert a delete event into the transaction log
diff --git a/Engine/Source/Programs/Horde/HordeStorage/Horde.Storage/Implementation/GC/RefCleanupService.cs b/Engine/Source/Programs/Horde/HordeStorage/Horde.Storage/Implementation/GC/RefCleanupService.cs
--- a/Engine/Source/Programs/Horde/HordeStorage/Horde.Storage/Implementation/GC/RefCleanupService.cs
+++ b/Engine/Source/Programs/Horde/HordeStorage/Horde.Storage/Implementation/GC/RefCleanupService.cs
@@ -4,6 +4,7 @@ using System;
 using System.Collections.Generic;
 using System.Threading;
 using System.Threading.Tasks;
+using Dasync.Collections;
 using Datadog.Trace;
 using Jupiter;
 using Jupiter.Implementation;
@@ -76,7 +77,8 @@ namespace Horde.Storage.Implementation
                     }
                 }
 
-                await foreach (NamespaceId ns in _referencesStore.GetNamespaces().WithCancellation(cancellationToken))
+                List<NamespaceId>? namespaces = await _referencesStore.GetNamespaces().ToListAsync(cancellationToken);
+                await foreach (NamespaceId ns in namespaces)
                 {
                     using Scope scope = Tracer.Instance.StartActive("gc.refs");
                     scope.Span.ResourceName = ns.ToString();
diff --git a/Engine/Source/Programs/Horde/HordeStorage/Horde.Storage/Implementation/Objects/ScyllaReferencesStore.cs b/Engine/Source/Programs/Horde/HordeStorage/Horde.Storage/Implementation/Objects/ScyllaReferencesStore.cs
--- a/Engine/Source/Programs/Horde/HordeStorage/Horde.Storage/Implementation/Objects/ScyllaReferencesStore.cs
+++ b/Engine/Source/Programs/Horde/HordeStorage/Horde.Storage/Implementation/Objects/ScyllaReferencesStore.cs
@@ -146,18 +146,26 @@ namespace Horde.Storage.Implementation
             using Scope _ = Tracer.Instance.StartActive("scylla.get_records");
             for (sbyte partitionIndex = sbyte.MinValue; partitionIndex < sbyte.MaxValue; partitionIndex++)
             {
-                RowSet rowSetDateBuckets = await _session.ExecuteAsync(new SimpleStatement("SELECT accessed_at FROM object_last_access WHERE namespace = ? AND partition_index = ? ALLOW FILTERING", ns.ToString(), partitionIndex));
+                using Scope scope = Tracer.Instance.StartActive("scylla.get_records_partition");
+                scope.Span.ResourceName = $"partition_{partitionIndex}";
+
                 SortedSet<DateTime> dateBuckets = new SortedSet<DateTime>();
-                foreach (Row row in rowSetDateBuckets)
                 {
-                    if (rowSetDateBuckets.GetAvailableWithoutFetching() == 0)
-                        await rowSetDateBuckets.FetchMoreResultsAsync();
+                    using Scope dateBucketsScope = Tracer.Instance.StartActive("scylla.get_records.date_buckets");
+
+                    RowSet rowSetDateBuckets = await _session.ExecuteAsync(new SimpleStatement("SELECT accessed_at FROM object_last_access WHERE namespace = ? AND partition_index = ? ALLOW FILTERING", ns.ToString(), partitionIndex));
+                    foreach (Row row in rowSetDateBuckets)
+                    {
+                        if (rowSetDateBuckets.GetAvailableWithoutFetching() == 0)
+                            await rowSetDateBuckets.FetchMoreResultsAsync();
 
-                    LocalDate date = row.GetValue<LocalDate>("accessed_at");
+                        LocalDate date = row.GetValue<LocalDate>("accessed_at");
 
-                    dateBuckets.Add(new DateTime(date.Year, date.Month, date.Day));
+                        dateBuckets.Add(new DateTime(date.Year, date.Month, date.Day));
+                    }
                 }
 
+
                 foreach (DateTime dateBucket in dateBuckets)
                 {
                     LocalDate date = new LocalDate(dateBucket.Year, dateBucket.Month, dateBucket.Day);
