Horde: Add preflight exclusive search to jobs controller and dashboard
#rnx

diff --git a/Engine/Source/Programs/Horde/Horde.Build/Controllers/JobsController.cs b/Engine/Source/Programs/Horde/Horde.Build/Controllers/JobsController.cs
--- a/Engine/Source/Programs/Horde/Horde.Build/Controllers/JobsController.cs
+++ b/Engine/Source/Programs/Horde/Horde.Build/Controllers/JobsController.cs
@@ -614,6 +614,7 @@ namespace HordeServer.Controllers
 		/// <param name="MinChange">The minimum changelist number</param>
 		/// <param name="MaxChange">The maximum changelist number</param>
 		/// <param name="IncludePreflight">Whether to include preflight jobs</param>
+		/// <param name="PreflightOnly">Whether to only include preflight jobs</param>
 		/// <param name="PreflightChange">The preflighted changelist</param>
 		/// <param name="StartedByUserId">User id for which to include jobs</param>
 		/// <param name="PreflightStartedByUserId">User id for which to include preflight jobs</param>
@@ -639,6 +640,7 @@ namespace HordeServer.Controllers
 			[FromQuery] int? MinChange = null,
 			[FromQuery] int? MaxChange = null,
 			[FromQuery] bool IncludePreflight = true,
+			[FromQuery] bool? PreflightOnly = null,
 			[FromQuery] int? PreflightChange = null,
 			[FromQuery] string? PreflightStartedByUserId = null,
 			[FromQuery] string? StartedByUserId = null,
@@ -681,7 +683,7 @@ namespace HordeServer.Controllers
 			using (IScope _ = GlobalTracer.Instance.BuildSpan("FindJobs").StartActive())
 			{
 				Jobs = await JobService.FindJobsAsync(JobIdValues, StreamIdValue, Name, TemplateRefIds, MinChange,
-					MaxChange, PreflightChange, PreflightStartedByUserIdValue, StartedByUserIdValue, MinCreateTime?.UtcDateTime, MaxCreateTime?.UtcDateTime, Target, State, Outcome,
+					MaxChange, PreflightChange, PreflightOnly, PreflightStartedByUserIdValue, StartedByUserIdValue, MinCreateTime?.UtcDateTime, MaxCreateTime?.UtcDateTime, Target, State, Outcome,
 					ModifiedBefore, ModifiedAfter, Index, Count, false);
 			}
 
diff --git a/Engine/Source/Programs/Horde/Horde.Build/Jobs/IJobCollection.cs b/Engine/Source/Programs/Horde/Horde.Build/Jobs/IJobCollection.cs
--- a/Engine/Source/Programs/Horde/Horde.Build/Jobs/IJobCollection.cs
+++ b/Engine/Source/Programs/Horde/Horde.Build/Jobs/IJobCollection.cs
@@ -97,6 +97,7 @@ namespace HordeServer.Collections
 		/// <param name="MinChange">The minimum changelist number</param>
 		/// <param name="MaxChange">The maximum changelist number</param>
 		/// <param name="PreflightChange">Preflight change to find</param>
+		/// <param name="PreflightOnly">Whether to only include preflights</param>
 		/// <param name="StartedByUser">User id for which to include jobs</param>
 		/// <param name="PreflightStartedByUser">User for which to include preflight jobs</param>
 		/// <param name="MinCreateTime">The minimum creation time</param>
@@ -108,7 +109,7 @@ namespace HordeServer.Collections
 		/// <param name="ConsistentRead">If the database read should be made to the replica server</param>
 		/// <param name="IndexHint">Name of index to be specified as a hint to the database query planner</param>
 		/// <returns>List of jobs matching the given criteria</returns>
-		Task<List<IJob>> FindAsync(JobId[]? JobIds = null, StreamId? StreamId = null, string? Name = null, TemplateRefId[]? Templates = null, int? MinChange = null, int? MaxChange = null, int? PreflightChange = null, UserId? PreflightStartedByUser = null, UserId? StartedByUser = null, DateTimeOffset? MinCreateTime = null, DateTimeOffset? MaxCreateTime = null, DateTimeOffset? ModifiedBefore = null, DateTimeOffset? ModifiedAfter = null, int? Index = null, int? Count = null, bool ConsistentRead = true, string? IndexHint = null);
+		Task<List<IJob>> FindAsync(JobId[]? JobIds = null, StreamId? StreamId = null, string? Name = null, TemplateRefId[]? Templates = null, int? MinChange = null, int? MaxChange = null, int? PreflightChange = null, bool? PreflightOnly = null, UserId? PreflightStartedByUser = null, UserId? StartedByUser = null, DateTimeOffset? MinCreateTime = null, DateTimeOffset? MaxCreateTime = null, DateTimeOffset? ModifiedBefore = null, DateTimeOffset? ModifiedAfter = null, int? Index = null, int? Count = null, bool ConsistentRead = true, string? IndexHint = null);
 
 		/// <summary>
 		/// Searches for jobs in a specified stream and templates
diff --git a/Engine/Source/Programs/Horde/Horde.Build/Jobs/Impl/JobCollection.cs b/Engine/Source/Programs/Horde/Horde.Build/Jobs/Impl/JobCollection.cs
--- a/Engine/Source/Programs/Horde/Horde.Build/Jobs/Impl/JobCollection.cs
+++ b/Engine/Source/Programs/Horde/Horde.Build/Jobs/Impl/JobCollection.cs
@@ -516,7 +516,7 @@ namespace HordeServer.Collections.Impl
 		}
 
 		/// <inheritdoc/>
-		public async Task<List<IJob>> FindAsync(JobId[]? JobIds, StreamId? StreamId, string? Name, TemplateRefId[]? Templates, int? MinChange, int? MaxChange, int? PreflightChange, UserId? PreflightStartedByUser, UserId? StartedByUser, DateTimeOffset? MinCreateTime, DateTimeOffset? MaxCreateTime, DateTimeOffset? ModifiedBefore, DateTimeOffset? ModifiedAfter, int? Index, int? Count, bool ConsistentRead, string? IndexHint)
+		public async Task<List<IJob>> FindAsync(JobId[]? JobIds, StreamId? StreamId, string? Name, TemplateRefId[]? Templates, int? MinChange, int? MaxChange, int? PreflightChange, bool? PreflightOnly, UserId ? PreflightStartedByUser, UserId? StartedByUser, DateTimeOffset? MinCreateTime, DateTimeOffset? MaxCreateTime, DateTimeOffset? ModifiedBefore, DateTimeOffset? ModifiedAfter, int? Index, int? Count, bool ConsistentRead, string? IndexHint)
 		{
 			FilterDefinitionBuilder<JobDocument> FilterBuilder = Builders<JobDocument>.Filter;
 
@@ -549,6 +549,10 @@ namespace HordeServer.Collections.Impl
 			{
 				Filter &= FilterBuilder.Eq(x => x.PreflightChange, PreflightChange);
 			}
+			if (PreflightOnly != null && PreflightOnly.Value)
+			{
+				Filter &= FilterBuilder.Ne(x => x.PreflightChange, 0);
+			}
 			if (PreflightStartedByUser != null)
 			{
 				Filter &= FilterBuilder.Or(FilterBuilder.Eq(x => x.PreflightChange, 0), FilterBuilder.Eq(x => x.StartedByUserId, PreflightStartedByUser));
diff --git a/Engine/Source/Programs/Horde/Horde.Build/Jobs/Impl/JobService.cs b/Engine/Source/Programs/Horde/Horde.Build/Jobs/Impl/JobService.cs
--- a/Engine/Source/Programs/Horde/Horde.Build/Jobs/Impl/JobService.cs
+++ b/Engine/Source/Programs/Horde/Horde.Build/Jobs/Impl/JobService.cs
@@ -520,6 +520,7 @@ namespace HordeServer.Services
 		/// <param name="MinChange">The minimum changelist number</param>
 		/// <param name="MaxChange">The maximum changelist number</param>		
 		/// <param name="PreflightChange">The preflight change to look for</param>
+		/// <param name="PreflightOnly">Whether to only include preflights</param>
 		/// <param name="PreflightStartedByUser">User for which to include preflight jobs</param>		
 		/// <param name="StartedByUser">User for which to include jobs</param>
 		/// <param name="MinCreateTime">The minimum creation time</param>
@@ -533,7 +534,7 @@ namespace HordeServer.Services
 		/// <param name="Count">Number of results to return</param>
 		/// <param name="ConsistentRead">If the database read should be made to the replica server</param>
 		/// <returns>List of jobs matching the given criteria</returns>
-		public async Task<List<IJob>> FindJobsAsync(JobId[]? JobIds = null, StreamId? StreamId = null, string? Name = null, TemplateRefId[]? Templates = null, int? MinChange = null, int? MaxChange = null, int? PreflightChange = null, UserId? PreflightStartedByUser = null, UserId? StartedByUser = null, DateTimeOffset ? MinCreateTime = null, DateTimeOffset? MaxCreateTime = null, string? Target = null, JobStepState[]? State = null, JobStepOutcome[]? Outcome = null, DateTimeOffset? ModifiedBefore = null, DateTimeOffset? ModifiedAfter = null, int? Index = null, int? Count = null, bool ConsistentRead = true)
+		public async Task<List<IJob>> FindJobsAsync(JobId[]? JobIds = null, StreamId? StreamId = null, string? Name = null, TemplateRefId[]? Templates = null, int? MinChange = null, int? MaxChange = null, int? PreflightChange = null, bool? PreflightOnly = null, UserId? PreflightStartedByUser = null, UserId? StartedByUser = null, DateTimeOffset ? MinCreateTime = null, DateTimeOffset? MaxCreateTime = null, string? Target = null, JobStepState[]? State = null, JobStepOutcome[]? Outcome = null, DateTimeOffset? ModifiedBefore = null, DateTimeOffset? ModifiedAfter = null, int? Index = null, int? Count = null, bool ConsistentRead = true)
 		{
 			using IScope Scope = GlobalTracer.Instance.BuildSpan("JobService.FindJobsAsync").StartActive();
 			Scope.Span.SetTag("JobIds", JobIds);
@@ -557,7 +558,7 @@ namespace HordeServer.Services
 			
 			if (Target == null && (State == null || State.Length == 0) && (Outcome == null || Outcome.Length == 0))
 			{
-				return await Jobs.FindAsync(JobIds, StreamId, Name, Templates, MinChange, MaxChange, PreflightChange, PreflightStartedByUser, StartedByUser, MinCreateTime, MaxCreateTime, ModifiedBefore, ModifiedAfter, Index, Count, ConsistentRead);
+				return await Jobs.FindAsync(JobIds, StreamId, Name, Templates, MinChange, MaxChange, PreflightChange, PreflightOnly, PreflightStartedByUser, StartedByUser, MinCreateTime, MaxCreateTime, ModifiedBefore, ModifiedAfter, Index, Count, ConsistentRead);
 			}
 			else
 			{
@@ -567,7 +568,7 @@ namespace HordeServer.Services
 				int MaxCount = (Count ?? 1);
 				while (Results.Count < MaxCount)
 				{
-					List<IJob> ScanJobs = await Jobs.FindAsync(JobIds, StreamId, Name, Templates, MinChange, MaxChange, PreflightChange, PreflightStartedByUser, StartedByUser, MinCreateTime, MaxCreateTime, ModifiedBefore, ModifiedAfter, 0, 5, ConsistentRead);
+					List<IJob> ScanJobs = await Jobs.FindAsync(JobIds, StreamId, Name, Templates, MinChange, MaxChange, PreflightChange, PreflightOnly, PreflightStartedByUser, StartedByUser, MinCreateTime, MaxCreateTime, ModifiedBefore, ModifiedAfter, 0, 5, ConsistentRead);
 					if (ScanJobs.Count == 0)
 					{
 						break;
