Horde: Fix missing file from job search commit
#rnx

diff --git a/Engine/Source/Programs/Horde/HordeDashboard/src/backend/Api.ts b/Engine/Source/Programs/Horde/HordeDashboard/src/backend/Api.ts
--- a/Engine/Source/Programs/Horde/HordeDashboard/src/backend/Api.ts
+++ b/Engine/Source/Programs/Horde/HordeDashboard/src/backend/Api.ts
@@ -274,9 +274,10 @@ export type JobQuery = {
 	minChange?: number;
 	maxChange?: number;
 	preflightChange?: number;
-	includePreflight?: boolean,
-	preflightStartedByUserId?: string,
-	startedByUserId?: string,
+	preflightOnly?: boolean;
+	includePreflight?: boolean;
+	preflightStartedByUserId?: string;
+	startedByUserId?: string;
 	minCreateTime?: string;
 	maxCreateTime?: string;
 	modifiedAfter?: string;
