Horde: Fix server read only mode error string
#rnx

diff --git a/Engine/Source/Programs/Horde/Horde.Build/Utilities/BaseDatabaseIndexes.cs b/Engine/Source/Programs/Horde/Horde.Build/Utilities/BaseDatabaseIndexes.cs
--- a/Engine/Source/Programs/Horde/Horde.Build/Utilities/BaseDatabaseIndexes.cs
+++ b/Engine/Source/Programs/Horde/Horde.Build/Utilities/BaseDatabaseIndexes.cs
@@ -39,7 +39,7 @@ namespace HordeServer.Utilities
 			{
 				if (Collection.Indexes.List().ToList().Find(x => x.TryGetString("name", out string? n) && n == Name) == null)
 				{
-					throw new ArgumentException("Index with name {Name} not found (running in database read-only mode)");
+					throw new ArgumentException($"Index with name {Name} not found (running in database read-only mode)");
 				}
 				return Name;
 			}
