Horde Storage
Attempted another implementation at the last access tracking in scylla that is more straight forward. The previous implementation tried to be smart about how it ordered data but we just ended up with a semi full table scan anyway. This now just does a full table scan of our objects and checks their last access times. While naive this can be improved a bit and requires much less tracking and head aches.

This also avoids downloading the payload of the object which is quite expensive, will add a general toggle for that as we often do not need it.

diff --git a/Engine/Source/Programs/Horde/HordeStorage/Horde.Storage/Implementation/GC/CacheCleanup.cs b/Engine/Source/Programs/Horde/HordeStorage/Horde.Storage/Implementation/GC/CacheCleanup.cs
--- a/Engine/Source/Programs/Horde/HordeStorage/Horde.Storage/Implementation/GC/CacheCleanup.cs
+++ b/Engine/Source/Programs/Horde/HordeStorage/Horde.Storage/Implementation/GC/CacheCleanup.cs
@@ -56,18 +56,17 @@ namespace Horde.Storage.Implementation
         {
             List<OldRecord> deletedRecords = new List<OldRecord>();
             DateTime cutoffTime = DateTime.Now.Add(_settings.CurrentValue.LastAccessCutoff);
-            await foreach (ObjectRecord record in _referencesStore.GetOldestRecords(ns).WithCancellation(cancellationToken))
+            await foreach ((BucketId bucket, IoHashKey name, DateTime lastAccessTime) in _referencesStore.GetRecords(ns).WithCancellation(cancellationToken))
             {
-                // the record are returned oldest first, so once we find a record that is newer then our cutoff time we can stop iterating
-                if (record.LastAccess > cutoffTime)
-                    break;
+                if (lastAccessTime > cutoffTime)
+                    continue;
 
                 using Scope scope = Tracer.Instance.StartActive("refCleanup.delete_record");
-                scope.Span.ResourceName = $"{record.Namespace}:{record.Bucket}.{record.Name}";
+                scope.Span.ResourceName = $"{ns}:{bucket}.{name}";
                 // delete the old record from the ref refs
-                Task storeDelete = _referencesStore.Delete(record.Namespace, record.Bucket, record.Name);
+                Task storeDelete = _referencesStore.Delete(ns, bucket, name);
                 // insert a delete event into the transaction log
-                Task<(string, Guid)> transactionLogDelete = _replicationLog.InsertDeleteEvent(record.Namespace, record.Bucket, record.Name, record.BlobIdentifier);
+                Task<(string, Guid)> transactionLogDelete = _replicationLog.InsertDeleteEvent(ns, bucket, name, null);
 
                 try
                 {
@@ -75,10 +74,10 @@ namespace Horde.Storage.Implementation
                 }
                 catch (Exception e)
                 {
-                    _logger.Warning(e, "Exception when attempting to delete record {Record} in {Namespace}", record, ns);
+                    _logger.Warning(e, "Exception when attempting to delete record {Bucket} {Name} in {Namespace}", bucket, name, ns);
                 }
                 // we convert the ObjectRecords key (which is a iohash) to a keyid (which is a generic string) not the prettiest but this result is only used for debugging purposes anyway
-                deletedRecords.Add(new OldRecord(record.Namespace, record.Bucket, new KeyId(record.Name.ToString())));
+                deletedRecords.Add(new OldRecord(ns, bucket, new KeyId(name.ToString())));
             }
 
             return deletedRecords;
diff --git a/Engine/Source/Programs/Horde/HordeStorage/Horde.Storage/Implementation/Objects/IReferencesStore.cs b/Engine/Source/Programs/Horde/HordeStorage/Horde.Storage/Implementation/Objects/IReferencesStore.cs
--- a/Engine/Source/Programs/Horde/HordeStorage/Horde.Storage/Implementation/Objects/IReferencesStore.cs
+++ b/Engine/Source/Programs/Horde/HordeStorage/Horde.Storage/Implementation/Objects/IReferencesStore.cs
@@ -15,7 +15,7 @@ namespace Horde.Storage.Implementation
 
 
         Task UpdateLastAccessTime(NamespaceId ns, BucketId bucket, IoHashKey key, DateTime newLastAccessTime);
-        IAsyncEnumerable<ObjectRecord> GetOldestRecords(NamespaceId ns);
+        IAsyncEnumerable<(BucketId, IoHashKey, DateTime)> GetRecords(NamespaceId ns);
 
         IAsyncEnumerable<NamespaceId> GetNamespaces();
         Task<long> Delete(NamespaceId ns, BucketId bucket, IoHashKey key);
diff --git a/Engine/Source/Programs/Horde/HordeStorage/Horde.Storage/Implementation/Objects/MemoryReferencesStore.cs b/Engine/Source/Programs/Horde/HordeStorage/Horde.Storage/Implementation/Objects/MemoryReferencesStore.cs
--- a/Engine/Source/Programs/Horde/HordeStorage/Horde.Storage/Implementation/Objects/MemoryReferencesStore.cs
+++ b/Engine/Source/Programs/Horde/HordeStorage/Horde.Storage/Implementation/Objects/MemoryReferencesStore.cs
@@ -66,12 +66,12 @@ namespace Horde.Storage.Implementation
             return Task.CompletedTask;
         }
 
-        public async IAsyncEnumerable<ObjectRecord> GetOldestRecords(NamespaceId ns)
+        public async IAsyncEnumerable<(BucketId, IoHashKey, DateTime)> GetRecords(NamespaceId ns)
         {
             foreach (MemoryStoreObject o in _objects.Values.Where(o => o.Namespace == ns).OrderBy(o => o.LastAccessTime))
             {
                 await Task.CompletedTask;
-                yield return o.ToObjectRecord();
+                yield return (o.Bucket, o.Name, o.LastAccessTime);
             }
         }
 
diff --git a/Engine/Source/Programs/Horde/HordeStorage/Horde.Storage/Implementation/Objects/ScyllaReferencesStore.cs b/Engine/Source/Programs/Horde/HordeStorage/Horde.Storage/Implementation/Objects/ScyllaReferencesStore.cs
--- a/Engine/Source/Programs/Horde/HordeStorage/Horde.Storage/Implementation/Objects/ScyllaReferencesStore.cs
+++ b/Engine/Source/Programs/Horde/HordeStorage/Horde.Storage/Implementation/Objects/ScyllaReferencesStore.cs
@@ -141,7 +141,7 @@ namespace Horde.Storage.Implementation
             await updateObjectTracking;
         }
 
-        public async IAsyncEnumerable<ObjectRecord> GetOldestRecords(NamespaceId ns)
+        /*public async IAsyncEnumerable<ObjectRecord> GetRecords(NamespaceId ns)
         {
             using Scope _ = Tracer.Instance.StartActive("scylla.get_records");
             for (sbyte partitionIndex = sbyte.MinValue; partitionIndex < sbyte.MaxValue; partitionIndex++)
@@ -153,7 +153,7 @@ namespace Horde.Storage.Implementation
                 {
                     using Scope dateBucketsScope = Tracer.Instance.StartActive("scylla.get_records.date_buckets");
 
-                    RowSet rowSetDateBuckets = await _session.ExecuteAsync(new SimpleStatement("SELECT accessed_at FROM object_last_access WHERE namespace = ? AND partition_index = ? ALLOW FILTERING", ns.ToString(), partitionIndex));
+                    RowSet rowSetDateBuckets = await _session.ExecuteAsync(new SimpleStatement("SELECT accessed_at FROM object_last_access WHERE namespace = ? AND partition_index = ? PER PARTITION LIMIT 1 ALLOW FILTERING", ns.ToString(), partitionIndex));
                     foreach (Row row in rowSetDateBuckets)
                     {
                         if (rowSetDateBuckets.GetAvailableWithoutFetching() == 0)
@@ -193,6 +193,24 @@ namespace Horde.Storage.Implementation
                 }
 
             }
+        }*/
+        
+        public async IAsyncEnumerable<(BucketId, IoHashKey, DateTime)> GetRecords(NamespaceId ns)
+        {
+            using Scope _ = Tracer.Instance.StartActive("scylla.get_records");
+
+            RowSet rowSet = await _session.ExecuteAsync(new SimpleStatement("SELECT bucket, name, last_access_time FROM objects WHERE namespace = ? ALLOW FILTERING", ns.ToString()));
+            foreach (Row row in rowSet)
+            {
+                if (rowSet.GetAvailableWithoutFetching() == 0)
+                    await rowSet.FetchMoreResultsAsync();
+
+                string bucket = row.GetValue<string>("bucket");
+                string name = row.GetValue<string>("name");
+                DateTime lastAccessTime = row.GetValue<DateTime>("last_access_time");
+
+                yield return (new BucketId(bucket), new IoHashKey(name), lastAccessTime);
+            }
         }
 
         public async IAsyncEnumerable<NamespaceId> GetNamespaces()
diff --git a/Engine/Source/Programs/Horde/HordeStorage/Horde.Storage/Implementation/Replication/RefsReplicator.cs b/Engine/Source/Programs/Horde/HordeStorage/Horde.Storage/Implementation/Replication/RefsReplicator.cs
--- a/Engine/Source/Programs/Horde/HordeStorage/Horde.Storage/Implementation/Replication/RefsReplicator.cs
+++ b/Engine/Source/Programs/Horde/HordeStorage/Horde.Storage/Implementation/Replication/RefsReplicator.cs
@@ -319,6 +319,12 @@ namespace Horde.Storage.Implementation
                         replicationTasks.Add(currentOffset);
                     }
 
+                    // we do not need to replicate delete events
+                    if (@event.Op == ReplicationLogEvent.OpType.Deleted)
+                        return;
+                    if (@event.Blob == null)
+                        throw new Exception($"Event: [REDACTED_EMAIL]} [REDACTED_EMAIL]} in namespace [REDACTED_EMAIL]} was missing a blob, unable to replicate it");
+
                     await ReplicateOp(@event.Namespace, @event.Blob, replicationToken);
                     // TODO: Avoid adding to the replication log as we will get into infinite recursion if we do with the remote site
                     // we could add to the replication log only if the op was missing, but that could cause issue for a 3rd site replicating from us
diff --git a/Engine/Source/Programs/Horde/HordeStorage/Horde.Storage/Implementation/TransactionLog/IReplicationLog.cs b/Engine/Source/Programs/Horde/HordeStorage/Horde.Storage/Implementation/TransactionLog/IReplicationLog.cs
--- a/Engine/Source/Programs/Horde/HordeStorage/Horde.Storage/Implementation/TransactionLog/IReplicationLog.cs
+++ b/Engine/Source/Programs/Horde/HordeStorage/Horde.Storage/Implementation/TransactionLog/IReplicationLog.cs
@@ -15,7 +15,7 @@ namespace Horde.Storage.Implementation
         IAsyncEnumerable<NamespaceId> GetNamespaces();
 
         Task<(string, Guid)> InsertAddEvent(NamespaceId ns, BucketId bucket, IoHashKey key, BlobIdentifier objectBlob, DateTime? timeBucket = null);
-        Task<(string, Guid)> InsertDeleteEvent(NamespaceId ns, BucketId bucket, IoHashKey key, BlobIdentifier objectBlob, DateTime? timeBucket = null);
+        Task<(string, Guid)> InsertDeleteEvent(NamespaceId ns, BucketId bucket, IoHashKey key, DateTime? timeBucket = null);
         IAsyncEnumerable<ReplicationLogEvent> Get(NamespaceId ns, string? lastBucket, Guid? lastEvent);
 
         Task AddSnapshot(SnapshotInfo snapshotHeader);
@@ -40,7 +40,7 @@ namespace Horde.Storage.Implementation
 
     public class ReplicationLogEvent
     {
-        public ReplicationLogEvent(NamespaceId @namespace, BucketId bucket, IoHashKey key, BlobIdentifier blob, Guid eventId, string timeBucket, DateTime timestamp, OpType op)
+        public ReplicationLogEvent(NamespaceId @namespace, BucketId bucket, IoHashKey key, BlobIdentifier? blob, Guid eventId, string timeBucket, DateTime timestamp, OpType op)
         {
             Namespace = @namespace;
             Bucket = bucket;
@@ -66,7 +66,7 @@ namespace Horde.Storage.Implementation
         public DateTime Timestamp { get; }
         public string TimeBucket { get; }
         public Guid EventId { get; }
-        public BlobIdentifier Blob { get; }
+        public BlobIdentifier? Blob { get; }
     }
 
     public class IncrementalLogNotAvailableException : Exception
diff --git a/Engine/Source/Programs/Horde/HordeStorage/Horde.Storage/Implementation/TransactionLog/MemoryReplicationLog.cs b/Engine/Source/Programs/Horde/HordeStorage/Horde.Storage/Implementation/TransactionLog/MemoryReplicationLog.cs
--- a/Engine/Source/Programs/Horde/HordeStorage/Horde.Storage/Implementation/TransactionLog/MemoryReplicationLog.cs
+++ b/Engine/Source/Programs/Horde/HordeStorage/Horde.Storage/Implementation/TransactionLog/MemoryReplicationLog.cs
@@ -26,7 +26,7 @@ namespace Horde.Storage.Implementation
             return DoInsert(ns, bucket, key, objectBlob, ReplicationLogEvent.OpType.Added, timestamp);
         }
 
-        private async Task<(string, Guid)> DoInsert(NamespaceId ns, BucketId bucket, IoHashKey key, BlobIdentifier hash, ReplicationLogEvent.OpType op, DateTime? lastTimestamp)
+        private async Task<(string, Guid)> DoInsert(NamespaceId ns, BucketId bucket, IoHashKey key, BlobIdentifier? hash, ReplicationLogEvent.OpType op, DateTime? lastTimestamp)
         {
             DateTime timestamp = lastTimestamp.GetValueOrDefault(DateTime.Now);
 
@@ -65,9 +65,9 @@ namespace Horde.Storage.Implementation
             });
         }
 
-        public Task<(string, Guid)> InsertDeleteEvent(NamespaceId ns, BucketId bucket, IoHashKey key, BlobIdentifier objectBlob, DateTime? timestamp)
+        public Task<(string, Guid)> InsertDeleteEvent(NamespaceId ns, BucketId bucket, IoHashKey key, DateTime? timestamp)
         {
-            return DoInsert(ns, bucket, key, objectBlob, ReplicationLogEvent.OpType.Deleted, timestamp); 
+            return DoInsert(ns, bucket, key, null, ReplicationLogEvent.OpType.Deleted, timestamp); 
         }
 
         public async IAsyncEnumerable<ReplicationLogEvent> Get(NamespaceId ns, string? lastBucket, Guid? lastEvent)
diff --git a/Engine/Source/Programs/Horde/HordeStorage/Horde.Storage/Implementation/TransactionLog/ReplicationLogSnapshot.cs b/Engine/Source/Programs/Horde/HordeStorage/Horde.Storage/Implementation/TransactionLog/ReplicationLogSnapshot.cs
--- a/Engine/Source/Programs/Horde/HordeStorage/Horde.Storage/Implementation/TransactionLog/ReplicationLogSnapshot.cs
+++ b/Engine/Source/Programs/Horde/HordeStorage/Horde.Storage/Implementation/TransactionLog/ReplicationLogSnapshot.cs
@@ -118,7 +118,7 @@ namespace Horde.Storage.Implementation.TransactionLog
             switch (entry.Op)
             {
                 case ReplicationLogEvent.OpType.Added:
-                    ProcessAddEvent(entry.Bucket, entry.Key, entry.Blob);
+                    ProcessAddEvent(entry.Bucket, entry.Key, entry.Blob!);
                     break;
                 case ReplicationLogEvent.OpType.Deleted:
                     ProcessDeleteEvent(entry.Bucket, entry.Key);
diff --git a/Engine/Source/Programs/Horde/HordeStorage/Horde.Storage/Implementation/TransactionLog/ScyllaReplicationLog.cs b/Engine/Source/Programs/Horde/HordeStorage/Horde.Storage/Implementation/TransactionLog/ScyllaReplicationLog.cs
--- a/Engine/Source/Programs/Horde/HordeStorage/Horde.Storage/Implementation/TransactionLog/ScyllaReplicationLog.cs
+++ b/Engine/Source/Programs/Horde/HordeStorage/Horde.Storage/Implementation/TransactionLog/ScyllaReplicationLog.cs
@@ -75,12 +75,12 @@ namespace Horde.Storage.Implementation
             return (log.GetReplicationBucketIdentifier(), log.ReplicationId);
         }
 
-        public async Task<(string, Guid)> InsertDeleteEvent(NamespaceId ns, BucketId bucket, IoHashKey key, BlobIdentifier objectBlob, DateTime? timestamp)
+        public async Task<(string, Guid)> InsertDeleteEvent(NamespaceId ns, BucketId bucket, IoHashKey key, DateTime? timestamp)
         {
             using Scope _ = Tracer.Instance.StartActive("scylla.insert_delete_event");
             Task addNamespaceTask = PotentiallyAddNamespace(ns);
             DateTime timeBucket = timestamp.GetValueOrDefault(DateTime.Now);
-            ScyllaReplicationLogEvent log = new ScyllaReplicationLogEvent(ns.ToString(), bucket.ToString(), key.ToString(), timeBucket, ScyllaReplicationLogEvent.OpType.Deleted, objectBlob);
+            ScyllaReplicationLogEvent log = new ScyllaReplicationLogEvent(ns.ToString(), bucket.ToString(), key.ToString(), timeBucket, ScyllaReplicationLogEvent.OpType.Deleted, null);
             await _mapper.InsertAsync<ScyllaReplicationLogEvent>(log, insertNulls: false,  ttl: (int)_settings.CurrentValue.ReplicationLogTimeToLive.TotalSeconds);
 
             await addNamespaceTask;
@@ -280,7 +280,7 @@ namespace Horde.Storage.Implementation
             Key = null!;
         }
 
-        public ScyllaReplicationLogEvent(string @namespace, string bucket, string key, DateTime lastTimestamp, OpType opType, BlobIdentifier objectIdentifier)
+        public ScyllaReplicationLogEvent(string @namespace, string bucket, string key, DateTime lastTimestamp, OpType opType, BlobIdentifier? objectIdentifier)
         {
             Namespace = @namespace;
             Bucket = bucket;
@@ -288,7 +288,7 @@ namespace Horde.Storage.Implementation
             ReplicationBucket = lastTimestamp.ToHourlyBucket().ToFileTimeUtc();
             ReplicationId = TimeUuid.NewId(lastTimestamp);
             Type = (int)opType;
-            ObjectIdentifier = new ScyllaBlobIdentifier(objectIdentifier);
+            ObjectIdentifier = objectIdentifier != null ? new ScyllaBlobIdentifier(objectIdentifier) : null;
         }
 
         [Cassandra.Mapping.Attributes.PartitionKey]
diff --git a/Engine/Source/Programs/Horde/HordeStorage/Horde.Storage/Tests/Functional/Objects/ReferencesTests.cs b/Engine/Source/Programs/Horde/HordeStorage/Horde.Storage/Tests/Functional/Objects/ReferencesTests.cs
--- a/Engine/Source/Programs/Horde/HordeStorage/Horde.Storage/Tests/Functional/Objects/ReferencesTests.cs
+++ b/Engine/Source/Programs/Horde/HordeStorage/Horde.Storage/Tests/Functional/Objects/ReferencesTests.cs
@@ -1224,12 +1224,12 @@ namespace Horde.Storage.FunctionalTests.References
                 result.EnsureSuccessStatusCode();
             }
 
-            List<ObjectRecord> records = await ReferencesStore.GetOldestRecords(TestNamespace).ToListAsync();
+            List<(BucketId, IoHashKey, DateTime)> records = await ReferencesStore.GetRecords(TestNamespace).ToListAsync();
 
-            ObjectRecord? oldRecord = records.Find(record => record.Name == key);
+            (BucketId, IoHashKey, DateTime)? oldRecord = records.Find(record => record.Item2 == key);
             Assert.IsNotNull(oldRecord);
-            Assert.AreEqual(key, oldRecord!.Name);
-            Assert.AreEqual("bucket", oldRecord.Bucket.ToString());
+            Assert.AreEqual(key, oldRecord.Value.Item2);
+            Assert.AreEqual("bucket", oldRecord.Value.Item1.ToString());
         }
     }
 }
