Horde: Fix incorrect index name for AgentId

diff --git a/Engine/Source/Programs/Horde/Horde.Build/Agents/Impl/LeaseCollection.cs b/Engine/Source/Programs/Horde/Horde.Build/Agents/Impl/LeaseCollection.cs
--- a/Engine/Source/Programs/Horde/Horde.Build/Agents/Impl/LeaseCollection.cs
+++ b/Engine/Source/Programs/Horde/Horde.Build/Agents/Impl/LeaseCollection.cs
@@ -79,7 +79,7 @@ namespace HordeServer.Collections.Impl
 
 			public DatabaseIndexes(IMongoCollection<LeaseDocument> Leases, bool DatabaseReadOnlyMode) : base(DatabaseReadOnlyMode)
 			{
-				AgentId = CreateOrGetIndex(Leases, "StreamId_1", Builders<LeaseDocument>.IndexKeys.Ascending(x => x.AgentId));
+				AgentId = CreateOrGetIndex(Leases, "AgentId_1", Builders<LeaseDocument>.IndexKeys.Ascending(x => x.AgentId));
 				SessionId = CreateOrGetIndex(Leases, "SessionId_1", Builders<LeaseDocument>.IndexKeys.Ascending(x => x.SessionId));
 				StartTime = CreateOrGetIndex(Leases, "StartTime_1", Builders<LeaseDocument>.IndexKeys.Ascending(x => x.StartTime));
 				FinishTime = CreateOrGetIndex(Leases, "FinishTime_1", Builders<LeaseDocument>.IndexKeys.Ascending(x => x.FinishTime));
