Horde Storage - If encountering a key that is not a io hash we skip it (we did not use to enforce that keys were io hashes so we have some legacy keys in the db)

diff --git a/Engine/Source/Programs/Horde/HordeStorage/Horde.Storage/Implementation/Objects/ScyllaReferencesStore.cs b/Engine/Source/Programs/Horde/HordeStorage/Horde.Storage/Implementation/Objects/ScyllaReferencesStore.cs
--- a/Engine/Source/Programs/Horde/HordeStorage/Horde.Storage/Implementation/Objects/ScyllaReferencesStore.cs
+++ b/Engine/Source/Programs/Horde/HordeStorage/Horde.Storage/Implementation/Objects/ScyllaReferencesStore.cs
@@ -209,6 +209,9 @@ namespace Horde.Storage.Implementation
                 string name = row.GetValue<string>("name");
                 DateTime lastAccessTime = row.GetValue<DateTime>("last_access_time");
 
+                // skip any names that are not conformant to io hash
+                if (name.Length != 40)
+                    continue;
                 yield return (new BucketId(bucket), new IoHashKey(name), lastAccessTime);
             }
         }
