Horde: Fix issue with job controller returning 500 on missing GraphHash
#rnx

diff --git a/Engine/Source/Programs/Horde/Horde.Build/Controllers/JobsController.cs b/Engine/Source/Programs/Horde/Horde.Build/Controllers/JobsController.cs
--- a/Engine/Source/Programs/Horde/Horde.Build/Controllers/JobsController.cs
+++ b/Engine/Source/Programs/Horde/Horde.Build/Controllers/JobsController.cs
@@ -694,7 +694,13 @@ namespace HordeServer.Controllers
 			{
 				using IScope JobScope = GlobalTracer.Instance.BuildSpan("JobIteration").StartActive();
 				JobScope.Span.SetTag("jobId", Job.Id.ToString());
-				
+
+				if (Job.GraphHash == null)
+				{
+					Logger.LogWarning("Job {JobId} has no GraphHash", Job.Id);
+					continue;
+				}
+
 				bool ViewJobAuthorized;
 				using (IScope _ = GlobalTracer.Instance.BuildSpan("AuthorizeViewJob").StartActive())
 				{
