Horde: Ensure that all leases terminate before restarting an agent session.

Leases will be terminated automatically on the server side, but leaving untracked leases running on agents will cause a hang on shutdown while waiting for them to complete.

Also add a rate limiter to UpdateSession requests; if more than 60 calls are made in a minute, output a warning and wait 10 seconds until continuing.

diff --git a/Engine/Source/Programs/Horde/Horde.Agent/Services/WorkerService.cs b/Engine/Source/Programs/Horde/Horde.Agent/Services/WorkerService.cs
--- a/Engine/Source/Programs/Horde/Horde.Agent/Services/WorkerService.cs
+++ b/Engine/Source/Programs/Horde/Horde.Agent/Services/WorkerService.cs
@@ -298,20 +298,28 @@ namespace HordeAgent.Services
 				Stopwatch SessionTime = Stopwatch.StartNew();
 				try
 				{
-					await HandleSessionAsync(StoppingToken);
+					if (ActiveLeases.Count == 0)
+					{
+						await HandleSessionAsync(StoppingToken);
+					}
 				}
 				catch (Exception Ex)
 				{
-					if (StoppingToken.IsCancellationRequested && IsCancellationException(Ex))
-					{
-						break;
-					}
-					else
+					if (!StoppingToken.IsCancellationRequested || !IsCancellationException(Ex))
 					{
 						Logger.LogError(Ex, "Exception while executing session. Restarting.");
 					}
 				}
 
+				try
+				{
+					await DrainLeasesAsync();
+				}
+				catch (Exception Ex)
+				{
+					Logger.LogError(Ex, "Exception while draining leases. Agent may be in an inconsistent state.");
+				}
+
 				if (RequestShutdown)
 				{
 					Logger.LogInformation("Initiating shutdown (restart={Restart})", RestartAfterShutdown);
@@ -353,6 +361,56 @@ namespace HordeAgent.Services
 			}
 		}
 
+		async Task DrainLeasesAsync()
+		{
+			for (int Idx = 0; Idx < ActiveLeases.Count; Idx++)
+			{
+				LeaseInfo ActiveLease = ActiveLeases[Idx];
+				if (ActiveLease.Task == null)
+				{
+					ActiveLeases.RemoveAt(Idx--);
+					Logger.LogInformation("Removed lease {LeaseId}", ActiveLease.Lease.Id);
+				}
+				else
+				{
+					Logger.LogInformation("Cancelling active lease {LeaseId}", ActiveLease.Lease.Id);
+					ActiveLease.CancellationTokenSource.Cancel();
+				}
+			}
+
+			while (ActiveLeases.Count > 0)
+			{
+				List<Task> Tasks = ActiveLeases.Select(x => x.Task!).ToList();
+				Tasks.Add(Task.Delay(TimeSpan.FromMinutes(1.0)));
+				await Task.WhenAny(Tasks);
+
+				for (int Idx = 0; Idx < ActiveLeases.Count; Idx++)
+				{
+					LeaseInfo ActiveLease = ActiveLeases[Idx];
+					if (ActiveLease.Task!.IsCompleted)
+					{
+						ActiveLeases.RemoveAt(Idx--);
+						try
+						{
+							await ActiveLease.Task;
+						}
+						catch (OperationCanceledException)
+						{
+						}
+						catch (Exception Ex)
+						{
+							Logger.LogError(Ex, "Lease {LeaseId} threw an exception while terminating", ActiveLease.Lease.Id);
+						}
+						Logger.LogInformation("Lease {LeaseId} has completed", ActiveLease.Lease.Id);
+					}
+					else
+					{
+						Logger.LogInformation("Still waiting for lease {LeaseId} to terminate...", ActiveLease.Lease.Id);
+					}
+				}
+			}
+		}
+
 		/// <summary>
 		/// Handles the lifetime of an agent session.
 		/// </summary>
@@ -420,6 +478,10 @@ namespace HordeAgent.Services
 			// Open a connection to the server
 			await using (IRpcConnection RpcCon = RpcConnection.Create(() => GrpcService.CreateGrpcChannel(CreateSessionResponse.Token), Logger))
 			{
+				// Track how many updates we get in 10 seconds. We'll start rate limiting this if it looks like we've got a problem that's causing us to spam the server.
+				Stopwatch UpdateTimer = Stopwatch.StartNew();
+				Queue<TimeSpan> UpdateTimes = new Queue<TimeSpan>();
+
 				// Loop until we're ready to exit
 				Stopwatch UpdateCapabilitiesTimer = Stopwatch.StartNew();
 				for (; ; )
@@ -461,7 +523,6 @@ namespace HordeAgent.Services
 						UpdateSessionRequest.Status = AgentStatus.Ok;
 					}
 
-
 					// Update the capabilities every 5m
 					if (UpdateCapabilitiesTimer.Elapsed > TimeSpan.FromMinutes(5.0))
 					{
@@ -526,6 +587,21 @@ namespace HordeAgent.Services
 							}
 						}
 					}
+
+					// Update the historical update times
+					TimeSpan UpdateTime = UpdateTimer.Elapsed;
+					while (UpdateTimes.TryPeek(out TimeSpan FirstTime) && FirstTime < UpdateTime - TimeSpan.FromMinutes(1.0))
+					{
+						UpdateTimes.Dequeue();
+					}
+					UpdateTimes.Enqueue(UpdateTime);
+
+					// If we're updating too much, introduce an artificial delay
+					if (UpdateTimes.Count > 60)
+					{
+						Logger.LogWarning("Agent is issuing large number of UpdateSession() calls. Delaying for 10 seconds.");
+						await Task.Delay(TimeSpan.FromSeconds(10.0));
+					}
 				}
 
 				Logger.LogInformation("Disposing RpcConnection");
