Horde Storage - Updated last access time on the object rather then the auxillary table to match how getObjects work

diff --git a/Engine/Source/Programs/Horde/HordeStorage/Horde.Storage/Implementation/Objects/ScyllaReferencesStore.cs b/Engine/Source/Programs/Horde/HordeStorage/Horde.Storage/Implementation/Objects/ScyllaReferencesStore.cs
--- a/Engine/Source/Programs/Horde/HordeStorage/Horde.Storage/Implementation/Objects/ScyllaReferencesStore.cs
+++ b/Engine/Source/Programs/Horde/HordeStorage/Horde.Storage/Implementation/Objects/ScyllaReferencesStore.cs
@@ -122,7 +122,7 @@ namespace Horde.Storage.Implementation
         {
             using Scope _ = Tracer.Instance.StartActive("scylla.update_last_access_time");
             // fetch the old record
-            ObjectRecord objectRecord = await Get(ns, bucket, name);
+            /* ObjectRecord objectRecord = await Get(ns, bucket, name);
 
             // update the object tracking
             Task updateObjectTracking = _mapper.UpdateAsync<ScyllaObject>("SET last_access_time = ? WHERE namespace=? AND bucket = ? AND name = ?", lastAccessTime, ns.ToString(), bucket.ToString(), name.ToString());
@@ -138,7 +138,11 @@ namespace Horde.Storage.Implementation
                 await _mapper.UpdateAsync<ScyllaObjectLastAccessAt>("SET objects = objects + ? WHERE accessed_at=? AND namespace = ? AND partition_index = ?", new string[] {$"{bucket}#{name}"}, newLastAccessAt, ns.ToString(), partitionIndex);
             }
 
-            await updateObjectTracking;
+            await updateObjectTracking;*/
+
+            await _mapper.UpdateAsync<ScyllaObjectReference>("SET last_access_time = ? WHERE namespace = ? AND bucket = ? AND name = ?", lastAccessTime, ns.ToString(), bucket.ToString(), name.ToString());
+
+            
         }
 
         /*public async IAsyncEnumerable<ObjectRecord> GetRecords(NamespaceId ns)
