Horde: Add a 'stopped' state for agents, to prevent queries against the state (such as autoscaling) returning incorrectly when a session is not valid.

diff --git a/Engine/Source/Programs/Horde/Horde.Build/Agents/Impl/AgentCollection.cs b/Engine/Source/Programs/Horde/Horde.Build/Agents/Impl/AgentCollection.cs
--- a/Engine/Source/Programs/Horde/Horde.Build/Agents/Impl/AgentCollection.cs
+++ b/Engine/Source/Programs/Horde/Horde.Build/Agents/Impl/AgentCollection.cs
@@ -487,6 +487,7 @@ namespace HordeServer.Collections.Impl
 			Update = Update.Unset(x => x.SessionId);
 			Update = Update.Unset(x => x.SessionExpiresAt);
 			Update = Update.Unset(x => x.Leases);
+			Update = Update.Set(x => x.Status, AgentStatus.Stopped);
 
 			bool bDeleted = Agent.Deleted || Agent.Ephemeral;
 			if (bDeleted != Agent.Deleted)
