Horde: Set the correct batch state when a job is aborted. Also add a new state for when a lease fails to execute in an unrecoverable.

diff --git a/Engine/Source/Programs/Horde/Horde.Build/Api/Jobs.cs b/Engine/Source/Programs/Horde/Horde.Build/Api/Jobs.cs
--- a/Engine/Source/Programs/Horde/Horde.Build/Api/Jobs.cs
+++ b/Engine/Source/Programs/Horde/Horde.Build/Api/Jobs.cs
@@ -661,9 +661,14 @@ namespace HordeServer.Api
 		LostConnection = 8,
 
 		/// <summary>
-		/// Lease terminated prematurely
+		/// Lease terminated prematurely but can be retried.
 		/// </summary>
 		Incomplete = 9,
+
+		/// <summary>
+		/// An error ocurred while executing the lease. Cannot be retried.
+		/// </summary>
+		ExecutionError = 10,
 	}
 
 	/// <summary>
diff --git a/Engine/Source/Programs/Horde/Horde.Build/Jobs/IJobCollection.cs b/Engine/Source/Programs/Horde/Horde.Build/Jobs/IJobCollection.cs
--- a/Engine/Source/Programs/Horde/Horde.Build/Jobs/IJobCollection.cs
+++ b/Engine/Source/Programs/Horde/Horde.Build/Jobs/IJobCollection.cs
@@ -228,8 +228,9 @@ namespace HordeServer.Collections
 		/// <param name="Job">The job to update</param>
 		/// <param name="BatchIdx">Index of the batch to cancel</param>
 		/// <param name="Graph">Graph for the job</param>
+		/// <param name="Reason">Reason for this batch being failed</param>
 		/// <returns>True if the job is updated</returns>
-		Task<IJob?> TryFailBatchAsync(IJob Job, int BatchIdx, IGraph Graph);
+		Task<IJob?> TryFailBatchAsync(IJob Job, int BatchIdx, IGraph Graph, JobStepBatchError Reason);
 
 		/// <summary>
 		/// Attempt to assign a lease to execute a batch
diff --git a/Engine/Source/Programs/Horde/Horde.Build/Jobs/Impl/JobCollection.cs b/Engine/Source/Programs/Horde/Horde.Build/Jobs/Impl/JobCollection.cs
--- a/Engine/Source/Programs/Horde/Horde.Build/Jobs/Impl/JobCollection.cs
+++ b/Engine/Source/Programs/Horde/Horde.Build/Jobs/Impl/JobCollection.cs
@@ -1217,7 +1217,7 @@ namespace HordeServer.Collections.Impl
 		}
 
 		/// <inheritdoc/>
-		public Task<IJob?> TryFailBatchAsync(IJob Job, int BatchIdx, IGraph Graph)
+		public Task<IJob?> TryFailBatchAsync(IJob Job, int BatchIdx, IGraph Graph, JobStepBatchError Error)
 		{
 			JobDocument JobDocument = Clone((JobDocument)Job);
 			JobStepBatchDocument Batch = JobDocument.Batches[BatchIdx];
@@ -1230,7 +1230,7 @@ namespace HordeServer.Collections.Impl
 				Batch.State = JobStepBatchState.Complete;
 				Updates.Add(UpdateBuilder.Set(x => x.Batches[BatchIdx].State, Batch.State));
 
-				Batch.Error = JobStepBatchError.LostConnection;
+				Batch.Error = Error;
 				Updates.Add(UpdateBuilder.Set(x => x.Batches[BatchIdx].Error, Batch.Error));
 
 				Batch.FinishTime = DateTimeOffset.UtcNow;
diff --git a/Engine/Source/Programs/Horde/Horde.Build/Jobs/Impl/JobTaskSource.cs b/Engine/Source/Programs/Horde/Horde.Build/Jobs/Impl/JobTaskSource.cs
--- a/Engine/Source/Programs/Horde/Horde.Build/Jobs/Impl/JobTaskSource.cs
+++ b/Engine/Source/Programs/Horde/Horde.Build/Jobs/Impl/JobTaskSource.cs
@@ -1026,8 +1026,18 @@ namespace HordeServer.Tasks.Impl
 
 					int RunningStepIdx = Batch.Steps.FindIndex(x => x.State == JobStepState.Running);
 
+					JobStepBatchError Error;
+					if (Outcome == LeaseOutcome.Cancelled)
+					{
+						Error = JobStepBatchError.Cancelled;
+					}
+					else
+					{
+						Error = JobStepBatchError.ExecutionError;
+					}
+
 					IGraph Graph = await Graphs.GetAsync(Job.GraphHash);
-					Job = await Jobs.TryFailBatchAsync(Job, BatchIdx, Graph);
+					Job = await Jobs.TryFailBatchAsync(Job, BatchIdx, Graph, Error);
 
 					if (Job != null)
 					{
