Horde: Adding execution error handling to job batches
#rnx

diff --git a/Engine/Source/Programs/Horde/HordeDashboard/src/backend/Api.ts b/Engine/Source/Programs/Horde/HordeDashboard/src/backend/Api.ts
--- a/Engine/Source/Programs/Horde/HordeDashboard/src/backend/Api.ts
+++ b/Engine/Source/Programs/Horde/HordeDashboard/src/backend/Api.ts
@@ -116,7 +116,11 @@ export enum JobStepBatchError {
 	LostConnection = "LostConnection",
 
 	/** Lease terminated prematurely */
-	Incomplete = "Incomplete"
+	Incomplete = "Incomplete",
+
+	/** An error ocurred while executing the lease. Cannot be retried. */
+	ExecutionError = "ExecutionError"
+
 
 }
 
