Horde: Possible fix for nondeterministic tests.

diff --git a/Engine/Source/Programs/Horde/Horde.Build.Tests/SchedulerTests.cs b/Engine/Source/Programs/Horde/Horde.Build.Tests/SchedulerTests.cs
--- a/Engine/Source/Programs/Horde/Horde.Build.Tests/SchedulerTests.cs
+++ b/Engine/Source/Programs/Horde/Horde.Build.Tests/SchedulerTests.cs
@@ -425,9 +425,11 @@ namespace HordeServerTests
 
 			// Create a job and fail it
 			IJob Job1 = await JobService.CreateJobAsync(null, Stream, NewTemplateRefId1, Template.Id, GraphA, "Hello", 1234, 1233, 999, null, null, null, null, null, null, true, true, null, null, new List<string> { "-Target=TriggerNext" });
-			Job1 = Deref(await JobService.UpdateBatchAsync(Job1, Job1.Batches[0].Id, LogId.GenerateNewId(), JobStepBatchState.Running));
-			Assert.IsNotNull(await JobService.UpdateStepAsync(Job1, Job1.Batches[0].Id, Job1.Batches[0].Steps[0].Id, JobStepState.Completed, JobStepOutcome.Failure));
-			Job1 = Deref(await JobService.UpdateBatchAsync(Job1, Job1.Batches[0].Id, LogId.GenerateNewId(), JobStepBatchState.Complete));
+			SubResourceId BatchId1 = Job1.Batches[0].Id;
+			SubResourceId StepId1 = Job1.Batches[0].Steps[0].Id;
+			Job1 = Deref(await JobService.UpdateBatchAsync(Job1, BatchId1, LogId.GenerateNewId(), JobStepBatchState.Running));
+			Job1 = Deref(await JobService.UpdateStepAsync(Job1, BatchId1, StepId1, JobStepState.Completed, JobStepOutcome.Failure));
+			Job1 = Deref(await JobService.UpdateBatchAsync(Job1, BatchId1, LogId.GenerateNewId(), JobStepBatchState.Complete));
 			await GetNewJobs();
 
 			// Tick the schedule and make sure it doesn't trigger
@@ -438,8 +440,10 @@ namespace HordeServerTests
 
 			// Create a job and make it succeed
 			IJob Job2 = await JobService.CreateJobAsync(null, Stream, NewTemplateRefId1, Template.Id, GraphA, "Hello", 1234, 1233, 999, null, null, null, null, null, null, true, true, null, null, new List<string> { "-Target=TriggerNext" });
-			Job2 = Deref(await JobService.UpdateBatchAsync(Job2, Job2.Batches[0].Id, LogId.GenerateNewId(), JobStepBatchState.Running));
-			Assert.IsNotNull(await JobService.UpdateStepAsync(Job2, Job2.Batches[0].Id, Job2.Batches[0].Steps[0].Id, JobStepState.Completed, JobStepOutcome.Success));
+			SubResourceId BatchId2 = Job2.Batches[0].Id;
+			SubResourceId StepId2 = Job2.Batches[0].Steps[0].Id;
+			Job2 = Deref(await JobService.UpdateBatchAsync(Job2, BatchId2, LogId.GenerateNewId(), JobStepBatchState.Running));
+			Job2 = Deref(await JobService.UpdateStepAsync(Job2, BatchId2, StepId2, JobStepState.Completed, JobStepOutcome.Success));
 
 			// Tick the schedule and make sure it does trigger
 			await ScheduleService.TickForTestingAsync();
