Horde: Disable log spam from TestTracer.

diff --git a/Engine/Source/Programs/Horde/Horde.Build/Program.cs b/Engine/Source/Programs/Horde/Horde.Build/Program.cs
--- a/Engine/Source/Programs/Horde/Horde.Build/Program.cs
+++ b/Engine/Source/Programs/Horde/Horde.Build/Program.cs
@@ -184,7 +184,7 @@ namespace HordeServer
 			{
 				using var _ = Datadog.Trace.Tracer.Instance.StartActive("Trace Test");
 				Serilog.Log.Logger.Information("Enabling datadog tracing");
-				GlobalTracer.Register(new TestTracer(Datadog.Trace.OpenTracing.OpenTracingTracerFactory.WrapTracer(Datadog.Trace.Tracer.Instance)));
+				GlobalTracer.Register(Datadog.Trace.OpenTracing.OpenTracingTracerFactory.WrapTracer(Datadog.Trace.Tracer.Instance));
 				using IScope Scope = GlobalTracer.Instance.BuildSpan("OpenTrace Test").StartActive();
 				Scope.Span.SetTag("TestProp", "hello");
 				Serilog.Log.Logger.Information("Enabling datadog tracing (OpenTrace)");
