Horde Storage - Added debug logging for ref cleanup and a option to set the logging level via the config map.

diff --git a/Engine/Source/Programs/Horde/HordeStorage/Horde.Storage/Implementation/GC/CacheCleanup.cs b/Engine/Source/Programs/Horde/HordeStorage/Horde.Storage/Implementation/GC/CacheCleanup.cs
--- a/Engine/Source/Programs/Horde/HordeStorage/Horde.Storage/Implementation/GC/CacheCleanup.cs
+++ b/Engine/Source/Programs/Horde/HordeStorage/Horde.Storage/Implementation/GC/CacheCleanup.cs
@@ -56,11 +56,16 @@ namespace Horde.Storage.Implementation
         {
             List<OldRecord> deletedRecords = new List<OldRecord>();
             DateTime cutoffTime = DateTime.Now.Add(_settings.CurrentValue.LastAccessCutoff);
+            int consideredCount = 0;
             await foreach ((BucketId bucket, IoHashKey name, DateTime lastAccessTime) in _referencesStore.GetRecords(ns).WithCancellation(cancellationToken))
             {
+                _logger.Debug("Considering object in {Namespace} {Bucket} {Name} for deletion, was last updated {LastAccessTime}", ns, bucket, name, lastAccessTime);
+                Interlocked.Increment(ref consideredCount);
+
                 if (lastAccessTime > cutoffTime)
                     continue;
 
+                _logger.Debug("Attempting to delete object {Namespace} {Bucket} {Name}.", ns, bucket, name);
                 using Scope scope = Tracer.Instance.StartActive("refCleanup.delete_record");
                 scope.Span.ResourceName = $"{ns}:{bucket}.{name}";
                 // delete the old record from the ref refs
@@ -80,6 +85,8 @@ namespace Horde.Storage.Implementation
                 deletedRecords.Add(new OldRecord(ns, bucket, new KeyId(name.ToString())));
             }
 
+            _logger.Information("Finished cleaning {Namespace}. Refs considered: {ConsideredCount} Refs Deleted: {DeletedCount}", ns, consideredCount, deletedRecords.Count);
+
             return deletedRecords;
         }
 
