Horde Storage - Fixed issue were the last access time update was trying to use a scylla table that did not exist (due to the wrong type being mapped)

diff --git a/Engine/Source/Programs/Horde/HordeStorage/Horde.Storage/Implementation/Objects/ScyllaReferencesStore.cs b/Engine/Source/Programs/Horde/HordeStorage/Horde.Storage/Implementation/Objects/ScyllaReferencesStore.cs
--- a/Engine/Source/Programs/Horde/HordeStorage/Horde.Storage/Implementation/Objects/ScyllaReferencesStore.cs
+++ b/Engine/Source/Programs/Horde/HordeStorage/Horde.Storage/Implementation/Objects/ScyllaReferencesStore.cs
@@ -140,7 +140,7 @@ namespace Horde.Storage.Implementation
 
             await updateObjectTracking;*/
 
-            await _mapper.UpdateAsync<ScyllaObjectReference>("SET last_access_time = ? WHERE namespace = ? AND bucket = ? AND name = ?", lastAccessTime, ns.ToString(), bucket.ToString(), name.ToString());
+            await _mapper.UpdateAsync<ScyllaObject>("SET last_access_time = ? WHERE namespace = ? AND bucket = ? AND name = ?", lastAccessTime, ns.ToString(), bucket.ToString(), name.ToString());
 
             
         }
