Horde Storage
Removed error logging that is not very useful

We used to have a issue were we had broken blobs and this was added to help capture that. Technically a blob can be 0 bytes, there is nothing wrong with that.
We now detect issues like this using the consitency checker instead.

diff --git a/Engine/Source/Programs/Horde/HordeStorage/Jupiter.Common/Implementation/BatchWriter.cs b/Engine/Source/Programs/Horde/HordeStorage/Jupiter.Common/Implementation/BatchWriter.cs
--- a/Engine/Source/Programs/Horde/HordeStorage/Jupiter.Common/Implementation/BatchWriter.cs
+++ b/Engine/Source/Programs/Horde/HordeStorage/Jupiter.Common/Implementation/BatchWriter.cs
@@ -65,11 +65,6 @@ namespace Jupiter.Implementation
                     if (opState != OpState.OK || opState == OpState.Exists || contentHash == null || blob == null)
                         continue;
 
-                    if (blob.Length == 0)
-                    {
-                        _logger.Error("0 byte blob encountered for hash {ContentHash} state was {OpState}. Skipping", contentHash.ToString(), opState);
-                    }
-
                     await outStream.WriteAsync(contentHash.HashData);
                     await outStream.WriteAsync(BitConverter.GetBytes((ulong) blob.Length));
                     await blob.Stream.CopyToAsync(outStream);
