Make sure authors get notified as soon as a job step fails or warns.

#rnx

diff --git a/Engine/Source/Programs/Horde/Horde.Build/Notifications/Impl/NotificationService.cs b/Engine/Source/Programs/Horde/Horde.Build/Notifications/Impl/NotificationService.cs
--- a/Engine/Source/Programs/Horde/Horde.Build/Notifications/Impl/NotificationService.cs
+++ b/Engine/Source/Programs/Horde/Horde.Build/Notifications/Impl/NotificationService.cs
@@ -475,7 +475,19 @@ namespace HordeServer.Notifications.Impl
 			EventRecord EventRecord = new StepCompleteEventRecord(Job.StreamId, Job.TemplateId, Node.Name, Step.Outcome);
 
 			List<IUser> UsersToNotify = await GetUsersToNotify(EventRecord, Step.NotificationTriggerId, true);
-			if(UsersToNotify.Count == 0)
+
+			// If this is not a success notification and the author isn't in the list to notify, add them manually.
+			if (Job.StartedByUserId.HasValue && !UsersToNotify.Any(x => x.Id == Job.StartedByUserId) && Step.Outcome != JobStepOutcome.Success)
+			{
+				Logger.LogInformation("Author {AuthorUserId} is not in notify list but step outcome is {JobStepOutcome}, adding them to the list...", Job.StartedByUserId, Step.Outcome);
+				IUser? AuthorUser = await UserCollection.GetUserAsync(Job.StartedByUserId.Value);
+				if (AuthorUser != null)
+				{
+					UsersToNotify.Add(AuthorUser);
+				}
+			}
+
+			if (UsersToNotify.Count == 0)
 			{
 				Logger.LogInformation("No users to notify for step {JobId}:{BatchId}:{StepId}", Job.Id, BatchId, StepId);
 				return;
