Catch exceptions when attempting to grow/shrink a pool

Allows for Autoscaling to continue to check other pools when updating one fails

#rnx

diff --git a/Engine/Source/Programs/Horde/Horde.Build/Fleet/AutoscaleService.cs b/Engine/Source/Programs/Horde/Horde.Build/Fleet/AutoscaleService.cs
--- a/Engine/Source/Programs/Horde/Horde.Build/Fleet/AutoscaleService.cs
+++ b/Engine/Source/Programs/Horde/Horde.Build/Fleet/AutoscaleService.cs
@@ -210,35 +210,43 @@ namespace HordeServer.Services
 					TargetAgents,
 					Delta);
 
-				if (Pool.EnableAutoscaling)
+				try
 				{
-					using (IScope Scope = GlobalTracer.Instance.BuildSpan("Scaling pool").StartActive())
+					if (Pool.EnableAutoscaling)
 					{
-						Scope.Span.SetTag("poolName", Pool.Name);
-						Scope.Span.SetTag("delta", Delta);
-
-						if (Delta > 0)
+						using (IScope Scope = GlobalTracer.Instance.BuildSpan("Scaling pool").StartActive())
 						{
-							await FleetManager.ExpandPool(Pool, PoolData.Agents, Delta);
-							await PoolCollection.TryUpdateAsync(Pool, LastScaleUpTime: DateTime.UtcNow);
-						}
+							Scope.Span.SetTag("poolName", Pool.Name);
+							Scope.Span.SetTag("delta", Delta);
 
-						if (Delta < 0)
-						{
-							bool bShrinkIsOnCoolDown = Pool.LastScaleDownTime != null && Pool.LastScaleDownTime + ShrinkPoolCoolDown > DateTime.UtcNow;
-							if (!bShrinkIsOnCoolDown)
+							if (Delta > 0)
 							{
-								await FleetManager.ShrinkPool(Pool, PoolData.Agents, -Delta);
-								await PoolCollection.TryUpdateAsync(Pool, LastScaleDownTime: DateTime.UtcNow);
+								await FleetManager.ExpandPool(Pool, PoolData.Agents, Delta);
+								await PoolCollection.TryUpdateAsync(Pool, LastScaleUpTime: DateTime.UtcNow);
 							}
-							else
+
+							if (Delta < 0)
 							{
-								Logger.LogDebug("Cannot shrink {PoolName} right now, it's on cool-down until {CoolDownTimeEnds}", Pool.Name, Pool.LastScaleDownTime + ShrinkPoolCoolDown);
+								bool bShrinkIsOnCoolDown = Pool.LastScaleDownTime != null && Pool.LastScaleDownTime + ShrinkPoolCoolDown > DateTime.UtcNow;
+								if (!bShrinkIsOnCoolDown)
+								{
+									await FleetManager.ShrinkPool(Pool, PoolData.Agents, -Delta);
+									await PoolCollection.TryUpdateAsync(Pool, LastScaleDownTime: DateTime.UtcNow);
+								}
+								else
+								{
+									Logger.LogDebug("Cannot shrink {PoolName} right now, it's on cool-down until {CoolDownTimeEnds}", Pool.Name, Pool.LastScaleDownTime + ShrinkPoolCoolDown);
+								}
 							}
 						}
 					}
 				}
-				
+				catch (Exception Ex)
+				{
+					Logger.LogError(Ex, "Failed to scale {PoolName}:\n{Exception}", Pool.Name, Ex);
+					continue;
+				}
+
 				DogStatsd.Gauge("agentpools.autoscale.target", TargetAgents, tags: new []{"pool:" + Pool.Name});
 				DogStatsd.Gauge("agentpools.autoscale.current", PoolData.Agents.Count, tags: new []{"pool:" + Pool.Name});
 			}
