Horde Storage

Scylla Reference lookups can now exclude the inline payload, which is quite large. This allows us to speed up some internal queries and the HEAD checks that do not need this potentially large object.

diff --git a/Engine/Source/Programs/Horde/HordeStorage/Horde.Storage/Controllers/ReferencesController.cs b/Engine/Source/Programs/Horde/HordeStorage/Horde.Storage/Controllers/ReferencesController.cs
--- a/Engine/Source/Programs/Horde/HordeStorage/Horde.Storage/Controllers/ReferencesController.cs
+++ b/Engine/Source/Programs/Horde/HordeStorage/Horde.Storage/Controllers/ReferencesController.cs
@@ -112,9 +112,12 @@ namespace Horde.Storage.Controllers
                 // TODO: Do we want to keep fields? we have no real need for it now
                 (ObjectRecord objectRecord, BlobContents? blob) = await _objectService.Get(ns, bucket, key, fields);
 
+                if (blob == null)
+                    throw new InvalidOperationException($"Blob was null when attempting to fetch {ns} {bucket} {key}");
+
                 if (!objectRecord.IsFinalized)
                 {
-                    // we do not considered un-finalized objects are valid
+                    // we do not consider un-finalized objects as valid
                     return BadRequest(new ProblemDetails { Title = $"Object {objectRecord.Bucket} {objectRecord.Name} is not finalized." });
                 }
 
diff --git a/Engine/Source/Programs/Horde/HordeStorage/Horde.Storage/Implementation/GC/OrphanBlobCleanupRefs.cs b/Engine/Source/Programs/Horde/HordeStorage/Horde.Storage/Implementation/GC/OrphanBlobCleanupRefs.cs
--- a/Engine/Source/Programs/Horde/HordeStorage/Horde.Storage/Implementation/GC/OrphanBlobCleanupRefs.cs
+++ b/Engine/Source/Programs/Horde/HordeStorage/Horde.Storage/Implementation/GC/OrphanBlobCleanupRefs.cs
@@ -82,7 +82,7 @@ namespace Horde.Storage.Implementation
 
                                 try
                                 {
-                                    (ObjectRecord, BlobContents) _ = await _objectService.Get(blobNamespace, bucket, key, Array.Empty<string>());
+                                    (ObjectRecord, BlobContents?) _ = await _objectService.Get(blobNamespace, bucket, key, new string[] {"name"});
                                     found = true;
                                 }
                                 catch (ObjectNotFoundException)
diff --git a/Engine/Source/Programs/Horde/HordeStorage/Horde.Storage/Implementation/ObjectService.cs b/Engine/Source/Programs/Horde/HordeStorage/Horde.Storage/Implementation/ObjectService.cs
--- a/Engine/Source/Programs/Horde/HordeStorage/Horde.Storage/Implementation/ObjectService.cs
+++ b/Engine/Source/Programs/Horde/HordeStorage/Horde.Storage/Implementation/ObjectService.cs
@@ -15,7 +15,7 @@ namespace Horde.Storage.Implementation
 {
     public interface IObjectService
     {
-        Task<(ObjectRecord, BlobContents)> Get(NamespaceId ns, BucketId bucket, IoHashKey key, string[] fields);
+        Task<(ObjectRecord, BlobContents?)> Get(NamespaceId ns, BucketId bucket, IoHashKey key, string[] fields);
         Task<PutObjectResult> Put(NamespaceId ns, BucketId bucket, IoHashKey key, BlobIdentifier blobHash, CompactBinaryObject payload);
         Task<BlobIdentifier[]> Finalize(NamespaceId ns, BucketId bucket, IoHashKey key, BlobIdentifier blobHash);
 
@@ -45,22 +45,40 @@ namespace Horde.Storage.Implementation
             _lastAccessTracker = lastAccessTracker;
         }
 
-        public async Task<(ObjectRecord, BlobContents)> Get(NamespaceId ns, BucketId bucket, IoHashKey key, string[]? fields = null)
+        public async Task<(ObjectRecord, BlobContents?)> Get(NamespaceId ns, BucketId bucket, IoHashKey key, string[]? fields = null)
         {
             // we do not wait for the last access tracking as it does not matter when it completes
             Task lastAccessTask = _lastAccessTracker.TrackUsed(new LastAccessRecord(ns, bucket, key));
 
-            ObjectRecord o = await _referencesStore.Get(ns, bucket, key);
-
-            BlobContents blobContents;
-            if (o.InlinePayload != null && o.InlinePayload.Length != 0)
+            // if no field filtering is being used we assume everything is needed
+            IReferencesStore.FieldFlags flags = IReferencesStore.FieldFlags.All;
+            if (fields != null)
             {
-                blobContents = new BlobContents(o.InlinePayload);
+                // empty array means fetch all fields
+                if (fields.Length == 0)
+                    flags = IReferencesStore.FieldFlags.All;
+                else
+                {
+                    flags = fields.Contains("payload")
+                        ? IReferencesStore.FieldFlags.IncludePayload
+                        : IReferencesStore.FieldFlags.None;
+                }
             }
-            else
+            ObjectRecord o = await _referencesStore.Get(ns, bucket, key, flags);
+
+            BlobContents? blobContents = null;
+            if ((flags & IReferencesStore.FieldFlags.IncludePayload) != 0)
             {
-                blobContents = await _blobService.GetObject(ns, o.BlobIdentifier);
+                if (o.InlinePayload != null && o.InlinePayload.Length != 0)
+                {
+                    blobContents = new BlobContents(o.InlinePayload);
+                }
+                else
+                {
+                    blobContents = await _blobService.GetObject(ns, o.BlobIdentifier);
+                }
             }
+
             return (o, blobContents);
         }
 
@@ -108,7 +126,9 @@ namespace Horde.Storage.Implementation
 
         public async Task<BlobIdentifier[]> Finalize(NamespaceId ns, BucketId bucket, IoHashKey key, BlobIdentifier blobHash)
         {
-            (ObjectRecord o, BlobContents blob) = await Get(ns, bucket, key);
+            (ObjectRecord o, BlobContents? blob) = await Get(ns, bucket, key);
+            if (blob == null)
+                throw new InvalidOperationException("No blob when attempting to finalize");
             byte[] blobContents = await blob.Stream.ToByteArray();
             CompactBinaryObject payload = CompactBinaryObject.Load(blobContents);
 
diff --git a/Engine/Source/Programs/Horde/HordeStorage/Horde.Storage/Implementation/Objects/IReferencesStore.cs b/Engine/Source/Programs/Horde/HordeStorage/Horde.Storage/Implementation/Objects/IReferencesStore.cs
--- a/Engine/Source/Programs/Horde/HordeStorage/Horde.Storage/Implementation/Objects/IReferencesStore.cs
+++ b/Engine/Source/Programs/Horde/HordeStorage/Horde.Storage/Implementation/Objects/IReferencesStore.cs
@@ -9,7 +9,16 @@ namespace Horde.Storage.Implementation
 {
     public interface IReferencesStore
     {
-        Task<ObjectRecord> Get(NamespaceId ns, BucketId bucket, IoHashKey key);
+        Task<ObjectRecord> Get(NamespaceId ns, BucketId bucket, IoHashKey key, FieldFlags flags);
+
+        [Flags]
+        public enum FieldFlags
+        {
+            None = 0,
+            IncludePayload = 1,
+            All = IncludePayload
+        }
+
         Task Put(NamespaceId ns, BucketId bucket, IoHashKey key, BlobIdentifier blobHash, byte[] blob, bool isFinalized);
         Task Finalize(NamespaceId ns, BucketId bucket, IoHashKey key);
 
diff --git a/Engine/Source/Programs/Horde/HordeStorage/Horde.Storage/Implementation/Objects/MemoryReferencesStore.cs b/Engine/Source/Programs/Horde/HordeStorage/Horde.Storage/Implementation/Objects/MemoryReferencesStore.cs
--- a/Engine/Source/Programs/Horde/HordeStorage/Horde.Storage/Implementation/Objects/MemoryReferencesStore.cs
+++ b/Engine/Source/Programs/Horde/HordeStorage/Horde.Storage/Implementation/Objects/MemoryReferencesStore.cs
@@ -20,11 +20,11 @@ namespace Horde.Storage.Implementation
 
         }
 
-        public Task<ObjectRecord> Get(NamespaceId ns, BucketId bucket, IoHashKey key)
+        public Task<ObjectRecord> Get(NamespaceId ns, BucketId bucket, IoHashKey key, IReferencesStore.FieldFlags flags)
         {
             if (_objects.TryGetValue(BuildKey(ns, bucket, key), out MemoryStoreObject? o))
             {
-                return Task.FromResult(o.ToObjectRecord());
+                return Task.FromResult(o.ToObjectRecord(flags));
             }
 
             throw new ObjectNotFoundException(ns, bucket, key);
@@ -180,9 +180,10 @@ namespace Horde.Storage.Implementation
             LastAccessTime = lastAccessTime;
         }
 
-        public ObjectRecord ToObjectRecord()
+        public ObjectRecord ToObjectRecord(IReferencesStore.FieldFlags fieldFlags)
         {
-            return new ObjectRecord(Namespace, Bucket, Name, LastAccessTime, Blob, BlobHash, IsFinalized);
+            bool includePayload = (fieldFlags & IReferencesStore.FieldFlags.IncludePayload) != 0;
+            return new ObjectRecord(Namespace, Bucket, Name, LastAccessTime, includePayload ? Blob : null, BlobHash, IsFinalized);
         }
 
     }
diff --git a/Engine/Source/Programs/Horde/HordeStorage/Horde.Storage/Implementation/Objects/ScyllaReferencesStore.cs b/Engine/Source/Programs/Horde/HordeStorage/Horde.Storage/Implementation/Objects/ScyllaReferencesStore.cs
--- a/Engine/Source/Programs/Horde/HordeStorage/Horde.Storage/Implementation/Objects/ScyllaReferencesStore.cs
+++ b/Engine/Source/Programs/Horde/HordeStorage/Horde.Storage/Implementation/Objects/ScyllaReferencesStore.cs
@@ -57,11 +57,21 @@ namespace Horde.Storage.Implementation
 
         }
 
-        public async Task<ObjectRecord> Get(NamespaceId ns, BucketId bucket, IoHashKey name)
+        public async Task<ObjectRecord> Get(NamespaceId ns, BucketId bucket, IoHashKey name, IReferencesStore.FieldFlags flags)
         {
             using Scope _ = Tracer.Instance.StartActive("scylla.get");
 
-            ScyllaObject? o = await _mapper.SingleOrDefaultAsync<ScyllaObject>("WHERE namespace = ? AND bucket = ? AND name = ?", ns.ToString(), bucket.ToString(), name.ToString());
+            ScyllaObject? o;
+            bool includePayload = (flags & IReferencesStore.FieldFlags.IncludePayload) != 0;
+            if (includePayload)
+            {
+                o = await _mapper.SingleOrDefaultAsync<ScyllaObject>("WHERE namespace = ? AND bucket = ? AND name = ?", ns.ToString(), bucket.ToString(), name.ToString());
+            }
+            else
+            {
+                // fetch everything except for the inline blob which is quite large
+                o = await _mapper.SingleOrDefaultAsync<ScyllaObject>("SELECT namespace, bucket, name , payload_hash, is_finalized, last_access_time FROM objects WHERE namespace = ? AND bucket = ? AND name = ?", ns.ToString(), bucket.ToString(), name.ToString());
+            }
 
             if (o == null)
                 throw new ObjectNotFoundException(ns, bucket, name);
@@ -240,7 +250,7 @@ namespace Horde.Storage.Implementation
             ObjectRecord record;
             try
             {
-                record = await Get(ns, bucket, key);
+                record = await Get(ns, bucket, key, IReferencesStore.FieldFlags.None);
             }
             catch (ObjectNotFoundException)
             {
diff --git a/Engine/Source/Programs/Horde/HordeStorage/Horde.Storage/Tests/Functional/Objects/ReferencesTests.cs b/Engine/Source/Programs/Horde/HordeStorage/Horde.Storage/Tests/Functional/Objects/ReferencesTests.cs
--- a/Engine/Source/Programs/Horde/HordeStorage/Horde.Storage/Tests/Functional/Objects/ReferencesTests.cs
+++ b/Engine/Source/Programs/Horde/HordeStorage/Horde.Storage/Tests/Functional/Objects/ReferencesTests.cs
@@ -324,7 +324,7 @@ namespace Horde.Storage.FunctionalTests.References
             {
                 BucketId bucket = new BucketId("bucket");
 
-                ObjectRecord objectRecord = await ReferencesStore.Get(TestNamespace, bucket, key);
+                ObjectRecord objectRecord = await ReferencesStore.Get(TestNamespace, bucket, key, IReferencesStore.FieldFlags.IncludePayload);
 
                 Assert.IsTrue(objectRecord.IsFinalized);
                 Assert.AreEqual(key, objectRecord.Name);
@@ -480,7 +480,7 @@ namespace Horde.Storage.FunctionalTests.References
             {
                 BucketId bucket = new BucketId("bucket");
 
-                ObjectRecord objectRecord = await ReferencesStore.Get(TestNamespace, bucket, key);
+                ObjectRecord objectRecord = await ReferencesStore.Get(TestNamespace, bucket, key, IReferencesStore.FieldFlags.IncludePayload);
 
                 Assert.IsTrue(objectRecord.IsFinalized);
                 Assert.AreEqual(key, objectRecord.Name);
@@ -645,7 +645,7 @@ namespace Horde.Storage.FunctionalTests.References
             {
                 BucketId bucket = new BucketId("bucket");
 
-                ObjectRecord objectRecord = await ReferencesStore.Get(TestNamespace, bucket, key);
+                ObjectRecord objectRecord = await ReferencesStore.Get(TestNamespace, bucket, key, IReferencesStore.FieldFlags.IncludePayload);
 
                 Assert.IsTrue(objectRecord.IsFinalized);
                 Assert.AreEqual(key, objectRecord.Name);
@@ -898,7 +898,7 @@ namespace Horde.Storage.FunctionalTests.References
 
             // check that actual internal representation
             {
-                ObjectRecord objectRecord = await ReferencesStore.Get(TestNamespace, bucket, key);
+                ObjectRecord objectRecord = await ReferencesStore.Get(TestNamespace, bucket, key, IReferencesStore.FieldFlags.IncludePayload);
 
                 Assert.IsFalse(objectRecord.IsFinalized);
                 Assert.AreEqual(key, objectRecord.Name);
@@ -935,7 +935,7 @@ namespace Horde.Storage.FunctionalTests.References
 
             // check that actual internal representation has updated its state
             {
-                ObjectRecord objectRecord = await ReferencesStore.Get(TestNamespace, bucket, key);
+                ObjectRecord objectRecord = await ReferencesStore.Get(TestNamespace, bucket, key, IReferencesStore.FieldFlags.None);
 
                 Assert.IsTrue(objectRecord.IsFinalized);
                 Assert.AreEqual(key, objectRecord.Name);
