Horde Storage

Added verification for object coming from the db, as any field could be null and we want a better error message when that happens.

diff --git a/Engine/Source/Programs/Horde/HordeStorage/Horde.Storage/Implementation/Objects/ScyllaReferencesStore.cs b/Engine/Source/Programs/Horde/HordeStorage/Horde.Storage/Implementation/Objects/ScyllaReferencesStore.cs
--- a/Engine/Source/Programs/Horde/HordeStorage/Horde.Storage/Implementation/Objects/ScyllaReferencesStore.cs
+++ b/Engine/Source/Programs/Horde/HordeStorage/Horde.Storage/Implementation/Objects/ScyllaReferencesStore.cs
@@ -76,6 +76,7 @@ namespace Horde.Storage.Implementation
             if (o == null)
                 throw new ObjectNotFoundException(ns, bucket, name);
 
+            o.ThrowIfRequiredFieldsAreMissing();
             // TODO: Check returned values for null
             return new ObjectRecord(new NamespaceId(o.Namespace!), new BucketId(o.Bucket!), new IoHashKey(o.Name!), o.LastAccessTime, o.InlinePayload, o.PayloadHash!.AsBlobIdentifier(), o.IsFinalized!.Value);
         }
@@ -447,6 +448,20 @@ namespace Horde.Storage.Implementation
         public bool? IsFinalized { get;set; }
         [Cassandra.Mapping.Attributes.Column("last_access_time")]
         public DateTime LastAccessTime { get; set; }
+
+        public void ThrowIfRequiredFieldsAreMissing()
+        {
+            if (string.IsNullOrEmpty(Namespace))
+                throw new ArgumentException("Namespace was not valid", nameof(Namespace));
+            if (string.IsNullOrEmpty(Bucket))
+                throw new ArgumentException("Bucket was not valid", nameof(Bucket));
+            if (string.IsNullOrEmpty(Name))
+                throw new ArgumentException("Name was not valid", nameof(Name));
+            if (PayloadHash == null)
+                throw new ArgumentException("PayloadHash was not valid", nameof(PayloadHash));
+            if (!IsFinalized.HasValue)
+                throw new ArgumentException("IsFinalized was not valid", nameof(IsFinalized));
+        }
     }
 
     [Cassandra.Mapping.Attributes.Table("buckets")]
