Horde: Allow server's private gRPC certificate to be represented as base64 in config

Works in addition to a normal file path.

diff --git a/Engine/Source/Programs/Horde/Horde.Build.Tests/ProgramTests.cs b/Engine/Source/Programs/Horde/Horde.Build.Tests/ProgramTests.cs
new file mode 100644
--- /dev/null
+++ b/Engine/Source/Programs/Horde/Horde.Build.Tests/ProgramTests.cs
@@ -0,0 +1,44 @@
+// Copyright Epic Games, Inc. All Rights Reserved.
+
+using System;
+using System.IO;
+using System.Security.Cryptography.X509Certificates;
+using EpicGames.Core;
+using HordeServer;
+using Microsoft.VisualStudio.TestTools.UnitTesting;
+
+namespace Horde.Build.Tests
+{
+	[TestClass]
+	public class ProgramTests
+	{
+		[TestMethod]
+		public void TestReadGrpcCertificate()
+		{
+			string FriendlyName = "A testing cert";
+			byte[] TempCertData = CertificateUtils.CreateSelfSignedCert("testing.epicgames.com", FriendlyName);
+			string TempCertPath = Path.GetTempFileName();
+			File.WriteAllBytes(TempCertPath, TempCertData);
+			
+			// No cert given
+			Assert.IsNull(Program.ReadGrpcCertificate(new () { ServerPrivateCert = null }));
+
+			// Cert as file path
+			{
+				X509Certificate2? Cert = Program.ReadGrpcCertificate(new() { ServerPrivateCert = TempCertPath });
+				Assert.IsNotNull(Cert);
+				Assert.AreEqual(FriendlyName, Cert!.FriendlyName);
+			}
+
+			// Cert as base64 data
+			{
+				string TempCertBase64 = Convert.ToBase64String(TempCertData);
+				X509Certificate2? Cert = Program.ReadGrpcCertificate(new() { ServerPrivateCert = "base64:" + TempCertBase64 });
+				Assert.IsNotNull(Cert);
+				Assert.AreEqual(FriendlyName, Cert!.FriendlyName);	
+			}
+			
+			File.Delete(TempCertPath);
+		}
+	}
+}
diff --git a/Engine/Source/Programs/Horde/Horde.Build/Program.cs b/Engine/Source/Programs/Horde/Horde.Build/Program.cs
--- a/Engine/Source/Programs/Horde/Horde.Build/Program.cs
+++ b/Engine/Source/Programs/Horde/Horde.Build/Program.cs
@@ -269,10 +269,17 @@ namespace HordeServer
 		/// <returns>Custom certificate to use for Grpc endpoints, or null for the default.</returns>
 		public static X509Certificate2? ReadGrpcCertificate(ServerSettings HordeSettings)
 		{
+			string Base64Prefix = "base64:";
+			
 			if (HordeSettings.ServerPrivateCert == null)
 			{
 				return null;
 			}
+			else if (HordeSettings.ServerPrivateCert.StartsWith(Base64Prefix))
+			{
+				byte[] CertData = Convert.FromBase64String(HordeSettings.ServerPrivateCert.Replace(Base64Prefix, "", StringComparison.Ordinal));
+				return new X509Certificate2(CertData);
+			}
 			else
 			{
 				FileReference? ServerPrivateCert = null;
