Horde Storage
Only updated last access time for objects that actually exists, avoids creating a partial object.
Could potentially be a bit slow as this will trigger a transacation (due to IF EXISTS) but this is being run async from the actual requests so it shouldnt matter.

diff --git a/Engine/Source/Programs/Horde/HordeStorage/Horde.Storage/Controllers/ReferencesController.cs b/Engine/Source/Programs/Horde/HordeStorage/Horde.Storage/Controllers/ReferencesController.cs
--- a/Engine/Source/Programs/Horde/HordeStorage/Horde.Storage/Controllers/ReferencesController.cs
+++ b/Engine/Source/Programs/Horde/HordeStorage/Horde.Storage/Controllers/ReferencesController.cs
@@ -109,7 +109,6 @@ namespace Horde.Storage.Controllers
 
             try
             {
-                // TODO: Do we want to keep fields? we have no real need for it now
                 (ObjectRecord objectRecord, BlobContents? blob) = await _objectService.Get(ns, bucket, key, fields);
 
                 if (blob == null)
diff --git a/Engine/Source/Programs/Horde/HordeStorage/Horde.Storage/Implementation/Objects/ScyllaReferencesStore.cs b/Engine/Source/Programs/Horde/HordeStorage/Horde.Storage/Implementation/Objects/ScyllaReferencesStore.cs
--- a/Engine/Source/Programs/Horde/HordeStorage/Horde.Storage/Implementation/Objects/ScyllaReferencesStore.cs
+++ b/Engine/Source/Programs/Horde/HordeStorage/Horde.Storage/Implementation/Objects/ScyllaReferencesStore.cs
@@ -77,7 +77,6 @@ namespace Horde.Storage.Implementation
                 throw new ObjectNotFoundException(ns, bucket, name);
 
             o.ThrowIfRequiredFieldsAreMissing();
-            // TODO: Check returned values for null
             return new ObjectRecord(new NamespaceId(o.Namespace!), new BucketId(o.Bucket!), new IoHashKey(o.Name!), o.LastAccessTime, o.InlinePayload, o.PayloadHash!.AsBlobIdentifier(), o.IsFinalized!.Value);
         }
 
@@ -151,7 +150,7 @@ namespace Horde.Storage.Implementation
 
             await updateObjectTracking;*/
 
-            await _mapper.UpdateAsync<ScyllaObject>("SET last_access_time = ? WHERE namespace = ? AND bucket = ? AND name = ?", lastAccessTime, ns.ToString(), bucket.ToString(), name.ToString());
+            await _mapper.UpdateIfAsync<ScyllaObject>("SET last_access_time = ? WHERE namespace = ? AND bucket = ? AND name = ? IF EXISTS", lastAccessTime, ns.ToString(), bucket.ToString(), name.ToString());
 
             
         }
