Horde: Fix warning failing build

diff --git a/Engine/Source/Programs/Horde/Horde.Build/Program.cs b/Engine/Source/Programs/Horde/Horde.Build/Program.cs
--- a/Engine/Source/Programs/Horde/Horde.Build/Program.cs
+++ b/Engine/Source/Programs/Horde/Horde.Build/Program.cs
@@ -275,7 +275,7 @@ namespace HordeServer
 			{
 				return null;
 			}
-			else if (HordeSettings.ServerPrivateCert.StartsWith(Base64Prefix))
+			else if (HordeSettings.ServerPrivateCert.StartsWith(Base64Prefix, StringComparison.Ordinal))
 			{
 				byte[] CertData = Convert.FromBase64String(HordeSettings.ServerPrivateCert.Replace(Base64Prefix, "", StringComparison.Ordinal));
 				return new X509Certificate2(CertData);
