Horde: Disable gRPC reading test on Windows as test cert creation fails on Unix

diff --git a/Engine/Source/Programs/Horde/Horde.Build.Tests/ProgramTests.cs b/Engine/Source/Programs/Horde/Horde.Build.Tests/ProgramTests.cs
--- a/Engine/Source/Programs/Horde/Horde.Build.Tests/ProgramTests.cs
+++ b/Engine/Source/Programs/Horde/Horde.Build.Tests/ProgramTests.cs
@@ -2,6 +2,7 @@
 
 using System;
 using System.IO;
+using System.Runtime.InteropServices;
 using System.Security.Cryptography.X509Certificates;
 using EpicGames.Core;
 using HordeServer;
@@ -15,6 +16,12 @@ namespace Horde.Build.Tests
 		[TestMethod]
 		public void TestReadGrpcCertificate()
 		{
+			if (!RuntimeInformation.IsOSPlatform(OSPlatform.Windows))
+			{
+				// Test only works on Windows as .FriendlyName cannot be set on Unix.
+				return;
+			}
+			
 			string FriendlyName = "A testing cert";
 			byte[] TempCertData = CertificateUtils.CreateSelfSignedCert("testing.epicgames.com", FriendlyName);
 			string TempCertPath = Path.GetTempFileName();
