Horde Storage

Made ref cleanup logging a bit more verbose to help diagnose issue with it cleaning to much

diff --git a/Engine/Source/Programs/Horde/HordeStorage/Horde.Storage/Implementation/GC/CacheCleanup.cs b/Engine/Source/Programs/Horde/HordeStorage/Horde.Storage/Implementation/GC/CacheCleanup.cs
--- a/Engine/Source/Programs/Horde/HordeStorage/Horde.Storage/Implementation/GC/CacheCleanup.cs
+++ b/Engine/Source/Programs/Horde/HordeStorage/Horde.Storage/Implementation/GC/CacheCleanup.cs
@@ -59,7 +59,7 @@ namespace Horde.Storage.Implementation
             int consideredCount = 0;
             await foreach ((BucketId bucket, IoHashKey name, DateTime lastAccessTime) in _referencesStore.GetRecords(ns).WithCancellation(cancellationToken))
             {
-                _logger.Debug("Considering object in {Namespace} {Bucket} {Name} for deletion, was last updated {LastAccessTime}", ns, bucket, name, lastAccessTime);
+                _logger.Information("Considering object in {Namespace} {Bucket} {Name} for deletion, was last updated {LastAccessTime}", ns, bucket, name, lastAccessTime);
                 Interlocked.Increment(ref consideredCount);
 
                 if (lastAccessTime > cutoffTime)
