Horde: Sort job details text selection issues
#rnx

diff --git a/Engine/Source/Programs/Horde/HordeDashboard/src/backend/Api.ts b/Engine/Source/Programs/Horde/HordeDashboard/src/backend/Api.ts
--- a/Engine/Source/Programs/Horde/HordeDashboard/src/backend/Api.ts
+++ b/Engine/Source/Programs/Horde/HordeDashboard/src/backend/Api.ts
@@ -75,12 +75,18 @@ export enum JobStepBatchState {
 	// Waiting for dependencies of at least one jobstep to complete
 	Waiting = "Waiting",
 
+	// Getting ready to execute
+	Starting = "Starting",
+
 	// Ready to execute
 	Ready = "Ready",
 
 	// Currently running
 	Running = "Running",
 
+	// Getting ready to execute
+	Stopping = "Stopping",
+
 	// All steps have finished executing
 	Complete = "Complete"
 }
diff --git a/Engine/Source/Programs/Horde/HordeDashboard/src/backend/JobDetails.ts b/Engine/Source/Programs/Horde/HordeDashboard/src/backend/JobDetails.ts
--- a/Engine/Source/Programs/Horde/HordeDashboard/src/backend/JobDetails.ts
+++ b/Engine/Source/Programs/Horde/HordeDashboard/src/backend/JobDetails.ts
@@ -1076,6 +1076,15 @@ export const getBatchSummaryMarkdown = (jobDetails: JobDetails, batchId: string)
 
     if (!summaryText) {
         summaryText = getBatchText({ batch: batch, agentId: batch?.agentId, agentType: agentType, agentPool: agentPool }) ?? "";
+
+        if (summaryText && batch.state === JobStepBatchState.Starting) {
+            summaryText += " batch is the starting state";
+        }
+
+        if (summaryText && batch.state === JobStepBatchState.Stopping) {
+            summaryText += " batch is the stopping state";
+        }
+
         if (!summaryText) {
             summaryText = batch?.agentId ?? (batch?.state ?? "Batch is unassigned");
         }
diff --git a/Engine/Source/Programs/Horde/HordeDashboard/src/index.css b/Engine/Source/Programs/Horde/HordeDashboard/src/index.css
--- a/Engine/Source/Programs/Horde/HordeDashboard/src/index.css
+++ b/Engine/Source/Programs/Horde/HordeDashboard/src/index.css
@@ -35,14 +35,10 @@ code {
 
 .job-item:hover {
   background:rgb(243, 242, 241);
-  user-select: none;
-  cursor: pointer;
 }
 
 .job-item.selected {
   background:rgb(243, 242, 241);
-  user-select: none;
-  cursor: pointer;
 }
 
 
