Horde Storage

Fixed gc deleting newly created compact binary objects

diff --git a/Engine/Source/Programs/Horde/HordeStorage/Horde.Storage/Implementation/GC/OrphanBlobCleanupRefs.cs b/Engine/Source/Programs/Horde/HordeStorage/Horde.Storage/Implementation/GC/OrphanBlobCleanupRefs.cs
--- a/Engine/Source/Programs/Horde/HordeStorage/Horde.Storage/Implementation/GC/OrphanBlobCleanupRefs.cs
+++ b/Engine/Source/Programs/Horde/HordeStorage/Horde.Storage/Implementation/GC/OrphanBlobCleanupRefs.cs
@@ -57,9 +57,10 @@ namespace Horde.Storage.Implementation
 
                 // only consider blobs that have been around for 60 minutes
                 // this due to cases were blobs are uploaded first
+                DateTime cutoff = DateTime.Now.AddMinutes(-60);
                 await foreach ((BlobIdentifier blob, DateTime lastModified) in _blobService.ListObjects(@namespace).WithCancellation(cancellationToken))
                 {
-                    if (lastModified > DateTime.Now.AddMinutes(-60))
+                    if (lastModified < cutoff)
                         continue;
                     
                     if (cancellationToken.IsCancellationRequested)
