Horde: Fix build errors.

diff --git a/Engine/Source/Programs/Horde/Horde.Build/Perforce/Impl/CommitService.cs b/Engine/Source/Programs/Horde/Horde.Build/Perforce/Impl/CommitService.cs
--- a/Engine/Source/Programs/Horde/Horde.Build/Perforce/Impl/CommitService.cs
+++ b/Engine/Source/Programs/Horde/Horde.Build/Perforce/Impl/CommitService.cs
@@ -760,22 +760,22 @@ namespace HordeServer.Commits.Impl
 		public class Utf8FStatRecord
 		{
 			[PerforceTag("depotFile", Optional = true)]
-			public Utf8String DepotFile;
+			public Utf8String DepotFile { get; set; }
 
 			[PerforceTag("headType", Optional = true)]
-			public Utf8String HeadType;
+			public Utf8String HeadType { get; set; }
 
 			[PerforceTag("headRev", Optional = true)]
-			public int HeadRevision;
+			public int HeadRevision { get; set; }
 
 			[PerforceTag("digest", Optional = true)]
-			public Utf8String Digest;
+			public Utf8String Digest { get; set; }
 
 			[PerforceTag("fileSize", Optional = true)]
-			public long FileSize;
+			public long FileSize { get; set; }
 
 			[PerforceTag("type", Optional = true)]
-			public Utf8String Type;
+			public Utf8String Type { get; set; }
 		}
 
 		/// <summary>
