Horde Storage

Fixed incorrect orphan blob cleanup for compact binaries

diff --git a/Engine/Source/Programs/Horde/HordeStorage/Horde.Storage/Implementation/GC/OrphanBlobCleanupRefs.cs b/Engine/Source/Programs/Horde/HordeStorage/Horde.Storage/Implementation/GC/OrphanBlobCleanupRefs.cs
--- a/Engine/Source/Programs/Horde/HordeStorage/Horde.Storage/Implementation/GC/OrphanBlobCleanupRefs.cs
+++ b/Engine/Source/Programs/Horde/HordeStorage/Horde.Storage/Implementation/GC/OrphanBlobCleanupRefs.cs
@@ -60,7 +60,7 @@ namespace Horde.Storage.Implementation
                 DateTime cutoff = DateTime.Now.AddMinutes(-60);
                 await foreach ((BlobIdentifier blob, DateTime lastModified) in _blobService.ListObjects(@namespace).WithCancellation(cancellationToken))
                 {
-                    if (lastModified < cutoff)
+                    if (lastModified > cutoff)
                         continue;
                     
                     if (cancellationToken.IsCancellationRequested)
diff --git a/Engine/Source/Programs/Horde/HordeStorage/Horde.Storage/Tests/Functional/GC/GCBlobTests.cs b/Engine/Source/Programs/Horde/HordeStorage/Horde.Storage/Tests/Functional/GC/GCBlobTests.cs
--- a/Engine/Source/Programs/Horde/HordeStorage/Horde.Storage/Tests/Functional/GC/GCBlobTests.cs
+++ b/Engine/Source/Programs/Horde/HordeStorage/Horde.Storage/Tests/Functional/GC/GCBlobTests.cs
@@ -281,7 +281,7 @@ namespace Horde.Storage.FunctionalTests.GC
             Assert.IsNotNull(objectService);
             (BlobIdentifier ob0_hash, CompactBinaryObject ob0_cb) = GetCBWithAttachment(object0id);
             await objectService.Put(TestNamespace, testBucket, IoHashKey.FromName("object0"), ob0_hash, ob0_cb);
-
+           
             (BlobIdentifier ob2_hash, CompactBinaryObject ob2_cb) = GetCBWithAttachment(object2id);
             await objectService.Put(TestNamespace, testBucket, IoHashKey.FromName("object2"), ob2_hash, ob2_cb);
 
@@ -291,6 +291,12 @@ namespace Horde.Storage.FunctionalTests.GC
             (BlobIdentifier ob6_hash, CompactBinaryObject ob6_cb) = GetCBWithAttachment(object6id);
             await objectService.Put(TestNamespace, testBucket, IoHashKey.FromName("object6"), ob6_hash, ob6_cb);
 
+            IReferencesStore referenceStore = server.Services.GetService<IReferencesStore>()!;
+            await referenceStore.UpdateLastAccessTime(TestNamespace, testBucket, IoHashKey.FromName("object0"), DateTime.Now.AddDays(-2));
+            await referenceStore.UpdateLastAccessTime(TestNamespace, testBucket, IoHashKey.FromName("object2"), DateTime.Now.AddDays(-2));
+            await referenceStore.UpdateLastAccessTime(TestNamespace, testBucket, IoHashKey.FromName("object3"), DateTime.Now.AddDays(-2));
+            await referenceStore.UpdateLastAccessTime(TestNamespace, testBucket, IoHashKey.FromName("object6"), DateTime.Now.AddDays(-2));
+
             IBlobIndex? blobIndex = server.Services.GetService<IBlobIndex>()!;
             Assert.IsNotNull(blobIndex);
             await blobIndex.AddBlobToIndex(TestNamespace, object0id);
