Horde: Fix path to schema for Horde build script.

diff --git a/Engine/Source/Programs/Horde/BuildHorde.xml b/Engine/Source/Programs/Horde/BuildHorde.xml
--- a/Engine/Source/Programs/Horde/BuildHorde.xml
+++ b/Engine/Source/Programs/Horde/BuildHorde.xml
@@ -1,5 +1,5 @@
 <?xml version='1.0' ?>
-<BuildGraph xmlns="http://www.epicgames.com/BuildGraph" xmlns:xsi="http://www.w3.org/2001/XMLSchema-instance" xsi:schemaLocation="http://www.epicgames.com/BuildGraph ../../Engine/Build/Graph/Schema.xsd" >
+<BuildGraph xmlns="http://www.epicgames.com/BuildGraph" xmlns:xsi="http://www.w3.org/2001/XMLSchema-instance" xsi:schemaLocation="http://www.epicgames.com/BuildGraph ../../../Engine/Build/Graph/Schema.xsd" >
 
 	<Option Name="PreflightChange" DefaultValue="" Description="Preflight changelist number"/>
 	<Option Name="Configuration" DefaultValue="release" Description="Configuration to build"/>
