Horde: Possible fix for non-deterministic test failures.

diff --git a/Engine/Source/Programs/Horde/Horde.Build.Tests/JobServiceTests.cs b/Engine/Source/Programs/Horde/Horde.Build.Tests/JobServiceTests.cs
--- a/Engine/Source/Programs/Horde/Horde.Build.Tests/JobServiceTests.cs
+++ b/Engine/Source/Programs/Horde/Horde.Build.Tests/JobServiceTests.cs
@@ -146,6 +146,9 @@ namespace HordeServerTests
 		[TestMethod]
 		public async Task TestRunEarly()
 		{
+			IAgent? Agent = await AgentService.CreateAgentAsync("TestAgent", true, null, new List<StringId<IPool>> { new StringId<IPool>("win") });
+			await AgentService.CreateSessionAsync(Agent, AgentStatus.Ok, new List<string>(), new Dictionary<string, int>(), null);
+
 			IProject? Project = await ProjectService.Collection.AddOrUpdateAsync(new ProjectId("ue5"), "", "", 0, new ProjectConfig { Name = "UE5" });
 			Assert.IsNotNull(Project);
 
