Horde: Fix typo in comment

#skipci

diff --git a/Engine/Source/Programs/Horde/Horde.Build/ServerSettings.cs b/Engine/Source/Programs/Horde/Horde.Build/ServerSettings.cs
--- a/Engine/Source/Programs/Horde/Horde.Build/ServerSettings.cs
+++ b/Engine/Source/Programs/Horde/Horde.Build/ServerSettings.cs
@@ -196,7 +196,7 @@ namespace HordeServer
 		public bool DatabaseReadOnlyMode { get; set; } = false;
 
 		/// <summary>
-		/// Optional PFX certificate to use for encryting agent SSL traffic. This can be a self-signed certificate, as long as it's trusted by agents.
+		/// Optional PFX certificate to use for encrypting agent SSL traffic. This can be a self-signed certificate, as long as it's trusted by agents.
 		/// </summary>
 		public string? ServerPrivateCert { get; set; }
 
