Horde Storage - Ref cleanup for compact binaries was applying the delete timespan incorrectly, resulting in deleting everything older then 14 days from the future (e.g. everything) and not everything older then 14 days from now.

Also added functional test to excerise this new code which exposed this issue.

diff --git a/Engine/Source/Programs/Horde/HordeStorage/Horde.Storage/Implementation/GC/CacheCleanup.cs b/Engine/Source/Programs/Horde/HordeStorage/Horde.Storage/Implementation/GC/CacheCleanup.cs
--- a/Engine/Source/Programs/Horde/HordeStorage/Horde.Storage/Implementation/GC/CacheCleanup.cs
+++ b/Engine/Source/Programs/Horde/HordeStorage/Horde.Storage/Implementation/GC/CacheCleanup.cs
@@ -55,17 +55,17 @@ namespace Horde.Storage.Implementation
         private async Task<List<OldRecord>> CleanNamespace(NamespaceId ns, CancellationToken cancellationToken)
         {
             List<OldRecord> deletedRecords = new List<OldRecord>();
-            DateTime cutoffTime = DateTime.Now.Add(_settings.CurrentValue.LastAccessCutoff);
+            DateTime cutoffTime = DateTime.Now.AddSeconds(-1 * _settings.CurrentValue.LastAccessCutoff.TotalSeconds);
             int consideredCount = 0;
             await foreach ((BucketId bucket, IoHashKey name, DateTime lastAccessTime) in _referencesStore.GetRecords(ns).WithCancellation(cancellationToken))
             {
-                _logger.Information("Considering object in {Namespace} {Bucket} {Name} for deletion, was last updated {LastAccessTime}", ns, bucket, name, lastAccessTime);
+                _logger.Debug("Considering object in {Namespace} {Bucket} {Name} for deletion, was last updated {LastAccessTime}", ns, bucket, name, lastAccessTime);
                 Interlocked.Increment(ref consideredCount);
 
                 if (lastAccessTime > cutoffTime)
                     continue;
 
-                _logger.Debug("Attempting to delete object {Namespace} {Bucket} {Name}.", ns, bucket, name);
+                _logger.Information("Attempting to delete object {Namespace} {Bucket} {Name} as it was last updated {LastAccessTime} which is older then {CutoffTime}", ns, bucket, name, lastAccessTime, cutoffTime);
                 using Scope scope = Tracer.Instance.StartActive("refCleanup.delete_record");
                 scope.Span.ResourceName = $"{ns}:{bucket}.{name}";
                 // delete the old record from the ref refs
diff --git a/Engine/Source/Programs/Horde/HordeStorage/Horde.Storage/Tests/Functional/GC/GCRefTests.cs b/Engine/Source/Programs/Horde/HordeStorage/Horde.Storage/Tests/Functional/GC/GCRefTests.cs
--- a/Engine/Source/Programs/Horde/HordeStorage/Horde.Storage/Tests/Functional/GC/GCRefTests.cs
+++ b/Engine/Source/Programs/Horde/HordeStorage/Horde.Storage/Tests/Functional/GC/GCRefTests.cs
@@ -8,6 +8,7 @@ using System.Threading.Tasks;
 using Dasync.Collections;
 using Horde.Storage.Controllers;
 using Horde.Storage.Implementation;
+using Horde.Storage.Implementation.Blob;
 using Jupiter.Implementation;
 using Microsoft.AspNetCore.TestHost;
 using Microsoft.AspNetCore.Hosting;
@@ -100,4 +101,149 @@ namespace Horde.Storage.FunctionalTests.GC
             _refMock.VerifyNoOtherCalls();
         }
     }
+
+
+    [TestClass]
+    public class MemoryGCReferencesTests : GCReferencesTests
+    {
+        protected override string GetImplementation()
+        {
+            return "Memory";
+        }
+    }
+
+    [TestClass]
+    public class ScyllaGCReferencesTests : GCReferencesTests
+    {
+        protected override string GetImplementation()
+        {
+            return "Scylla";
+        }
+    }
+
+
+    public abstract class GCReferencesTests
+    {
+        private TestServer? _server;
+        private HttpClient? _httpClient;
+
+        private readonly NamespaceId TestNamespace = new NamespaceId("test-namespace");
+        private readonly BucketId DefaultBucket = new BucketId("default");
+
+        private readonly BlobIdentifier object0id = new BlobIdentifier("0000000000000000000000000000000000000000");
+        private readonly BlobIdentifier object1id = new BlobIdentifier("1111111111111111111111111111111111111111");
+        private readonly BlobIdentifier object2id = new BlobIdentifier("2222222222222222222222222222222222222222");
+        private readonly BlobIdentifier object3id = new BlobIdentifier("3333333333333333333333333333333333333333");
+        private readonly BlobIdentifier object4id = new BlobIdentifier("4444444444444444444444444444444444444444");
+        private readonly BlobIdentifier object5id = new BlobIdentifier("5555555555555555555555555555555555555555");
+        private readonly BlobIdentifier object6id = new BlobIdentifier("6666666666666666666666666666666666666666");
+
+        private readonly IoHashKey object0Name = IoHashKey.FromName("object0");
+        private readonly IoHashKey object1Name = IoHashKey.FromName("object1");
+        private readonly IoHashKey object2Name = IoHashKey.FromName("object2");
+        private readonly IoHashKey object3Name = IoHashKey.FromName("object3");
+        private readonly IoHashKey object4Name = IoHashKey.FromName("object4");
+        private readonly IoHashKey object5Name = IoHashKey.FromName("object5");
+        private readonly IoHashKey object6Name = IoHashKey.FromName("object6");
+
+        [TestInitialize]
+        public async Task Setup()
+        {
+            IConfigurationRoot configuration = new ConfigurationBuilder()
+                // we are not reading the base appSettings here as we want exact control over what runs in the tests
+                .AddJsonFile("appsettings.Testing.json", false)
+                .AddEnvironmentVariables()
+                .AddInMemoryCollection(new List<KeyValuePair<string, string>>()
+                {
+                    new KeyValuePair<string, string>("GC:CleanNamespaces:0", TestNamespace.ToString()),
+                    new KeyValuePair<string, string>("Horde_Storage:StorageImplementations:0", "MemoryBlobStore"),
+                    new KeyValuePair<string, string>("Horde_Storage:BlobIndexImplementation", GetImplementation()),
+                })
+                .Build();
+            Logger logger = new LoggerConfiguration()
+                .ReadFrom.Configuration(configuration)
+                .CreateLogger();
+
+            TestServer server = new TestServer(new WebHostBuilder()
+                .UseConfiguration(configuration)
+                .UseEnvironment("Testing")
+                .UseSerilog(logger)
+                .UseStartup<HordeStorageStartup>()
+            );
+            _httpClient = server.CreateClient();
+            _server = server;
+            
+
+            IObjectService? objectService = server.Services.GetService<IObjectService>()!;
+            Assert.IsNotNull(objectService);
+            (BlobIdentifier ob0_hash, CompactBinaryObject ob0_cb) = GetCBWithAttachment(object0id);
+            await objectService.Put(TestNamespace, DefaultBucket, object0Name, ob0_hash, ob0_cb);
+           
+            (BlobIdentifier ob1_hash, CompactBinaryObject ob1_cb) = GetCBWithAttachment(object1id);
+            await objectService.Put(TestNamespace, DefaultBucket, object1Name, ob1_hash, ob1_cb);
+
+            (BlobIdentifier ob2_hash, CompactBinaryObject ob2_cb) = GetCBWithAttachment(object2id);
+            await objectService.Put(TestNamespace, DefaultBucket, object2Name, ob2_hash, ob2_cb);
+
+            (BlobIdentifier ob3_hash, CompactBinaryObject ob3_cb) = GetCBWithAttachment(object3id);
+            await objectService.Put(TestNamespace, DefaultBucket, object3Name, ob3_hash, ob3_cb);
+
+            (BlobIdentifier ob4_hash, CompactBinaryObject ob4_cb) = GetCBWithAttachment(object4id);
+            await objectService.Put(TestNamespace, DefaultBucket, object4Name, ob4_hash, ob4_cb);
+
+            (BlobIdentifier ob5_hash, CompactBinaryObject ob5_cb) = GetCBWithAttachment(object5id);
+            await objectService.Put(TestNamespace, DefaultBucket, object5Name, ob5_hash, ob5_cb);
+
+            (BlobIdentifier ob6_hash, CompactBinaryObject ob6_cb) = GetCBWithAttachment(object6id);
+            await objectService.Put(TestNamespace, DefaultBucket, object6Name, ob6_hash, ob6_cb);
+
+            IReferencesStore referenceStore = server.Services.GetService<IReferencesStore>()!;
+            DateTime oldTimestamp = DateTime.Now.AddDays(-30);
+            DateTime newTimestamp = DateTime.Now;
+            await referenceStore.UpdateLastAccessTime(TestNamespace, DefaultBucket, object0Name, oldTimestamp);
+            await referenceStore.UpdateLastAccessTime(TestNamespace, DefaultBucket, object1Name, newTimestamp);
+            await referenceStore.UpdateLastAccessTime(TestNamespace, DefaultBucket, object2Name, oldTimestamp);
+            await referenceStore.UpdateLastAccessTime(TestNamespace, DefaultBucket, object3Name, oldTimestamp);
+            await referenceStore.UpdateLastAccessTime(TestNamespace, DefaultBucket, object4Name, newTimestamp);
+            await referenceStore.UpdateLastAccessTime(TestNamespace, DefaultBucket, object5Name, newTimestamp);
+            await referenceStore.UpdateLastAccessTime(TestNamespace, DefaultBucket, object6Name, oldTimestamp);
+        }
+
+        protected abstract string GetImplementation();
+
+
+        [TestMethod]
+        public async Task RunRefCleanup()
+        {
+            // trigger the cleanup
+            var cleanupResponse = await _httpClient!.PostAsync($"api/v1/admin/refCleanup/{TestNamespace}", new StringContent(string.Empty));
+            cleanupResponse.EnsureSuccessStatusCode();
+            RemovedRefRecordsResponse removedRefRecords = await cleanupResponse.Content.ReadAsAsync<RemovedRefRecordsResponse>();
+
+            string[] expectedRemovedRefRecords = new[]
+            {
+                object0Name.ToString(),
+                object2Name.ToString(),
+                object3Name.ToString(),
+                object6Name.ToString()
+            };
+            Assert.AreEqual(4, removedRefRecords.RemovedRecords.Length);
+            foreach (RemovedRefRecordsResponse.RemovedRecord removedRecord in removedRefRecords.RemovedRecords)
+            {
+                Assert.AreEqual(DefaultBucket, removedRecord.Bucket);
+                Assert.IsTrue(expectedRemovedRefRecords.Contains(removedRecord.Name.ToString()));
+            }
+        }
+
+        private (BlobIdentifier, CompactBinaryObject) GetCBWithAttachment(BlobIdentifier blobIdentifier)
+        {
+            CompactBinaryWriter compactBinaryWriter = new CompactBinaryWriter();
+            compactBinaryWriter.BeginObject();
+            compactBinaryWriter.AddBinaryAttachment(blobIdentifier);
+            compactBinaryWriter.EndObject();
+
+            byte[] b = compactBinaryWriter.Save();
+            return (BlobIdentifier.FromBlob(b), CompactBinaryObject.Load(b));
+        }
+    }
 }
