Horde Storage - Return a better error message when the object you are fetching references for does not exist

diff --git a/Engine/Source/Programs/Horde/HordeStorage/Horde.Storage/Controllers/ObjectController.cs b/Engine/Source/Programs/Horde/HordeStorage/Horde.Storage/Controllers/ObjectController.cs
--- a/Engine/Source/Programs/Horde/HordeStorage/Horde.Storage/Controllers/ObjectController.cs
+++ b/Engine/Source/Programs/Horde/HordeStorage/Horde.Storage/Controllers/ObjectController.cs
@@ -172,6 +172,10 @@ namespace Horde.Storage.Controllers
             {
                 return BadRequest(new ValidationProblemDetails {Title = $"Object {id} is missing blobs", Detail = $"Following blobs are missing: {string.Join(",", e.UnresolvedReferences)}"});
             }
+            catch (BlobNotFoundException e)
+            {
+                return NotFound(new ValidationProblemDetails {Title = $"Object {e.Blob} not found"});
+            }
         }
 
 
