Horde Storage - Refactored blob index tracking to avoid transactions in scylla, reducing load.
This can result in references and regions not being set depending on order of operations, which I updated the code to handle.

diff --git a/Engine/Source/Programs/Horde/HordeStorage/Horde.Storage/Implementation/Blob/ScyllaBlobIndex.cs b/Engine/Source/Programs/Horde/HordeStorage/Horde.Storage/Implementation/Blob/ScyllaBlobIndex.cs
--- a/Engine/Source/Programs/Horde/HordeStorage/Horde.Storage/Implementation/Blob/ScyllaBlobIndex.cs
+++ b/Engine/Source/Programs/Horde/HordeStorage/Horde.Storage/Implementation/Blob/ScyllaBlobIndex.cs
@@ -37,17 +37,9 @@ public class ScyllaBlobIndex : IBlobIndex
     public async Task AddBlobToIndex(NamespaceId ns, BlobIdentifier id)
     {
         using Scope _ = Tracer.Instance.StartActive("scylla.insert_blob_index");
-        ScyllaBlobIndexTable tableRow = new ScyllaBlobIndexTable(ns.ToString(), id,
-            new HashSet<string> { _jupiterSettings.CurrentValue.CurrentSite }, new List<ScyllaObjectReference>());
 
-        AppliedInfo<ScyllaBlobIndexTable>? result = await _mapper.InsertIfNotExistsAsync(tableRow);
-
-        // if it already existed we just update the regions the blob existed in
-        if (!result.Applied)
-        {
-            await _mapper.UpdateAsync<ScyllaBlobIndexTable>("SET regions = regions + ? WHERE namespace = ? AND blob_id = ?",
-                new string[] { _jupiterSettings.CurrentValue.CurrentSite }, ns.ToString(), new ScyllaBlobIdentifier(id));
-        }
+        await _mapper.UpdateAsync<ScyllaBlobIndexTable>("SET regions = regions + ? WHERE namespace = ? AND blob_id = ?",
+            new string[] { _jupiterSettings.CurrentValue.CurrentSite }, ns.ToString(), new ScyllaBlobIdentifier(id));
     }
 
     public async Task<IBlobIndex.BlobInfo?> GetBlobInfo(NamespaceId ns, BlobIdentifier id)
@@ -63,9 +55,9 @@ public class ScyllaBlobIndex : IBlobIndex
 
         return new IBlobIndex.BlobInfo
         {
-            Regions = blobIndex.Regions,
+            Regions = blobIndex.Regions ?? new HashSet<string>(),
             Namespace = new NamespaceId(blobIndex.Namespace),
-            References = blobIndex.References.Select(reference => reference.AsTuple()).ToList()
+            References = blobIndex.References != null ? blobIndex.References.Select(reference => reference.AsTuple()).ToList() : new List<(BucketId, IoHashKey)>()
         };
     }
 
@@ -115,8 +107,8 @@ class ScyllaBlobIndexTable
     {
         Namespace = null!;
         BlobId = null!;
-        Regions = null!;
-        References = null!;
+        Regions = null;
+        References = null;
     }
 
     public ScyllaBlobIndexTable(string @namespace, BlobIdentifier blobId, HashSet<string> regions, List<ScyllaObjectReference> references)
@@ -135,8 +127,8 @@ class ScyllaBlobIndexTable
     public ScyllaBlobIdentifier BlobId { get;set; }
 
     [Cassandra.Mapping.Attributes.Column("regions")]
-    public HashSet<string> Regions { get; set; }
+    public HashSet<string>? Regions { get; set; }
 
     [Cassandra.Mapping.Attributes.Column("references")]
-    public List<ScyllaObjectReference> References { get; set; }
+    public List<ScyllaObjectReference>? References { get; set; }
 }
