Horde Storage

Improved datadog scopes in the blob service to better track time spent

diff --git a/Engine/Source/Programs/Horde/HordeStorage/Horde.Storage/Implementation/Blob/IBlobService.cs b/Engine/Source/Programs/Horde/HordeStorage/Horde.Storage/Implementation/Blob/IBlobService.cs
--- a/Engine/Source/Programs/Horde/HordeStorage/Horde.Storage/Implementation/Blob/IBlobService.cs
+++ b/Engine/Source/Programs/Horde/HordeStorage/Horde.Storage/Implementation/Blob/IBlobService.cs
@@ -101,6 +101,8 @@ public class BlobService : IBlobService
 
     public async Task<BlobIdentifier> PutObject(NamespaceId ns, IBufferedPayload payload, BlobIdentifier identifier)
     {
+        using Scope _ = Tracer.Instance.StartActive("put_blob");
+
         await using Stream hashStream = payload.GetStream();
         BlobIdentifier id = await VerifyContentMatchesHash(hashStream, identifier);
 
@@ -115,6 +117,8 @@ public class BlobService : IBlobService
 
     public async Task<BlobIdentifier> PutObject(NamespaceId ns, byte[] payload, BlobIdentifier identifier)
     {
+        using Scope _ = Tracer.Instance.StartActive("put_blob");
+
         await using Stream hashStream = new MemoryStream(payload);
         BlobIdentifier id = await VerifyContentMatchesHash(hashStream, identifier);
 
@@ -129,6 +133,7 @@ public class BlobService : IBlobService
     public async Task<BlobIdentifier> PutObjectKnownHash(NamespaceId ns, IBufferedPayload content,
         BlobIdentifier identifier)
     {
+        using Scope _ = Tracer.Instance.StartActive("put_blob");
         Task<BlobIdentifier> putObjectTask = PutObjectToStores(ns, content, identifier);
         Task addToBlobIndexTask = _blobIndex.AddBlobToIndex(ns, identifier);
 
@@ -141,6 +146,10 @@ public class BlobService : IBlobService
     {
         await Parallel.ForEachAsync(_blobStores, async (store, cancellationToken) =>
         {
+            using Scope scope = Tracer.Instance.StartActive("put_blob_to_store");
+            scope.Span.ResourceName = identifier.ToString();
+            scope.Span.SetTag("store", store.ToString());
+
             await using Stream s = bufferedPayload.GetStream();
             await store.PutObject(ns, s, identifier);
         });
