Horde: Add a "HibernatingTime" property to utilization reports, showing time that agents were autoscaled down.

diff --git a/Engine/Source/Programs/Horde/Horde.Build/Fleet/Impl/TelemetryCollection.cs b/Engine/Source/Programs/Horde/Horde.Build/Fleet/Impl/TelemetryCollection.cs
--- a/Engine/Source/Programs/Horde/Horde.Build/Fleet/Impl/TelemetryCollection.cs
+++ b/Engine/Source/Programs/Horde/Horde.Build/Fleet/Impl/TelemetryCollection.cs
@@ -34,6 +34,7 @@ namespace HordeServer.Collections.Impl
 			IReadOnlyList<IPoolUtilizationTelemetry> IUtilizationTelemetry.Pools => Pools;
 
 			public double AdminTime { get; set; }
+			public double HibernatingTime { get; set; }
 			public int UpdateIndex { get; set; }
 
 			[BsonConstructor]
@@ -48,6 +49,7 @@ namespace HordeServer.Collections.Impl
 				this.FinishTime = Other.FinishTime;
 				this.NumAgents = Other.NumAgents;
 				this.Pools = Other.Pools.ConvertAll(x => new PoolUtilizationDocument(x));
+				this.HibernatingTime = Other.HibernatingTime;
 				this.AdminTime = Other.AdminTime;
 			}
 		}
@@ -61,6 +63,7 @@ namespace HordeServer.Collections.Impl
 			IReadOnlyList<IStreamUtilizationTelemetry> IPoolUtilizationTelemetry.Streams => Streams;
 
 			public double AdminTime { get; set; }
+			public double HibernatingTime { get; set; }
 			public double OtherTime { get; set; }
 
 			[BsonConstructor]
@@ -75,6 +78,7 @@ namespace HordeServer.Collections.Impl
 				this.NumAgents = Other.NumAgents;
 				this.Streams = Other.Streams.ConvertAll(x => new StreamUtilizationDocument(x));
 				this.AdminTime = Other.AdminTime;
+				this.HibernatingTime = Other.HibernatingTime;
 				this.OtherTime = Other.OtherTime;
 			}
 		}
diff --git a/Engine/Source/Programs/Horde/Horde.Build/Fleet/Telemetry.cs b/Engine/Source/Programs/Horde/Horde.Build/Fleet/Telemetry.cs
--- a/Engine/Source/Programs/Horde/Horde.Build/Fleet/Telemetry.cs
+++ b/Engine/Source/Programs/Horde/Horde.Build/Fleet/Telemetry.cs
@@ -36,6 +36,11 @@ namespace HordeServer.Models
 		/// </summary>
 		public IReadOnlyList<IPoolUtilizationTelemetry> Pools { get; }
 
+		/// <summary>
+		/// Amount of time that agents were hibernating
+		/// </summary>
+		public double HibernatingTime { get; }
+
 		/// <summary>
 		/// Total time spent running administrative tasks (conform, etc...)
 		/// </summary>
@@ -62,6 +67,9 @@ namespace HordeServer.Models
 
 		Dictionary<PoolId, NewPoolUtilizationTelemetry> PoolsLookup { get; set; } = new Dictionary<PoolId, NewPoolUtilizationTelemetry>();
 
+		/// <inheritdoc/>
+		public double HibernatingTime { get; set; }
+
 		/// <inheritdoc/>
 		public double AdminTime { get; set; }
 
@@ -115,10 +123,15 @@ namespace HordeServer.Models
 		public IReadOnlyList<IStreamUtilizationTelemetry> Streams { get; }
 
 		/// <summary>
-		/// Admount of time spent running 
+		/// Amount of time spent running 
 		/// </summary>
 		public double AdminTime { get; }
 
+		/// <summary>
+		/// Amount of time that agents were hibernating
+		/// </summary>
+		public double HibernatingTime { get; }
+
 		/// <summary>
 		/// Amount of time spent by agents in this pool servicing other pools
 		/// </summary>
@@ -145,6 +158,9 @@ namespace HordeServer.Models
 		/// <inheritdoc/>
 		public double AdminTime { get; set; }
 
+		/// <inheritdoc/>
+		public double HibernatingTime { get; set; }
+
 		/// <inheritdoc/>
 		public double OtherTime { get; set; }
 
