- remove streams filter
- fix crash when an interval step is more hours than there is currently data for in a day
- add hibernation time

diff --git a/Engine/Source/Programs/Horde/HordeDashboard/src/backend/Api.ts b/Engine/Source/Programs/Horde/HordeDashboard/src/backend/Api.ts
--- a/Engine/Source/Programs/Horde/HordeDashboard/src/backend/Api.ts
+++ b/Engine/Source/Programs/Horde/HordeDashboard/src/backend/Api.ts
@@ -2783,6 +2783,8 @@ export type GetUtilizationTelemetryPool = {
 
 	otherTime: number;
 
+	hibernatingTime?: number | undefined;
+
 	streams: GetUtilizationTelemetryStream[];
 }
 
