Horde: Store requested device platforms on v1 reservations
jira none
#rnx

diff --git a/Engine/Source/Programs/Horde/Horde.Build/Controllers/DevicesController.cs b/Engine/Source/Programs/Horde/Horde.Build/Controllers/DevicesController.cs
--- a/Engine/Source/Programs/Horde/Horde.Build/Controllers/DevicesController.cs
+++ b/Engine/Source/Programs/Horde/Horde.Build/Controllers/DevicesController.cs
@@ -567,7 +567,7 @@ namespace HordeServer.Controllers
 
 				}					
 
-				RequestedDevices.Add(new DeviceRequestData(PlatformIdValue, DeviceRequest.IncludeModels, DeviceRequest.ExcludeModels));
+				RequestedDevices.Add(new DeviceRequestData(PlatformIdValue, PlatformIdValue.ToString(), DeviceRequest.IncludeModels, DeviceRequest.ExcludeModels));
 
 			}
 
@@ -805,7 +805,7 @@ namespace HordeServer.Controllers
 				}
 
 
-				RequestedDevices.Add(new DeviceRequestData(PlatformId, IncludeModels, ExcludeModels));
+				RequestedDevices.Add(new DeviceRequestData(PlatformId, PlatformName, IncludeModels, ExcludeModels));
 
 			}
 
@@ -908,18 +908,39 @@ namespace HordeServer.Controllers
 				return BadRequest($"Unknown device {DeviceName}");
 			}
 
-			DevicePlatformMapV1 MapV1 = await DeviceService.GetPlatformMapV1();
+			IDeviceReservation? Reservation = await DeviceService.TryGetDeviceReservation(Device.Id);
 
-			string? Platform;
-			if (!MapV1.PlatformReverseMap.TryGetValue(Device.PlatformId, out Platform))
+			string? PlatformName = null;
+
+			if (Reservation != null && Reservation.Devices.Count == Reservation.RequestedDevicePlatforms.Count)
 			{
-				return BadRequest($"Unable to map platform for {DeviceName} : {Device.PlatformId}");
+				for (int i = 0; i < Reservation.Devices.Count; i++)
+				{
+					if (Reservation.Devices[i] == Device.Id)
+					{
+						PlatformName = Reservation.RequestedDevicePlatforms[i];
+					}
+				}
+			}
+
+			DevicePlatformMapV1 MapV1 = await DeviceService.GetPlatformMapV1();
+
+			if (string.IsNullOrEmpty(PlatformName))
+			{				
+				if (!MapV1.PlatformReverseMap.TryGetValue(Device.PlatformId, out PlatformName))
+				{
+					return BadRequest($"Unable to map platform for {DeviceName} : {Device.PlatformId}");
+				}
+			}
+
+			if (string.IsNullOrEmpty(PlatformName)) {
+				return BadRequest($"Unable to get platform for {DeviceName} from reservation or mapping : {Device.PlatformId}");
 			}
 
 			GetLegacyDeviceResponse Response = new GetLegacyDeviceResponse();
 
 			Response.Name = Device.Name;
-			Response.Type = Platform;
+			Response.Type = PlatformName;
 			Response.IPOrHostName = Device.Address ?? "";
 			Response.AvailableStartTime = "00:00:00";
 			Response.AvailableEndTime = "00:00:00";
diff --git a/Engine/Source/Programs/Horde/Horde.Build/Devices/IDevice.cs b/Engine/Source/Programs/Horde/Horde.Build/Devices/IDevice.cs
--- a/Engine/Source/Programs/Horde/Horde.Build/Devices/IDevice.cs
+++ b/Engine/Source/Programs/Horde/Horde.Build/Devices/IDevice.cs
@@ -70,6 +70,12 @@ namespace HordeServer.Models
 		/// </summary>
 		public List<DeviceId> Devices { get; }
 
+		/// <summary>
+		/// The requested device platforms for reservation, which may differ from IDevice platform due to devices that support more than one platform, or legacy platforms
+		/// </summary>
+		public List<string> RequestedDevicePlatforms { get; }
+
+
 		/// <summary>
 		/// The legacy reservation system guid, to be removed once can update Gauntlet client in all streams
 		/// </summary>
diff --git a/Engine/Source/Programs/Horde/Horde.Build/Devices/IDeviceCollection.cs b/Engine/Source/Programs/Horde/Horde.Build/Devices/IDeviceCollection.cs
--- a/Engine/Source/Programs/Horde/Horde.Build/Devices/IDeviceCollection.cs
+++ b/Engine/Source/Programs/Horde/Horde.Build/Devices/IDeviceCollection.cs
@@ -24,6 +24,10 @@ namespace HordeServer.Collections
 		/// </summary>
 		public DevicePlatformId PlatformId { get; set; }
 
+		/// <summary>
+		/// The platform client requested, which can differ from the devices platformid due to mapping
+		/// </summary>
+		public string RequestedPlatform { get; set; }
 
 		/// <summary>
 		/// Models to include for this request
@@ -38,11 +42,12 @@ namespace HordeServer.Collections
 		/// <summary>
 		/// Constructor
 		/// </summary>
-		public DeviceRequestData(DevicePlatformId PlatformId, List<string>? IncludeModels = null, List<string>? ExcludeModels = null)
+		public DeviceRequestData(DevicePlatformId PlatformId, string RequestedPlatform, List<string>? IncludeModels = null, List<string>? ExcludeModels = null)
 		{
 			this.PlatformId = PlatformId;
 			this.IncludeModels = IncludeModels ?? new List<string>();
 			this.ExcludeModels = ExcludeModels ?? new List<string>();
+			this.RequestedPlatform = RequestedPlatform;
 		}
 	}
 
diff --git a/Engine/Source/Programs/Horde/Horde.Build/Devices/Impl/DeviceCollection.cs b/Engine/Source/Programs/Horde/Horde.Build/Devices/Impl/DeviceCollection.cs
--- a/Engine/Source/Programs/Horde/Horde.Build/Devices/Impl/DeviceCollection.cs
+++ b/Engine/Source/Programs/Horde/Horde.Build/Devices/Impl/DeviceCollection.cs
@@ -130,10 +130,12 @@ namespace HordeServer.Collections.Impl
 			public string? ReservationDetails { get; set; }
 
 			/// <summary>
-			/// 
+			/// DeviceIds in the reservation
 			/// </summary>
 			public List<DeviceId> Devices { get; set; } = new List<DeviceId>();
 
+			public List<string> RequestedDevicePlatforms { get; set; } = new List<string>();
+
 			public DateTime CreateTimeUtc { get; set; }
 
 			public DateTime UpdateTimeUtc { get; set; }
@@ -147,11 +149,12 @@ namespace HordeServer.Collections.Impl
 
 			}
 
-			public DeviceReservationDocument(ObjectId Id, DevicePoolId PoolId, List<DeviceId> Devices, DateTime CreateTimeUtc, string? Hostname, string? ReservationDetails, string? JobId, string? StepId)
+			public DeviceReservationDocument(ObjectId Id, DevicePoolId PoolId, List<DeviceId> Devices, List<string> RequestedDevicePlatforms, DateTime CreateTimeUtc, string? Hostname, string? ReservationDetails, string? JobId, string? StepId)
 			{
 				this.Id = Id;
 				this.PoolId = PoolId;
 				this.Devices = Devices;
+				this.RequestedDevicePlatforms = RequestedDevicePlatforms;
 				this.CreateTimeUtc = CreateTimeUtc;
 				this.UpdateTimeUtc = CreateTimeUtc;
 				this.Hostname = Hostname;
@@ -541,6 +544,7 @@ namespace HordeServer.Collections.Impl
             }
 			
 			HashSet<DeviceId> Allocated = new HashSet<DeviceId>();
+			Dictionary<DeviceId, string> PlatformRequestMap = new Dictionary<DeviceId, string>();
 
 			List<DeviceReservationDocument> PoolReservations = await Reservations.Find(x => x.PoolId == PoolId).ToListAsync();
 
@@ -613,6 +617,7 @@ namespace HordeServer.Collections.Impl
 				}
 
 				Allocated.Add(Device.Id);
+				PlatformRequestMap.Add(Device.Id, Data.RequestedPlatform);
 			}
 
 			// update reservation time and utilization for allocated devices 
@@ -644,8 +649,11 @@ namespace HordeServer.Collections.Impl
 				await Devices.FindOneAndUpdateAsync<DeviceDocument>(x => x.Id == Id, DeviceBuilder.Combine(DeviceUpdates));
 			}
 
+			List<DeviceId> DeviceIds = Allocated.ToList();
+			List<string> RequestedPlatforms = DeviceIds.Select(x => PlatformRequestMap[x]).ToList();
+
 			// Create new reservation
-			DeviceReservationDocument NewReservation = new DeviceReservationDocument(ObjectId.GenerateNewId(), PoolId, Allocated.ToList(), ReservationTimeUtc, Hostname, ReservationDetails, JobId, StepId);
+			DeviceReservationDocument NewReservation = new DeviceReservationDocument(ObjectId.GenerateNewId(), PoolId, DeviceIds, RequestedPlatforms, ReservationTimeUtc, Hostname, ReservationDetails, JobId, StepId);
 			await Reservations.InsertOneAsync(NewReservation);
 
 			return NewReservation;
