Horde Storage

Removed transaction to update last access time, can result in empty rows in the objects table which will generate an error. If it does we can update the code to just consider this as a non existent object. This will generate much less load on the database, and doesnt seem worth for something which we could just handle.

Also removed other call to UpdateIfAsync for finalizing a blob, which isnt needed as its not actually using any conditional statament (so UpdateAsync is the same thing)

diff --git a/Engine/Source/Programs/Horde/HordeStorage/Horde.Storage/Implementation/Objects/ScyllaReferencesStore.cs b/Engine/Source/Programs/Horde/HordeStorage/Horde.Storage/Implementation/Objects/ScyllaReferencesStore.cs
--- a/Engine/Source/Programs/Horde/HordeStorage/Horde.Storage/Implementation/Objects/ScyllaReferencesStore.cs
+++ b/Engine/Source/Programs/Horde/HordeStorage/Horde.Storage/Implementation/Objects/ScyllaReferencesStore.cs
@@ -103,9 +103,7 @@ namespace Horde.Storage.Implementation
         {
             using Scope _ = Tracer.Instance.StartActive("scylla.finalize");
 
-            AppliedInfo<ScyllaObject> info = await _mapper.UpdateIfAsync<ScyllaObject>("SET is_finalized=true WHERE namespace=? AND bucket=? AND name=?", ns.ToString(), bucket.ToString(), name.ToString());
-            if (!info.Applied)
-                throw new Exception($"Failed to finalize object in namespace {ns} {bucket} {name}");
+            await _mapper.UpdateAsync<ScyllaObject>("SET is_finalized=true WHERE namespace=? AND bucket=? AND name=?", ns.ToString(), bucket.ToString(), name.ToString());
         }
 
         private async Task CreateTTLRecord(NamespaceId ns, BucketId bucket, IoHashKey name, sbyte partitionIndex)
@@ -150,7 +148,7 @@ namespace Horde.Storage.Implementation
 
             await updateObjectTracking;*/
 
-            await _mapper.UpdateIfAsync<ScyllaObject>("SET last_access_time = ? WHERE namespace = ? AND bucket = ? AND name = ? IF EXISTS", lastAccessTime, ns.ToString(), bucket.ToString(), name.ToString());
+            await _mapper.UpdateAsync<ScyllaObject>("SET last_access_time = ? WHERE namespace = ? AND bucket = ? AND name = ?", lastAccessTime, ns.ToString(), bucket.ToString(), name.ToString());
 
             
         }
