Horde: Run Tick() callbacks deterministically for testing. The IClock interface now allows registering callbacks at particular intervals, and has a deterministic implementation for FakeClock.

diff --git a/Engine/Source/Programs/Horde/Horde.Build.Tests/AgentServiceTest.cs b/Engine/Source/Programs/Horde/Horde.Build.Tests/AgentServiceTest.cs
--- a/Engine/Source/Programs/Horde/Horde.Build.Tests/AgentServiceTest.cs
+++ b/Engine/Source/Programs/Horde/Horde.Build.Tests/AgentServiceTest.cs
@@ -46,7 +46,7 @@ namespace HordeServerTests
 		        "test");
 	        
 	        Assert.IsTrue(AgentService.AuthorizeSession(Agent, GetUser(Agent)));
-	        Clock.Advance(TimeSpan.FromMinutes(20));
+	        await Clock.AdvanceAsync(TimeSpan.FromMinutes(20));
 	        Assert.IsFalse(AgentService.AuthorizeSession(Agent, GetUser(Agent)));
         }
 
diff --git a/Engine/Source/Programs/Horde/Horde.Build.Tests/LogFileServiceTests.cs b/Engine/Source/Programs/Horde/Horde.Build.Tests/LogFileServiceTests.cs
--- a/Engine/Source/Programs/Horde/Horde.Build.Tests/LogFileServiceTests.cs
+++ b/Engine/Source/Programs/Horde/Horde.Build.Tests/LogFileServiceTests.cs
@@ -1,5 +1,6 @@
 // Copyright Epic Games, Inc. All Rights Reserved.
 
+using System;
 using System.Collections.Generic;
 using System.Diagnostics;
 using System.IO;
@@ -9,6 +10,7 @@ using System.Runtime.InteropServices.ComTypes;
 using System.Security.Claims;
 using System.Text;
 using System.Threading.Tasks;
+using HordeCommon;
 using HordeServer;
 using HordeServer.Api;
 using HordeServer.Collections.Impl;
@@ -38,6 +40,7 @@ namespace HordeServerTests
 	[TestClass]
     public class LogFileServiceTest : DatabaseIntegrationTest
     {
+		private readonly FakeClock Clock;
         private readonly LogFileService LogFileService;
 
         public LogFileServiceTest()
@@ -52,7 +55,8 @@ namespace HordeServerTests
 
 			ILogBuilder LogBuilder = new RedisLogBuilder(GetRedisConnectionPool(), NullLogger.Instance);
 			ILogStorage LogStorage = new PersistentLogStorage(new TransientStorageBackend().ForType<PersistentLogStorage>(), NullLogger<PersistentLogStorage>.Instance);
-			LogFileService = new LogFileService(LogFileCollection, null!, LogBuilder, LogStorage, Logger);
+			Clock = new FakeClock();
+			LogFileService = new LogFileService(LogFileCollection, null!, LogBuilder, LogStorage, Clock, Logger);
         }
 
 		[TestMethod]
@@ -116,7 +120,7 @@ namespace HordeServerTests
 	        return new StreamReader(Stream).ReadToEnd();
         }
         
-        public static async Task WriteLogLifecycle(ILogFileService Lfs, int MaxChunkLength)
+        public async Task WriteLogLifecycle(ILogFileService Lfs, int MaxChunkLength)
         {
 			JobId JobId = JobId.GenerateNewId();
             ILogFile LogFile = await Lfs.CreateLogFileAsync(JobId, null, LogType.Text);
@@ -131,8 +135,7 @@ namespace HordeServerTests
 
             // First write with flush. Will become chunk #1
             LogFile = (await Lfs.WriteLogDataAsync(LogFile, Offset, LineIndex, Encoding.ASCII.GetBytes(Str1), true, MaxChunkLength))!;
-            await (Lfs as LogFileService)!.TickOnlyForTestingAsync();
-            await Task.Delay(2000);
+			await Clock.AdvanceAsync(TimeSpan.FromSeconds(2.0));
             Assert.AreEqual(Str1, await ReadLogFile(Lfs, LogFile, 0, Str1.Length));
             Assert.AreEqual(Str1, await ReadLogFile(Lfs, LogFile, 0, Str1.Length + 100)); // Reading too far is valid?
             await AssertMetadata(Lfs, LogFile, Str1.Length, 1);
@@ -146,9 +149,8 @@ namespace HordeServerTests
             Offset += Str1.Length;
             LineIndex += Str1.Count(f => f == '\n');
             LogFile = (await Lfs.WriteLogDataAsync(LogFile, Offset, LineIndex, Encoding.ASCII.GetBytes(Str2), false, MaxChunkLength))!;
-            await (Lfs as LogFileService)!.TickOnlyForTestingAsync();
-            await Task.Delay(2000);
-            Assert.AreEqual(Str1 + Str2, await ReadLogFile(Lfs, LogFile, 0, Str1.Length + Str2.Length));
+			await Clock.AdvanceAsync(TimeSpan.FromSeconds(2.0));
+			Assert.AreEqual(Str1 + Str2, await ReadLogFile(Lfs, LogFile, 0, Str1.Length + Str2.Length));
             await AssertMetadata(Lfs, LogFile, Str1.Length + Str2.Length, 3); // FIXME: what are max line index?
             await AssertChunk(Lfs, LogFile.Id, 2, 0, 0, Str1.Length, 0);
             await AssertChunk(Lfs, LogFile.Id, 2, 1, Str1.Length, 0, 1); // Last chunk have length zero as it's being written
@@ -161,9 +163,8 @@ namespace HordeServerTests
             Offset += Str2.Length;
             LineIndex += Str3.Count(f => f == '\n');
             LogFile = (await Lfs.WriteLogDataAsync(LogFile, Offset, LineIndex, Encoding.ASCII.GetBytes(Str3), false, MaxChunkLength))!;
-            await (Lfs as LogFileService)!.TickOnlyForTestingAsync();
-            await Task.Delay(2000);
-            Assert.AreEqual(Str1 + Str2 + Str3, await ReadLogFile(Lfs, LogFile, 0, Str1.Length + Str2.Length + Str3.Length));
+			await Clock.AdvanceAsync(TimeSpan.FromSeconds(2.0));
+			Assert.AreEqual(Str1 + Str2 + Str3, await ReadLogFile(Lfs, LogFile, 0, Str1.Length + Str2.Length + Str3.Length));
             //await AssertMetadata(Lfs, LogFile, Str1.Length + Str2.Length + Str3.Length, 8);
             // Since no flush has happened, chunks should be identical to last write
             await AssertChunk(Lfs, LogFile.Id, 2, 0, 0, Str1.Length, 0);
@@ -178,9 +179,8 @@ namespace HordeServerTests
             Offset += Str3.Length;
             LineIndex += Str4.Count(f => f == '\n');
             LogFile = (await Lfs.WriteLogDataAsync(LogFile, Offset, LineIndex, Encoding.ASCII.GetBytes(Str4), true, MaxChunkLength))!;
-            await (Lfs as LogFileService)!.TickOnlyForTestingAsync();
-            await Task.Delay(2000);
-            Assert.AreEqual(Str1 + Str2 + Str3 + Str4,
+			await Clock.AdvanceAsync(TimeSpan.FromSeconds(2.0));
+			Assert.AreEqual(Str1 + Str2 + Str3 + Str4,
                 await ReadLogFile(Lfs, LogFile, 0, Str1.Length + Str2.Length + Str3.Length + Str4.Length));
             Assert.AreEqual(Str3 + Str4,
                 await ReadLogFile(Lfs, LogFile, Str1.Length + Str2.Length, Str3.Length + Str4.Length));
@@ -210,10 +210,9 @@ namespace HordeServerTests
             // Dividing it in two like this will work however
             LogFile = (await Lfs.WriteLogDataAsync(LogFile, Offset, LineIndex, Encoding.ASCII.GetBytes(A), false, MaxChunkLength))!;
 			LogFile = (await Lfs.WriteLogDataAsync(LogFile, Offset + A.Length, LineIndex + 1, Encoding.ASCII.GetBytes(B), true, MaxChunkLength))!;
-            
-            await (Lfs as LogFileService)!.TickOnlyForTestingAsync();
-            await Task.Delay(2000);
-            await AssertMetadata(Lfs, LogFile, Str1.Length + Str2.Length + Str3.Length + Str4.Length + Str5.Length, 9);
+
+			await Clock.AdvanceAsync(TimeSpan.FromSeconds(2.0));
+			await AssertMetadata(Lfs, LogFile, Str1.Length + Str2.Length + Str3.Length + Str4.Length + Str5.Length, 9);
             await AssertChunk(Lfs, LogFile.Id, 4, 0, 0, Str1.Length, 0);
             await AssertChunk(Lfs, LogFile.Id, 4, 1, Str1.Length, Str2.Length + Str3.Length + Str4.Length, 1);
             await AssertChunk(Lfs, LogFile.Id, 4, 2, Offset, A.Length, 7);
diff --git a/Engine/Source/Programs/Horde/Horde.Build.Tests/LogIndexingTests.cs b/Engine/Source/Programs/Horde/Horde.Build.Tests/LogIndexingTests.cs
--- a/Engine/Source/Programs/Horde/Horde.Build.Tests/LogIndexingTests.cs
+++ b/Engine/Source/Programs/Horde/Horde.Build.Tests/LogIndexingTests.cs
@@ -22,6 +22,7 @@ using System.Linq;
 using System.Text;
 using System.Threading.Tasks;
 using HordeServer.Collections.Impl;
+using HordeCommon;
 
 namespace HordeServerTests
 {
@@ -49,7 +50,7 @@ namespace HordeServerTests
 
 			ILogBuilder LogBuilder = new LocalLogBuilder();
 			ILogStorage LogStorage = new LocalLogStorage(20, new NullLogStorage());
-			LogFileService = new LogFileService(LogFileCollection, null!, LogBuilder, LogStorage, Logger);
+			LogFileService = new LogFileService(LogFileCollection, null!, LogBuilder, LogStorage, new FakeClock(), Logger);
 		}
 
 		[TestMethod]
diff --git a/Engine/Source/Programs/Horde/Horde.Build.Tests/SchedulerTests.cs b/Engine/Source/Programs/Horde/Horde.Build.Tests/SchedulerTests.cs
--- a/Engine/Source/Programs/Horde/Horde.Build.Tests/SchedulerTests.cs
+++ b/Engine/Source/Programs/Horde/Horde.Build.Tests/SchedulerTests.cs
@@ -90,7 +90,7 @@ namespace HordeServerTests
 			Schedule.Files = Files.ToList();
 			await SetScheduleAsync(Schedule);
 
-			Clock.Advance(TimeSpan.FromHours(1.25));
+			await Clock.AdvanceAsync(TimeSpan.FromHours(1.25));
 			await ScheduleService.TickForTestingAsync();
 
 			return await GetNewJobs();
@@ -194,7 +194,7 @@ namespace HordeServerTests
 			Assert.AreEqual(0, Jobs1.Count);
 
 			// Trigger a job
-			Clock.Advance(TimeSpan.FromHours(1.25));
+			await Clock.AdvanceAsync(TimeSpan.FromHours(1.25));
 			await ScheduleService.TickForTestingAsync();
 
 			List<IJob> Jobs2 = await GetNewJobs();
@@ -221,7 +221,7 @@ namespace HordeServerTests
 			Assert.AreEqual(0, Jobs1.Count);
 
 			// Trigger a job
-			Clock.Advance(TimeSpan.FromHours(1.25));
+			await Clock.AdvanceAsync(TimeSpan.FromHours(1.25));
 			await ScheduleService.TickForTestingAsync();
 
 			List<IJob> Jobs2 = await GetNewJobs();
@@ -235,7 +235,7 @@ namespace HordeServerTests
 			Assert.AreEqual(Clock.UtcNow, Schedule2.LastTriggerTime);
 
 			// Trigger another job
-			Clock.Advance(TimeSpan.FromHours(0.5));
+			await Clock.AdvanceAsync(TimeSpan.FromHours(0.5));
 			await ScheduleService.TickForTestingAsync();
 
 			List<IJob> Jobs3 = await GetNewJobs();
@@ -268,7 +268,7 @@ namespace HordeServerTests
 			PerforceService.AddChange("//UE5/Main", 105, Bob, "", new string[] { "foo.uasset" });
 			PerforceService.AddChange("//UE5/Main", 106, Bob, "", new string[] { "foo.cpp" });
 
-			Clock.Advance(TimeSpan.FromHours(1.25));
+			await Clock.AdvanceAsync(TimeSpan.FromHours(1.25));
 			await ScheduleService.TickForTestingAsync();
 
 			List<IJob> Jobs2 = await GetNewJobs();
@@ -305,7 +305,7 @@ namespace HordeServerTests
 			PerforceService.AddChange("//UE5/Main", 105, Bob, "", new string[] { "foo.uasset" });
 			PerforceService.AddChange("//UE5/Main", 106, Bob, "", new string[] { "foo.cpp" });
 
-			Clock.Advance(TimeSpan.FromHours(1.25));
+			await Clock.AdvanceAsync(TimeSpan.FromHours(1.25));
 			await ScheduleService.TickForTestingAsync();
 
 			List<IJob> Jobs2 = await GetNewJobs();
@@ -330,7 +330,7 @@ namespace HordeServerTests
 			await SetScheduleAsync(Schedule);
 
 			// Trigger a job
-			Clock.Advance(TimeSpan.FromHours(1.25));
+			await Clock.AdvanceAsync(TimeSpan.FromHours(1.25));
 			await ScheduleService.TickForTestingAsync();
 
 			List<IJob> Jobs2 = await GetNewJobs();
@@ -339,7 +339,7 @@ namespace HordeServerTests
 			Assert.AreEqual(100, Jobs2[0].CodeChange);
 
 			// Test that another job does not trigger
-			Clock.Advance(TimeSpan.FromHours(0.5));
+			await Clock.AdvanceAsync(TimeSpan.FromHours(0.5));
 			await ScheduleService.TickForTestingAsync();
 
 			List<IJob> Jobs3 = await GetNewJobs();
@@ -349,7 +349,7 @@ namespace HordeServerTests
 			await JobService.UpdateJobAsync(Jobs2[0], AbortedByUserId: KnownUsers.System);
 
 			// Test that another job does not trigger
-			Clock.Advance(TimeSpan.FromHours(0.5));
+			await Clock.AdvanceAsync(TimeSpan.FromHours(0.5));
 			await ScheduleService.TickForTestingAsync();
 
 			List<IJob> Jobs4 = await GetNewJobs();
@@ -370,7 +370,7 @@ namespace HordeServerTests
 			IStream Stream = await SetScheduleAsync(Schedule);
 
 			// Trigger a job
-			Clock.Advance(TimeSpan.FromHours(1.25));
+			await Clock.AdvanceAsync(TimeSpan.FromHours(1.25));
 			await ScheduleService.TickForTestingAsync();
 
 			List<IJob> Jobs2 = await GetNewJobs();
@@ -379,7 +379,7 @@ namespace HordeServerTests
 			Assert.AreEqual(100, Jobs2[0].CodeChange);
 
 			// Check another job does not trigger due to the change above
-			Clock.Advance(TimeSpan.FromHours(1.25));
+			await Clock.AdvanceAsync(TimeSpan.FromHours(1.25));
 			await ScheduleService.TickForTestingAsync();
 
 			List<IJob> Jobs3 = await GetNewJobs();
@@ -433,7 +433,7 @@ namespace HordeServerTests
 			await GetNewJobs();
 
 			// Tick the schedule and make sure it doesn't trigger
-			Clock.Advance(TimeSpan.FromMinutes(30.0));
+			await Clock.AdvanceAsync(TimeSpan.FromMinutes(30.0));
 			await ScheduleService.TickForTestingAsync();
 			List<IJob> Jobs3 = await GetNewJobs();
 			Assert.AreEqual(0, Jobs3.Count);
@@ -466,7 +466,7 @@ namespace HordeServerTests
 			Schedule.MaxActive = 2;
 			await SetScheduleAsync(Schedule);
 
-			Clock.Advance(TimeSpan.FromHours(1.25));
+			await Clock.AdvanceAsync(TimeSpan.FromHours(1.25));
 			await ScheduleService.TickForTestingAsync();
 
 			List<IJob> Jobs1 = await GetNewJobs();
@@ -503,15 +503,15 @@ namespace HordeServerTests
 
 			IStreamCollection StreamCollection = ServiceProvider.GetRequiredService<IStreamCollection>();
 			await StreamCollection.TryUpdatePauseStateAsync(Stream, NewPausedUntil: StartTime.AddHours(5), NewPauseComment: "testing");
-			
+
 			// Try trigger a job. No job should be scheduled as the stream is paused
-			Clock.Advance(TimeSpan.FromHours(1.25));
+			await Clock.AdvanceAsync(TimeSpan.FromHours(1.25));
 			await ScheduleService.TickForTestingAsync();
 			List<IJob> Jobs2 = await GetNewJobs();
 			Assert.AreEqual(0, Jobs2.Count);
 
 			// Advance time beyond the pause period. A build should now trigger
-			Clock.Advance(TimeSpan.FromHours(5.25));
+			await Clock.AdvanceAsync(TimeSpan.FromHours(5.25));
 			await ScheduleService.TickForTestingAsync();
 
 			List<IJob> Jobs3 = await GetNewJobs();
diff --git a/Engine/Source/Programs/Horde/Horde.Build/Agents/Impl/AgentService.cs b/Engine/Source/Programs/Horde/Horde.Build/Agents/Impl/AgentService.cs
--- a/Engine/Source/Programs/Horde/Horde.Build/Agents/Impl/AgentService.cs
+++ b/Engine/Source/Programs/Horde/Horde.Build/Agents/Impl/AgentService.cs
@@ -31,7 +31,7 @@ namespace HordeServer.Services
 	/// <summary>
 	/// Wraps funtionality for manipulating agents
 	/// </summary>
-	public sealed class AgentService : TickedBackgroundService
+	public sealed class AgentService : IHostedService, IDisposable
 	{
 		/// <summary>
 		/// Maximum time between updates for an agent to be considered online
@@ -48,14 +48,7 @@ namespace HordeServer.Services
 		/// </summary>
 		public static readonly TimeSpan SessionRenewTime = TimeSpan.FromSeconds(50);
 
-		/// <summary>
-		/// The ACL service instance
-		/// </summary>
 		AclService AclService;
-
-		/// <summary>
-		/// The downtime service instance
-		/// </summary>
 		IDowntimeService DowntimeService;
 
 		/// <summary>
@@ -63,61 +56,28 @@ namespace HordeServer.Services
 		/// </summary>
 		public IAgentCollection Agents { get; }
 
-		/// <summary>
-		/// Collection of lease documents
-		/// </summary>
 		ILeaseCollection Leases;
-
-		/// <summary>
-		/// Collection of session documents
-		/// </summary>
 		ISessionCollection Sessions;
-
-		/// <summary>
-		/// DogStatsD metric client
-		/// </summary>
 		IDogStatsd DogStatsd;
-
-		/// <summary>
-		/// List of task sources
-		/// </summary>
 		ITaskSource[] TaskSources;
-
-		/// <summary>
-		/// The application lifetime instance
-		/// </summary>
 		IHostApplicationLifetime ApplicationLifetime;
+		IClock Clock;
+		ILogger Logger;
+		ITicker Ticker;
 
-		/// <summary>
-		/// Log output writer
-		/// </summary>
-		ILogger<AgentService> Logger;
-		
-		/// <summary>
-		/// Clock
-		/// </summary>
-		readonly IClock Clock;
-
-		/// <summary>
-		/// Lazily updated list of current pools
-		/// </summary>
+		// Lazily updated list of current pools
 		AsyncCachedValue<List<IPool>> PoolsList;
 
-		/// <summary>
-		/// All the agents currently performing a long poll for work on this server
-		/// </summary>
+		// All the agents currently performing a long poll for work on this server
 		Dictionary<AgentId, CancellationTokenSource> WaitingAgents = new Dictionary<AgentId, CancellationTokenSource>();
 
-		/// <summary>
-		/// Subscription for update events
-		/// </summary>
+		// Subscription for update events
 		IDisposable Subscription;
 
 		/// <summary>
 		/// Constructor
 		/// </summary>
-		public AgentService(IAgentCollection Agents, ILeaseCollection Leases, ISessionCollection Sessions, AclService AclService, IDowntimeService DowntimeService, IPoolCollection PoolCollection, IDogStatsd DogStatsd, IEnumerable<ITaskSource> TaskSources, IHostApplicationLifetime ApplicationLifetime, ILogger<AgentService> Logger, IClock Clock)
-			: base(TimeSpan.FromSeconds(30.0), Logger)
+		public AgentService(IAgentCollection Agents, ILeaseCollection Leases, ISessionCollection Sessions, AclService AclService, IDowntimeService DowntimeService, IPoolCollection PoolCollection, IDogStatsd DogStatsd, IEnumerable<ITaskSource> TaskSources, IHostApplicationLifetime ApplicationLifetime, IClock Clock, ILogger<AgentService> Logger)
 		{
 			this.Agents = Agents;
 			this.Leases = Leases;
@@ -128,17 +88,24 @@ namespace HordeServer.Services
 			this.DogStatsd = DogStatsd;
 			this.TaskSources = TaskSources.ToArray();
 			this.ApplicationLifetime = ApplicationLifetime;
-			this.Logger = Logger;
 			this.Clock = Clock;
+			this.Ticker = Clock.CreateTicker(TimeSpan.FromSeconds(30.0), TickAsync, Logger);
+			this.Logger = Logger;
 
 			Subscription = Agents.SubscribeToUpdateEventsAsync(OnAgentUpdate).Result;
 		}
 
 		/// <inheritdoc/>
-		public override void Dispose()
+		public Task StartAsync(CancellationToken cancellationToken) => Ticker.StartAsync();
+
+		/// <inheritdoc/>
+		public Task StopAsync(CancellationToken cancellationToken) => Ticker.StopAsync();
+
+		/// <inheritdoc/>
+		public void Dispose()
 		{
-			base.Dispose();
 			Subscription.Dispose();
+			Ticker.Dispose();
 		}
 
 		/// <summary>
@@ -802,7 +769,7 @@ namespace HordeServer.Services
 		/// </summary>
 		/// <param name="StoppingToken">Token indicating the service is shutting down</param>
 		/// <returns>Async task</returns>
-		protected override async Task TickAsync(CancellationToken StoppingToken)
+		async ValueTask TickAsync(CancellationToken StoppingToken)
 		{
 			while (!StoppingToken.IsCancellationRequested)
 			{
diff --git a/Engine/Source/Programs/Horde/Horde.Build/Compute/Impl/ComputeService.cs b/Engine/Source/Programs/Horde/Horde.Build/Compute/Impl/ComputeService.cs
--- a/Engine/Source/Programs/Horde/Horde.Build/Compute/Impl/ComputeService.cs
+++ b/Engine/Source/Programs/Horde/Horde.Build/Compute/Impl/ComputeService.cs
@@ -129,17 +129,16 @@ namespace HordeServer.Compute.Impl
 	/// <summary>
 	/// Dispatches remote actions. Does not implement any cross-pod communication to satisfy leases; only agents connected to this server instance will be stored.
 	/// </summary>
-	class ComputeService : TaskSourceBase<ComputeTaskMessage>, IComputeService, IDisposable
+	class ComputeService : TaskSourceBase<ComputeTaskMessage>, IComputeService, IHostedService, IDisposable
 	{
 		public override string Type => "Compute";
 
 		public static NamespaceId DefaultNamespaceId { get; } = new NamespaceId("default");
 
 		IObjectCollection ObjectCollection;
-		IPoolCollection PoolCollection;
 		ITaskScheduler<IoHash, ComputeTaskInfo> TaskScheduler;
 		RedisMessageQueue<ComputeTaskStatus> MessageQueue;
-		BackgroundTick ExpireTasksTicker;
+		ITicker ExpireTasksTicker;
 		IMemoryCache RequirementsCache;
 		LazyCachedValue<Task<List<IPool>>> CachedPools;
 		ILogger Logger;
@@ -147,17 +146,12 @@ namespace HordeServer.Compute.Impl
 		/// <summary>
 		/// Constructor
 		/// </summary>
-		/// <param name="Redis">Redis instance</param>
-		/// <param name="ObjectCollection"></param>
-		/// <param name="PoolCollection">Collection of pool documents</param>
-		/// <param name="Logger">The logger instance</param>
-		public ComputeService(IDatabase Redis, IObjectCollection ObjectCollection, IPoolCollection PoolCollection, ILogger<ComputeService> Logger)
+		public ComputeService(IDatabase Redis, IObjectCollection ObjectCollection, IPoolCollection PoolCollection, IClock Clock, ILogger<ComputeService> Logger)
 		{
 			this.ObjectCollection = ObjectCollection;
-			this.PoolCollection = PoolCollection;
 			this.TaskScheduler = new RedisTaskScheduler<IoHash, ComputeTaskInfo>(Redis, "compute/tasks/", Logger);
 			this.MessageQueue = new RedisMessageQueue<ComputeTaskStatus>(Redis, "compute/messages/");
-			this.ExpireTasksTicker = new BackgroundTick(ExpireTasksAsync, TimeSpan.FromMinutes(2.0), Logger);
+			this.ExpireTasksTicker = Clock.CreateTicker(TimeSpan.FromMinutes(2.0), ExpireTasksAsync, Logger);
 			this.RequirementsCache = new MemoryCache(new MemoryCacheOptions());
 			this.CachedPools = new LazyCachedValue<Task<List<IPool>>>(() => PoolCollection.GetAsync(), TimeSpan.FromSeconds(30.0));
 			this.Logger = Logger;
@@ -165,12 +159,18 @@ namespace HordeServer.Compute.Impl
 			OnLeaseStartedProperties.Add(x => x.Task);
 		}
 
+		/// <inheritdoc/>
+		public Task StartAsync(CancellationToken Token) => ExpireTasksTicker.StartAsync();
+
+		/// <inheritdoc/>
+		public Task StopAsync(CancellationToken Token) => ExpireTasksTicker.StopAsync();
+
 		/// <summary>
 		/// Expire tasks that are in inactive queues (ie. no machines can execute them)
 		/// </summary>
 		/// <param name="CancellationToken"></param>
 		/// <returns></returns>
-		async Task ExpireTasksAsync(CancellationToken CancellationToken)
+		async ValueTask ExpireTasksAsync(CancellationToken CancellationToken)
 		{
 			List<IoHash> RequirementsHashes = await TaskScheduler.GetInactiveQueuesAsync();
 			foreach (IoHash RequirementsHash in RequirementsHashes)
diff --git a/Engine/Source/Programs/Horde/Horde.Build/Fleet/PoolUpdateService.cs b/Engine/Source/Programs/Horde/Horde.Build/Fleet/PoolUpdateService.cs
--- a/Engine/Source/Programs/Horde/Horde.Build/Fleet/PoolUpdateService.cs
+++ b/Engine/Source/Programs/Horde/Horde.Build/Fleet/PoolUpdateService.cs
@@ -1,8 +1,10 @@
 // Copyright Epic Games, Inc. All Rights Reserved.
 
+using HordeCommon;
 using HordeServer.Collections;
 using HordeServer.Models;
 using HordeServer.Utilities;
+using Microsoft.Extensions.Hosting;
 using Microsoft.Extensions.Logging;
 using System;
 using System.Collections.Generic;
@@ -17,45 +19,39 @@ namespace HordeServer.Services
 	/// <summary>
 	/// Periodically updates pool documents to contain the correct workspaces
 	/// </summary>
-	public class PoolUpdateService : TickedBackgroundService
+	public sealed class PoolUpdateService : IHostedService, IDisposable
 	{
-		/// <summary>
-		/// Collection of pool documents
-		/// </summary>
 		IPoolCollection Pools;
-
-		/// <summary>
-		/// Collection of stream documents
-		/// </summary>
 		IStreamCollection Streams;
-
-		/// <summary>
-		/// The logger instance for this service
-		/// </summary>
 		ILogger<PoolService> Logger;
+		ITicker Ticker;
 
 		/// <summary>
 		/// Constructor
 		/// </summary>
-		/// <param name="Pools">Collection of pool documents</param>
-		/// <param name="Streams">Collection of stream documents</param>
-		/// <param name="Logger">Logger instance</param>
-		public PoolUpdateService(IPoolCollection Pools, IStreamCollection Streams, ILogger<PoolService> Logger)
-			: base(TimeSpan.FromSeconds(30.0), Logger)
+		public PoolUpdateService(IPoolCollection Pools, IStreamCollection Streams, IClock Clock, ILogger<PoolService> Logger)
 		{
 			this.Pools = Pools;
 			this.Streams = Streams;
 			this.Logger = Logger;
+			this.Ticker = Clock.CreateTicker(TimeSpan.FromSeconds(30.0), UpdatePoolsAsync, Logger);
 		}
 
+		/// <inheritdoc/>
+		public Task StartAsync(CancellationToken CancellationToken) => Ticker.StartAsync();
+
+		/// <inheritdoc/>
+		public Task StopAsync(CancellationToken CancellationToken) => Ticker.StopAsync();
 
+		/// <inheritdoc/>
+		public void Dispose() => Ticker.Dispose();
 
 		/// <summary>
 		/// Execute the background task
 		/// </summary>
 		/// <param name="StoppingToken">Cancellation token for the async task</param>
 		/// <returns>Async task</returns>
-		protected override async Task TickAsync(CancellationToken StoppingToken)
+		async ValueTask UpdatePoolsAsync(CancellationToken StoppingToken)
 		{
 			// Capture the start time for this operation. We use this to attempt to sequence updates to agents, and prevent overriding another server's updates.
 			DateTime StartTime = DateTime.UtcNow;
diff --git a/Engine/Source/Programs/Horde/Horde.Build/Jobs/Impl/JobTaskSource.cs b/Engine/Source/Programs/Horde/Horde.Build/Jobs/Impl/JobTaskSource.cs
--- a/Engine/Source/Programs/Horde/Horde.Build/Jobs/Impl/JobTaskSource.cs
+++ b/Engine/Source/Programs/Horde/Horde.Build/Jobs/Impl/JobTaskSource.cs
@@ -42,7 +42,7 @@ namespace HordeServer.Tasks.Impl
 	/// <summary>
 	/// Background service to dispatch pending work to agents in priority order.
 	/// </summary>
-	public sealed class JobTaskSource : TaskSourceBase<ExecuteJobTask>, IDisposable
+	public sealed class JobTaskSource : TaskSourceBase<ExecuteJobTask>, IHostedService, IDisposable
 	{
 		/// <inheritdoc/>
 		public override string Type => "Job";
@@ -188,133 +188,49 @@ namespace HordeServer.Tasks.Impl
 			}
 		}
 
-		/// <summary>
-		/// The database service instance
-		/// </summary>
 		DatabaseService DatabaseService;
-
-		/// <summary>
-		/// The stream service instance
-		/// </summary>
 		StreamService StreamService;
-
-		/// <summary>
-		/// The log file service instance
-		/// </summary>
 		ILogFileService LogFileService;
-
-		/// <summary>
-		/// Collection of agent documents
-		/// </summary>
 		IAgentCollection AgentsCollection;
-
-		/// <summary>
-		/// Collection of job documents
-		/// </summary>
 		IJobCollection Jobs;
-
-		/// <summary>
-		/// Collection of jobstepref documents
-		/// </summary>
 		IJobStepRefCollection JobStepRefs;
-
-		/// <summary>
-		/// Collection of graph documents
-		/// </summary>
 		IGraphCollection Graphs;
-
-		/// <summary>
-		/// Collection of pool documents
-		/// </summary>
 		IPoolCollection PoolCollection;
-
-		/// <summary>
-		/// Collection of UGS metadata documents
-		/// </summary>
 		IUgsMetadataCollection UgsMetadataCollection;
-
-		/// <summary>
-		/// The Perforce load balancer
-		/// </summary>
 		PerforceLoadBalancer PerforceLoadBalancer;
-
-		/// <summary>
-		/// The application lifetime
-		/// </summary>
 		IHostApplicationLifetime ApplicationLifetime;
-
-		/// <summary>
-		/// Settings instance
-		/// </summary>
 		IOptionsMonitor<ServerSettings> Settings;
-
-		/// <summary>
-		/// Log writer
-		/// </summary>
 		ILogger<JobTaskSource> Logger;
+		ITicker Ticker;
 
-		/// <summary>
-		/// Object used for ensuring mutual exclusion to the queues
-		/// </summary>
+		// Object used for ensuring mutual exclusion to the queues
 		object LockObject = new object();
 
-		/// <summary>
-		/// List of items waiting to be executed
-		/// </summary>
+		// List of items waiting to be executed
 		SortedSet<QueueItem> Queue = new SortedSet<QueueItem>(new QueueItemComparer());
 
-		/// <summary>
-		/// Map from batch id to the corresponding queue item
-		/// </summary>
+		// Map from batch id to the corresponding queue item
 		Dictionary<(JobId, SubResourceId), QueueItem> BatchIdToQueueItem = new Dictionary<(JobId, SubResourceId), QueueItem>();
 
-		/// <summary>
-		/// Set of long-poll tasks waiting to be satisfied 
-		/// </summary>
+		// Set of long-poll tasks waiting to be satisfied 
 		HashSet<QueueWaiter> Waiters = new HashSet<QueueWaiter>();
 
-		/// <summary>
-		/// During a background queue refresh operation, any updated batches are added to this dictionary for merging into the updated queue.
-		/// </summary>
+		// During a background queue refresh operation, any updated batches are added to this dictionary for merging into the updated queue.
 		List<QueueItem>? NewQueueItemsDuringUpdate;
 
-		/// <summary>
-		/// Cache of pools
-		/// </summary>
+		// Cache of pools
 		Dictionary<PoolId, IPool> CachedPoolIdToInstance = new Dictionary<PoolId, IPool>();
 
-		/// <summary>
-		/// Cache of stream objects. Used to resolve agent types.
-		/// </summary>
+		// Cache of stream objects. Used to resolve agent types.
 		private Dictionary<StreamId, IStream> Streams = new Dictionary<StreamId, IStream>();
 
-		/// <summary>
-		/// Background tick registration
-		/// </summary>
-		BackgroundTick Ticker;
-
-		/// <summary>
-		/// Interval between querying the database for jobs to execute
-		/// </summary>
+		// Interval between querying the database for jobs to execute
 		static readonly TimeSpan RefreshInterval = TimeSpan.FromSeconds(5.0);
 
 		/// <summary>
 		/// Constructor
 		/// </summary>
-		/// <param name="DatabaseService">The database service instance</param>
-		/// <param name="Agents">The agents collection</param>
-		/// <param name="Jobs">The jobs collection</param>
-		/// <param name="JobStepRefs">Collection of jobstepref documents</param>
-		/// <param name="Graphs">The graphs collection</param>
-		/// <param name="Pools">The pools collection</param>
-		/// <param name="UgsMetadataCollection">Ugs metadata collection</param>
-		/// <param name="StreamService">The stream service instance</param>
-		/// <param name="LogFileService">The log file service instance</param>
-		/// <param name="PerforceLoadBalancer">Perforce load balancer</param>
-		/// <param name="ApplicationLifetime">Application lifetime interface</param>
-		/// <param name="Settings">Settings for the server</param>
-		/// <param name="Logger">Log writer</param>
-		public JobTaskSource(DatabaseService DatabaseService, IAgentCollection Agents, IJobCollection Jobs, IJobStepRefCollection JobStepRefs, IGraphCollection Graphs, IPoolCollection Pools, IUgsMetadataCollection UgsMetadataCollection, StreamService StreamService, ILogFileService LogFileService, PerforceLoadBalancer PerforceLoadBalancer, IHostApplicationLifetime ApplicationLifetime, IOptionsMonitor<ServerSettings> Settings, ILogger<JobTaskSource> Logger)
+		public JobTaskSource(DatabaseService DatabaseService, IAgentCollection Agents, IJobCollection Jobs, IJobStepRefCollection JobStepRefs, IGraphCollection Graphs, IPoolCollection Pools, IUgsMetadataCollection UgsMetadataCollection, StreamService StreamService, ILogFileService LogFileService, PerforceLoadBalancer PerforceLoadBalancer, IHostApplicationLifetime ApplicationLifetime, IClock Clock, IOptionsMonitor<ServerSettings> Settings, ILogger<JobTaskSource> Logger)
 		{
 			this.DatabaseService = DatabaseService;
 			this.AgentsCollection = Agents;
@@ -327,7 +243,7 @@ namespace HordeServer.Tasks.Impl
 			this.LogFileService = LogFileService;
 			this.PerforceLoadBalancer = PerforceLoadBalancer;
 			this.ApplicationLifetime = ApplicationLifetime;
-			this.Ticker = new BackgroundTick(TickAsync, RefreshInterval, Logger);
+			this.Ticker = Clock.CreateTicker(RefreshInterval, TickAsync, Logger);
 			this.Settings = Settings;
 			this.Logger = Logger;
 
@@ -335,16 +251,13 @@ namespace HordeServer.Tasks.Impl
 		}
 
 		/// <inheritdoc/>
-		public async ValueTask DisposeAsync()
-		{
-			await Ticker.DisposeAsync();
-		}
+		public Task StartAsync(CancellationToken cancellationToken) => Ticker.StartAsync();
 
 		/// <inheritdoc/>
-		public void Dispose()
-		{
-			Ticker.Dispose();
-		}
+		public Task StopAsync(CancellationToken cancellationToken) => Ticker.StopAsync();
+
+		/// <inheritdoc/>
+		public void Dispose() => Ticker.Dispose();
 
 		/// <summary>
 		/// Gets an object containing the stats of the queue for diagnostic purposes.
@@ -392,7 +305,7 @@ namespace HordeServer.Tasks.Impl
 		/// </summary>
 		/// <param name="StoppingToken">Token that indicates that the service should shut down</param>
 		/// <returns>Async task</returns>
-		async Task TickAsync(CancellationToken StoppingToken)
+		async ValueTask TickAsync(CancellationToken StoppingToken)
 		{
 			// Set the NewBatchIdToQueueItem member, so we capture any updated jobs during the DB query.
 			lock (LockObject)
diff --git a/Engine/Source/Programs/Horde/Horde.Build/Logs/Impl/LogFileService.cs b/Engine/Source/Programs/Horde/Horde.Build/Logs/Impl/LogFileService.cs
--- a/Engine/Source/Programs/Horde/Horde.Build/Logs/Impl/LogFileService.cs
+++ b/Engine/Source/Programs/Horde/Horde.Build/Logs/Impl/LogFileService.cs
@@ -32,6 +32,7 @@ using ILogger = Microsoft.Extensions.Logging.ILogger;
 using Stream = System.IO.Stream;
 using OpenTracing.Util;
 using OpenTracing;
+using Microsoft.Extensions.Hosting;
 
 namespace HordeServer.Services
 {
@@ -279,7 +280,7 @@ namespace HordeServer.Services
 	/// <summary>
 	/// Wraps functionality for manipulating logs
 	/// </summary>
-	public class LogFileService : TickedBackgroundService, ILogFileService
+	public sealed class LogFileService : IHostedService, ILogFileService, IDisposable
 	{
 		/// <summary>
 		/// Information Logger
@@ -494,23 +495,43 @@ namespace HordeServer.Services
 			public override void Write(byte[] buffer, int offset, int count) => throw new NotImplementedException();
 		}
 
+		ITicker Ticker;
+
 		/// <summary>
 		/// Constructor
 		/// </summary>
-		/// <param name="LogFiles">The log file collection</param>
-		/// <param name="LogEvents">The log events collection</param>
-		/// <param name="Builder">The log builder</param>
-		/// <param name="Storage">THe log storage hierarchy</param>
-		/// <param name="Logger">Log interface</param>
-		public LogFileService(ILogFileCollection LogFiles, ILogEventCollection LogEvents, ILogBuilder Builder, ILogStorage Storage, ILogger<LogFileService> Logger)
-			: base(TimeSpan.FromSeconds(30.0), Logger)
+		public LogFileService(ILogFileCollection LogFiles, ILogEventCollection LogEvents, ILogBuilder Builder, ILogStorage Storage, IClock Clock, ILogger<LogFileService> Logger)
 		{
 			this.LogFiles = LogFiles;
 			this.LogEvents = LogEvents;
-			this.Logger = Logger;
 			this.LogFileCache = new MemoryCache(new MemoryCacheOptions());
 			this.Builder = Builder;
 			this.Storage = Storage;
+			this.Ticker = Clock.CreateTicker(TimeSpan.FromSeconds(30.0), TickAsync, Logger);
+			this.Logger = Logger;
+		}
+
+		/// <inheritdoc/>
+		public Task StartAsync(CancellationToken CancellationToken) => Ticker.StartAsync();
+
+		/// <inheritdoc/>
+		public async Task StopAsync(CancellationToken CancellationToken)
+		{
+			Logger.LogInformation("Stopping log file service");
+			if (Builder.FlushOnShutdown)
+			{
+				await FlushAsync();
+			}
+			await Ticker.StopAsync();
+			Logger.LogInformation("Log service stopped");
+		}
+
+		/// <inheritdoc/>
+		public void Dispose()
+		{
+			LogFileCache.Dispose();
+			Storage.Dispose();
+			Ticker.Dispose();
 		}
 
 		/// <inheritdoc/>
@@ -990,7 +1011,7 @@ namespace HordeServer.Services
 		/// Executes a background task
 		/// </summary>
 		/// <param name="StoppingToken">Cancellation token</param>
-		protected override async Task TickAsync(CancellationToken StoppingToken)
+		async ValueTask TickAsync(CancellationToken StoppingToken)
 		{
 			lock (WriteLock)
 			{
@@ -1005,15 +1026,7 @@ namespace HordeServer.Services
 			}
 			await IncrementalFlush();
 		}
-		
-		/// <summary>
-		/// Executes a background task. Publicly exposed for tests.
-		/// </summary>
-		public async Task TickOnlyForTestingAsync()
-		{
-			await TickAsync(new CancellationToken());
-		}
-		
+				
 		/// <summary>
 		/// Flushes complete chunks to the storage provider
 		/// </summary>
@@ -1034,21 +1047,6 @@ namespace HordeServer.Services
 			WriteCompleteChunks(FlushChunks, true);
 		}
 
-		/// <summary>
-		/// Called when the service is shutting down
-		/// </summary>
-		/// <param name="CancellationToken"></param>
-		/// <returns></returns>
-		[SuppressMessage("Design", "CA1031:Do not catch general exception types")]
-		public override async Task StopAsync(CancellationToken CancellationToken)
-		{
-			Logger.LogInformation("Stopping log file service");
-			if (Builder.FlushOnShutdown)
-			{
-				await FlushAsync();
-			}
-			Logger.LogInformation("Log service stopped");
-		}
 
 		/// <summary>
 		/// Flushes the write cache
diff --git a/Engine/Source/Programs/Horde/Horde.Build/Server/DowntimeService.cs b/Engine/Source/Programs/Horde/Horde.Build/Server/DowntimeService.cs
--- a/Engine/Source/Programs/Horde/Horde.Build/Server/DowntimeService.cs
+++ b/Engine/Source/Programs/Horde/Horde.Build/Server/DowntimeService.cs
@@ -1,8 +1,9 @@
 // Copyright Epic Games, Inc. All Rights Reserved.
 
-using Amazon.Runtime.Internal.Util;
+using HordeCommon;
 using HordeServer.Models;
 using HordeServer.Utilities;
+using Microsoft.Extensions.Hosting;
 using Microsoft.Extensions.Logging;
 using Microsoft.Extensions.Options;
 using System;
@@ -31,7 +32,7 @@ namespace HordeServer.Services
 	/// <summary>
 	/// Service which manages the downtime schedule
 	/// </summary>
-	public class DowntimeService : TickedBackgroundService, IDowntimeService
+	public sealed class DowntimeService : IDowntimeService, IHostedService, IDisposable
 	{
 		/// <summary>
 		/// Whether the server is currently in downtime
@@ -43,32 +44,44 @@ namespace HordeServer.Services
 		}
 
 		DatabaseService DatabaseService;
+		ITicker Ticker;
 		IOptionsMonitor<ServerSettings> Settings;
-		ILogger<DowntimeService> Logger;
+		ILogger Logger;
 
 		/// <summary>
 		/// Constructor
 		/// </summary>
 		/// <param name="DatabaseService">The database service instance</param>
+		/// <param name="Clock"></param>
 		/// <param name="Settings">The server settings</param>
 		/// <param name="Logger">Logger instance</param>
-		public DowntimeService(DatabaseService DatabaseService, IOptionsMonitor<ServerSettings> Settings, ILogger<DowntimeService> Logger)
-			: base(TimeSpan.FromMinutes(1.0), Logger)
+		public DowntimeService(DatabaseService DatabaseService, IClock Clock, IOptionsMonitor<ServerSettings> Settings, ILogger<DowntimeService> Logger)
 		{
 			this.DatabaseService = DatabaseService;
 			this.Settings = Settings;
 			this.Logger = Logger;
 
 			// Ensure the initial value to be correct
-			TickAsync(CancellationToken.None).Wait();
+			TickAsync(CancellationToken.None).AsTask().Wait();
+
+			Ticker = Clock.CreateTicker(TimeSpan.FromMinutes(1.0), TickAsync, Logger);
 		}
 
+		/// <inheritdoc/>
+		public Task StartAsync(CancellationToken CancellationToken) => Ticker.StartAsync();
+
+		/// <inheritdoc/>
+		public Task StopAsync(CancellationToken CancellationToken) => Ticker.StopAsync();
+
+		/// <inheritdoc/>
+		public void Dispose() => Ticker.Dispose();
+
 		/// <summary>
 		/// Periodically called tick function
 		/// </summary>
 		/// <param name="StoppingToken">Token indicating that the service should stop</param>
 		/// <returns>Async task</returns>
-		protected override async Task TickAsync(CancellationToken StoppingToken)
+		async ValueTask TickAsync(CancellationToken StoppingToken)
 		{
 			Globals Globals = await DatabaseService.GetGlobalsAsync();
 
diff --git a/Engine/Source/Programs/Horde/Horde.Build/Server/MetricService.cs b/Engine/Source/Programs/Horde/Horde.Build/Server/MetricService.cs
--- a/Engine/Source/Programs/Horde/Horde.Build/Server/MetricService.cs
+++ b/Engine/Source/Programs/Horde/Horde.Build/Server/MetricService.cs
@@ -6,39 +6,47 @@ using System;
 using System.Threading;
 using System.Threading.Tasks;
 using StatsdClient;
+using Microsoft.Extensions.Hosting;
+using HordeCommon;
 
 namespace HordeServer.Services
 {
 	/// <summary>
 	/// Periodically send metrics for the CLR and other services that cannot be collected on a per-request basis
 	/// </summary>
-	public class MetricService : TickedBackgroundService
+	public sealed class MetricService : IHostedService, IDisposable
 	{
-		/// <summary>
-		/// DogStatsd instance
-		/// </summary>
-		private readonly IDogStatsd DogStatsd;
+		IDogStatsd DogStatsd;
+		ITicker Ticker;
 
 		/// <summary>
 		/// Constructor
 		/// </summary>
-		/// <param name="DogStatsd"></param>
-		/// <param name="Logger">Logger instance</param>
-		public MetricService(IDogStatsd DogStatsd, ILogger<MetricService> Logger) : base(TimeSpan.FromSeconds(20.0), Logger)
+		public MetricService(IDogStatsd DogStatsd, IClock Clock, ILogger<MetricService> Logger)
 		{
 			this.DogStatsd = DogStatsd;
+			this.Ticker = Clock.CreateTicker(TimeSpan.FromSeconds(20.0), TickAsync, Logger);
 		}
 
+		/// <inheritdoc/>
+		public Task StartAsync(CancellationToken cancellationToken) => Ticker.StartAsync();
+
+		/// <inheritdoc/>
+		public Task StopAsync(CancellationToken cancellationToken) => Ticker.StopAsync();
+
+		/// <inheritdoc/>
+		public void Dispose() => Ticker.Dispose();
+
 		/// <summary>
 		/// Execute the background task
 		/// </summary>
 		/// <param name="StoppingToken">Cancellation token for the async task</param>
 		/// <returns>Async task</returns>
-		protected override Task TickAsync(CancellationToken StoppingToken)
+		ValueTask TickAsync(CancellationToken StoppingToken)
 		{
 			ReportGcMetrics();
 			ReportThreadMetrics();
-			return Task.CompletedTask;
+			return new ValueTask();
 		}
 
 		private void ReportThreadMetrics()
diff --git a/Engine/Source/Programs/Horde/Horde.Build/Startup.cs b/Engine/Source/Programs/Horde/Horde.Build/Startup.cs
--- a/Engine/Source/Programs/Horde/Horde.Build/Startup.cs
+++ b/Engine/Source/Programs/Horde/Horde.Build/Startup.cs
@@ -569,7 +569,9 @@ namespace HordeServer
 
 			// Task sources. Order of registration is important here; it dictates the priority in which sources are served.
 			Services.AddSingleton<JobTaskSource>();
+			Services.AddHostedService<JobTaskSource>(Provider => Provider.GetRequiredService<JobTaskSource>());
 			Services.AddSingleton<ConformTaskSource>();
+			Services.AddHostedService<ConformTaskSource>(Provider => Provider.GetRequiredService<ConformTaskSource>());
 			Services.AddSingleton<IComputeService, ComputeService>();
 
 			Services.AddSingleton<ITaskSource, UpgradeTaskSource>();
diff --git a/Engine/Source/Programs/Horde/Horde.Build/Utilities/BackgroundTick.cs b/Engine/Source/Programs/Horde/Horde.Build/Utilities/BackgroundTick.cs
deleted file mode 100644
diff --git a/Engine/Source/Programs/Horde/Horde.Build/Utilities/Clock.cs b/Engine/Source/Programs/Horde/Horde.Build/Utilities/Clock.cs
--- a/Engine/Source/Programs/Horde/Horde.Build/Utilities/Clock.cs
+++ b/Engine/Source/Programs/Horde/Horde.Build/Utilities/Clock.cs
@@ -1,12 +1,33 @@
 // Copyright Epic Games, Inc. All Rights Reserved.
 
+using EpicGames.Core;
+using Microsoft.Extensions.Logging;
 using System;
+using System.Collections.Generic;
+using System.Diagnostics;
+using System.Threading;
+using System.Threading.Tasks;
 
 namespace HordeCommon
 {
 	/// <summary>
-	/// Interface representing time to make it pluggable during testing
-	/// In normal use, the Clock implementation below is used. 
+	/// Base interface for a scheduled event
+	/// </summary>
+	public interface ITicker : IDisposable
+	{
+		/// <summary>
+		/// Start the ticker
+		/// </summary>
+		Task StartAsync();
+
+		/// <summary>
+		/// Stop the ticker
+		/// </summary>
+		Task StopAsync();
+	}
+
+	/// <summary>
+	/// Interface representing time and scheduling events which is pluggable during testing. In normal use, the Clock implementation below is used. 
 	/// </summary>
 	public interface IClock
 	{
@@ -15,15 +36,135 @@ namespace HordeCommon
 		/// </summary>
 		/// <returns></returns>
 		DateTime UtcNow { get; }
+
+		/// <summary>
+		/// Create an event that will trigger after the given time
+		/// </summary>
+		/// <param name="Interval">Time after which the event will trigger</param>
+		/// <param name="TickAsync">Callback for the tick. Returns the time interval until the next tick, or null to cancel the tick.</param>
+		/// <param name="Logger">Logger for error messages</param>
+		/// <returns>Handle to the event</returns>
+		ITicker CreateTicker(TimeSpan Interval, Func<CancellationToken, ValueTask<TimeSpan?>> TickAsync, ILogger Logger);
 	}
-	
+
+	/// <summary>
+	/// Extension methods for <see cref="IClock"/>
+	/// </summary>
+	public static class ClockExtensions
+	{
+		/// <summary>
+		/// Create an event that will trigger after the given time
+		/// </summary>
+		/// <param name="Clock">Clock to schedule the event on</param>
+		/// <param name="Interval">Interval for the callback</param>
+		/// <param name="TickAsync">Trigger callback</param>
+		/// <param name="Logger">Logger for any error messages</param>
+		/// <returns>Handle to the event</returns>
+		public static ITicker CreateTicker(this IClock Clock, TimeSpan Interval, Func<CancellationToken, ValueTask> TickAsync, ILogger Logger)
+		{
+			Func<CancellationToken, ValueTask<TimeSpan?>> WrappedTrigger = async Token =>
+			{
+				Stopwatch Timer = Stopwatch.StartNew();
+				await TickAsync(Token);
+				return Interval - Timer.Elapsed;
+			};
+			return Clock.CreateTicker(Interval, WrappedTrigger, Logger);
+		}
+	}
+
 	/// <summary>
 	/// Implementation of <see cref="IClock"/> which returns the current time
 	/// </summary>
 	public class Clock : IClock
 	{
+		class TickerImpl : ITicker
+		{
+			[System.Diagnostics.CodeAnalysis.SuppressMessage("Usage", "CA2213:Disposable fields should be disposed", Justification = "<Pending>")]
+			CancellationTokenSource? CancellationSource;
+			Func<Task> TickFunc;
+			Task? BackgroundTask;
+
+			public TickerImpl(TimeSpan Delay, Func<CancellationToken, ValueTask<TimeSpan?>> TriggerAsync, ILogger Logger)
+			{
+				TickFunc = () => Run(Delay, TriggerAsync, Logger);
+			}
+
+			public async Task StartAsync()
+			{
+				await StopAsync();
+				CancellationSource = new CancellationTokenSource();
+				BackgroundTask = Task.Run(TickFunc);
+			}
+
+			public async Task StopAsync()
+			{
+				if (CancellationSource != null)
+				{
+					CancellationSource.Cancel();
+					await BackgroundTask!;
+					CancellationSource.Dispose();
+					CancellationSource = null;
+				}
+			}
+
+			public void Dispose()
+			{
+				StopAsync().Wait();
+				if (CancellationSource != null)
+				{
+					CancellationSource.Dispose();
+				}
+			}
+
+			public async Task Run(TimeSpan Delay, Func<CancellationToken, ValueTask<TimeSpan?>> TriggerAsync, ILogger Logger)
+			{
+				while (!CancellationSource!.IsCancellationRequested)
+				{
+					try
+					{
+						if (Delay > TimeSpan.Zero)
+						{
+							await Task.Delay(Delay, CancellationSource.Token);
+						}
+
+						TimeSpan? NextDelay = await TriggerAsync(CancellationSource.Token);
+						if(NextDelay == null)
+						{
+							break;
+						}
+
+						Delay = NextDelay.Value;
+					}
+					catch (OperationCanceledException) when (CancellationSource.IsCancellationRequested)
+					{
+					}
+					catch (Exception Ex)
+					{
+						Logger.LogError(Ex, "Exception while executing scheduled event");
+						if (Delay < TimeSpan.Zero)
+						{
+							Delay = TimeSpan.FromSeconds(5.0);
+							Logger.LogWarning("Delaying tick for 5 seconds");
+						}
+					}
+				}
+			}
+		}
+
 		/// <inheritdoc/>
 		public DateTime UtcNow => DateTime.UtcNow;
+
+		/// <summary>
+		/// Create an event that will trigger after the given time for one agent in the cluster
+		/// </summary>
+		/// <param name="Delay">Time after which the event will trigger</param>
+		/// <param name="TriggerAsync">Callback for the event triggering</param>
+		/// <param name="Logger">Logger for error messages</param>
+		/// <returns>Handle to the event</returns>
+		public ITicker CreateTicker(TimeSpan Delay, Func<CancellationToken, ValueTask<TimeSpan?>> TriggerAsync, ILogger Logger)
+		{
+			return new TickerImpl(Delay, TriggerAsync, Logger);
+		}
 	}
 	
 	/// <summary>
@@ -32,7 +173,54 @@ namespace HordeCommon
 	/// </summary>
 	public class FakeClock : IClock
 	{
+		class TickerImpl : ITicker
+		{
+			FakeClock Outer { get; }
+			TimeSpan Interval;
+			public DateTime? NextTime { get; set; }
+			public Func<CancellationToken, ValueTask<TimeSpan?>> TickAsync { get; }
+
+			public TickerImpl(FakeClock Outer, TimeSpan Interval, Func<CancellationToken, ValueTask<TimeSpan?>> TickAsync)
+			{
+				this.Outer = Outer;
+				this.Interval = Interval;
+				this.TickAsync = TickAsync;
+
+				lock (Outer.Triggers)
+				{
+					Outer.Triggers.Add(this);
+				}
+			}
+
+			public Task StartAsync()
+			{
+				NextTime = Outer.UtcNow + Interval;
+				return Task.CompletedTask;
+			}
+
+			public Task StopAsync()
+			{
+				NextTime = null;
+				return Task.CompletedTask;
+			}
+
+			public void Dispose()
+			{
+				DisposeAsync().AsTask().Wait();
+			}
+
+			public ValueTask DisposeAsync()
+			{
+				lock (Outer.Triggers)
+				{
+					Outer.Triggers.Remove(this);
+				}
+				return new ValueTask();
+			}
+		}
+
 		DateTime UtcNowPrivate;
+		List<TickerImpl> Triggers = new List<TickerImpl>();
 
 		/// <summary>
 		/// Constructor
@@ -47,9 +235,24 @@ namespace HordeCommon
 		/// Useful for letting time progress during tests
 		/// </summary>
 		/// <param name="Period">Time span to advance</param>
-		public void Advance(TimeSpan Period)
+		public async Task AdvanceAsync(TimeSpan Period)
 		{
 			UtcNowPrivate = UtcNowPrivate.Add(Period);
+
+			for (int Idx = 0; Idx < Triggers.Count; Idx++)
+			{
+				TickerImpl Trigger = Triggers[Idx];
+				while (Trigger.NextTime != null && UtcNowPrivate > Trigger.NextTime)
+				{
+					TimeSpan? Delay = await Trigger.TickAsync(CancellationToken.None);
+					if (Delay == null)
+					{
+						Triggers.RemoveAt(Idx--);
+						break;
+					}
+					Trigger.NextTime = UtcNowPrivate + Delay.Value;
+				}
+			}
 		}
 
 		/// <inheritdoc/>
@@ -58,5 +261,11 @@ namespace HordeCommon
 			get => UtcNowPrivate;
 			set => UtcNowPrivate = value.ToUniversalTime(); 
 		}
+
+		/// <inheritdoc/>
+		public ITicker CreateTicker(TimeSpan Interval, Func<CancellationToken, ValueTask<TimeSpan?>> TriggerAsync, ILogger Logger)
+		{
+			return new TickerImpl(this, Interval, TriggerAsync);
+		}
 	}
 }
\ No newline at end of file
