Horde: Use Redis for server elections, and use deterministic tick callbacks to select them.

diff --git a/Engine/Source/Programs/Horde/Horde.Build/Agents/Impl/AgentService.cs b/Engine/Source/Programs/Horde/Horde.Build/Agents/Impl/AgentService.cs
--- a/Engine/Source/Programs/Horde/Horde.Build/Agents/Impl/AgentService.cs
+++ b/Engine/Source/Programs/Horde/Horde.Build/Agents/Impl/AgentService.cs
@@ -89,7 +89,7 @@ namespace HordeServer.Services
 			this.TaskSources = TaskSources.ToArray();
 			this.ApplicationLifetime = ApplicationLifetime;
 			this.Clock = Clock;
-			this.Ticker = Clock.CreateTicker(TimeSpan.FromSeconds(30.0), TickAsync, Logger);
+			this.Ticker = Clock.AddTicker(TimeSpan.FromSeconds(30.0), TickAsync, Logger);
 			this.Logger = Logger;
 
 			Subscription = Agents.SubscribeToUpdateEventsAsync(OnAgentUpdate).Result;
diff --git a/Engine/Source/Programs/Horde/Horde.Build/Agents/Impl/Tasks/ConformTaskSource.cs b/Engine/Source/Programs/Horde/Horde.Build/Agents/Impl/Tasks/ConformTaskSource.cs
--- a/Engine/Source/Programs/Horde/Horde.Build/Agents/Impl/Tasks/ConformTaskSource.cs
+++ b/Engine/Source/Programs/Horde/Horde.Build/Agents/Impl/Tasks/ConformTaskSource.cs
@@ -40,18 +40,12 @@ namespace HordeServer.Tasks.Impl
 		PerforceLoadBalancer PerforceLoadBalancer;
 		ILogFileService LogService;
 		ILogger Logger;
-		ElectedTick TickConformList;
+		ITicker TickConformList;
 
 		/// <summary>
 		/// Constructor
 		/// </summary>
-		/// <param name="DatabaseService"></param>
-		/// <param name="AgentCollection"></param>
-		/// <param name="PoolService"></param>
-		/// <param name="LogService"></param>
-		/// <param name="PerforceLoadBalancer"></param>
-		/// <param name="Logger"></param>
-		public ConformTaskSource(DatabaseService DatabaseService, IAgentCollection AgentCollection, PoolService PoolService, ILogFileService LogService, PerforceLoadBalancer PerforceLoadBalancer, ILogger<ConformTaskSource> Logger)
+		public ConformTaskSource(DatabaseService DatabaseService, IAgentCollection AgentCollection, PoolService PoolService, ILogFileService LogService, PerforceLoadBalancer PerforceLoadBalancer, IClock Clock, ILogger<ConformTaskSource> Logger)
 		{
 			this.DatabaseService = DatabaseService;
 			this.AgentCollection = AgentCollection;
@@ -60,36 +54,26 @@ namespace HordeServer.Tasks.Impl
 			this.PerforceLoadBalancer = PerforceLoadBalancer;
 			this.LogService = LogService;
 			this.Logger = Logger;
-			this.TickConformList = new ElectedTick(DatabaseService, new ObjectId("60afc5cf555a9a76aff0a50c"), CleanConformListAsync, TimeSpan.FromMinutes(1.0), Logger);
+			this.TickConformList = Clock.AddSharedTicker<ConformTaskSource>(TimeSpan.FromMinutes(1.0), CleanConformListAsync, Logger);
 
 			this.OnLeaseStartedProperties.Add(nameof(ConformTask.LogId), x => new LogId(x.LogId));
 		}
 
 		/// <inheritdoc/>
-		public void Dispose()
-		{
-			TickConformList.Dispose();
-		}
+		public Task StartAsync(CancellationToken CancellationToken) => TickConformList.StartAsync();
 
 		/// <inheritdoc/>
-		public Task StartAsync(CancellationToken CancellationToken)
-		{
-			TickConformList.Start();
-			return Task.CompletedTask;
-		}
+		public Task StopAsync(CancellationToken CancellationToken) => TickConformList.StopAsync();
 
 		/// <inheritdoc/>
-		public async Task StopAsync(CancellationToken CancellationToken)
-		{
-			await TickConformList.StopAsync();
-		}
+		public void Dispose() => TickConformList.Dispose();
 
 		/// <summary>
 		/// Clean up the conform list of any outdated entries
 		/// </summary>
 		/// <param name="CancellationToken">Cancellation token</param>
 		/// <returns>Async task</returns>
-		async Task CleanConformListAsync(CancellationToken CancellationToken)
+		async ValueTask CleanConformListAsync(CancellationToken CancellationToken)
 		{
 			DateTime UtcNow = DateTime.UtcNow;
 			DateTime LastCheckTimeUtc = UtcNow - TimeSpan.FromMinutes(30.0);
diff --git a/Engine/Source/Programs/Horde/Horde.Build/Compute/Impl/ComputeService.cs b/Engine/Source/Programs/Horde/Horde.Build/Compute/Impl/ComputeService.cs
--- a/Engine/Source/Programs/Horde/Horde.Build/Compute/Impl/ComputeService.cs
+++ b/Engine/Source/Programs/Horde/Horde.Build/Compute/Impl/ComputeService.cs
@@ -151,7 +151,7 @@ namespace HordeServer.Compute.Impl
 			this.ObjectCollection = ObjectCollection;
 			this.TaskScheduler = new RedisTaskScheduler<IoHash, ComputeTaskInfo>(Redis, "compute/tasks/", Logger);
 			this.MessageQueue = new RedisMessageQueue<ComputeTaskStatus>(Redis, "compute/messages/");
-			this.ExpireTasksTicker = Clock.CreateTicker(TimeSpan.FromMinutes(2.0), ExpireTasksAsync, Logger);
+			this.ExpireTasksTicker = Clock.AddTicker(TimeSpan.FromMinutes(2.0), ExpireTasksAsync, Logger);
 			this.RequirementsCache = new MemoryCache(new MemoryCacheOptions());
 			this.CachedPools = new LazyCachedValue<Task<List<IPool>>>(() => PoolCollection.GetAsync(), TimeSpan.FromSeconds(30.0));
 			this.Logger = Logger;
diff --git a/Engine/Source/Programs/Horde/Horde.Build/Devices/Impl/DeviceService.cs b/Engine/Source/Programs/Horde/Horde.Build/Devices/Impl/DeviceService.cs
--- a/Engine/Source/Programs/Horde/Horde.Build/Devices/Impl/DeviceService.cs
+++ b/Engine/Source/Programs/Horde/Horde.Build/Devices/Impl/DeviceService.cs
@@ -17,6 +17,7 @@ using System.Diagnostics.CodeAnalysis;
 using HordeServer.Notifications;
 using System.Linq;
 using HordeServer.Api;
+using Microsoft.Extensions.Hosting;
 
 namespace HordeServer.Services
 {
@@ -88,7 +89,7 @@ namespace HordeServer.Services
 	/// <summary>
 	/// Device management service
 	/// </summary>
-	public class DeviceService : TickedBackgroundService
+	public sealed class DeviceService : IHostedService, IDisposable
 	{
 		/// <summary>
 		/// The ACL service instance
@@ -100,7 +101,6 @@ namespace HordeServer.Services
 		/// </summary>
 		INotificationService NotificationService;
 
-
 		/// <summary>
 		/// Singleton instance of the job service
 		/// </summary>
@@ -126,6 +126,8 @@ namespace HordeServer.Services
 		/// </summary>
 		IDeviceCollection Devices;
 
+		ITicker Ticker;
+
 		/// <summary>
 		/// Platform map V1 singleton
 		/// </summary>
@@ -134,8 +136,7 @@ namespace HordeServer.Services
 		/// <summary>
 		/// Device service constructor
 		/// </summary>
-		public DeviceService(IDeviceCollection Devices, ISingletonDocument<DevicePlatformMapV1> PlatformMapSingleton, JobService JobService, ProjectService ProjectService, StreamService StreamService, AclService AclService, INotificationService NotificationService, ILogger<DeviceService> Logger)
-			: base(TimeSpan.FromMinutes(1.0), Logger)
+		public DeviceService(IDeviceCollection Devices, ISingletonDocument<DevicePlatformMapV1> PlatformMapSingleton, JobService JobService, ProjectService ProjectService, StreamService StreamService, AclService AclService, INotificationService NotificationService, IClock Clock, ILogger<DeviceService> Logger)
 		{
 			this.Devices = Devices;
             this.JobService = JobService;
@@ -143,16 +144,26 @@ namespace HordeServer.Services
             this.StreamService = StreamService;
             this.AclService = AclService;
             this.NotificationService = NotificationService;
+			this.Ticker = Clock.AddSharedTicker<DeviceService>(TimeSpan.FromMinutes(1.0), TickAsync, Logger);
             this.Logger = Logger;
 			
 			this.PlatformMapSingleton = PlatformMapSingleton;
 
 		}
 
+		/// <inheritdoc/>
+		public Task StartAsync(CancellationToken cancellationToken) => Ticker.StartAsync();
+
+		/// <inheritdoc/>
+		public Task StopAsync(CancellationToken cancellationToken) => Ticker.StopAsync();
+
+		/// <inheritdoc/>
+		public void Dispose() => Ticker.Dispose();
+
 		/// <summary>
 		/// Ticks service
 		/// </summary>
-		protected override async Task TickAsync(CancellationToken StoppingToken)
+		async ValueTask TickAsync(CancellationToken StoppingToken)
 		{
 
 			if (!StoppingToken.IsCancellationRequested)
diff --git a/Engine/Source/Programs/Horde/Horde.Build/Fleet/AutoscaleService.cs b/Engine/Source/Programs/Horde/Horde.Build/Fleet/AutoscaleService.cs
--- a/Engine/Source/Programs/Horde/Horde.Build/Fleet/AutoscaleService.cs
+++ b/Engine/Source/Programs/Horde/Horde.Build/Fleet/AutoscaleService.cs
@@ -25,6 +25,7 @@ using Microsoft.Extensions.Options;
 using OpenTracing;
 using OpenTracing.Util;
 using StatsdClient;
+using Microsoft.Extensions.Hosting;
 
 namespace HordeServer.Services
 {
@@ -33,7 +34,7 @@ namespace HordeServer.Services
 	/// <summary>
 	/// Service for managing the autoscaling of agent pools
 	/// </summary>
-	public class AutoscaleService : ElectedBackgroundService
+	public sealed class AutoscaleService : IHostedService, IDisposable
 	{
 		const int NumSamples = 10;
 		const int NumSamplesForResult = 9;
@@ -97,27 +98,38 @@ namespace HordeServer.Services
 		ILeaseCollection LeaseCollection;
 		IFleetManager FleetManager;
 		IDogStatsd DogStatsd;
+		IClock Clock;
+		ITicker Ticker;
 		ILogger<AutoscaleService> Logger;
 
 		/// <summary>
 		/// Constructor
 		/// </summary>
-		public AutoscaleService(DatabaseService DatabaseService, IAgentCollection AgentCollection, IPoolCollection PoolCollection, ILeaseCollection LeaseCollection, IFleetManager FleetManager, IDogStatsd DogStatsd, ILogger<AutoscaleService> Logger)
-			: base(DatabaseService, new ObjectId("5fdbadf968de915b2cac0d15"), Logger)
+		public AutoscaleService(IAgentCollection AgentCollection, IPoolCollection PoolCollection, ILeaseCollection LeaseCollection, IFleetManager FleetManager, IDogStatsd DogStatsd, IClock Clock, ILogger<AutoscaleService> Logger)
 		{
 			this.AgentCollection = AgentCollection;
 			this.PoolCollection = PoolCollection;
 			this.LeaseCollection = LeaseCollection;
 			this.FleetManager = FleetManager;
 			this.DogStatsd = DogStatsd;
+			this.Clock = Clock;
+			this.Ticker = Clock.AddSharedTicker<AutoscaleService>(TimeSpan.FromMinutes(5.0), TickLeaderAsync, Logger);
 			this.Logger = Logger;
 		}
 
 		/// <inheritdoc/>
-		protected override async Task<DateTime> TickLeaderAsync(CancellationToken StoppingToken)
+		public Task StartAsync(CancellationToken CancellationToken) => Ticker.StartAsync();
+
+		/// <inheritdoc/>
+		public Task StopAsync(CancellationToken CancellationToken) => Ticker.StopAsync();
+
+		/// <inheritdoc/>
+		public void Dispose() => Ticker.Dispose();
+
+		/// <inheritdoc/>
+		async ValueTask TickLeaderAsync(CancellationToken StoppingToken)
 		{
-			DateTime UtcNow = DateTime.UtcNow;
-			DateTime NextTickTime = UtcNow + TimeSpan.FromMinutes(5.0);
+			DateTime UtcNow = Clock.UtcNow;
 
 			Logger.LogInformation("Autoscaling pools...");
 			Stopwatch Stopwatch = Stopwatch.StartNew();
@@ -129,7 +141,7 @@ namespace HordeServer.Services
 				List<IAgent> Agents = await AgentCollection.FindAsync(Status: AgentStatus.Ok);
 
 				// Query leases in last interval
-				DateTime MaxTime = DateTime.UtcNow;
+				DateTime MaxTime = Clock.UtcNow;
 				DateTime MinTime = MaxTime - (SampleTime * NumSamples);
 				List<ILease> Leases = await LeaseCollection.FindLeasesAsync(MinTime, MaxTime);
 
@@ -253,8 +265,6 @@ namespace HordeServer.Services
 			
 			Stopwatch.Stop();
 			Logger.LogInformation("Autoscaling pools took {ElapsedTime} ms", Stopwatch.ElapsedMilliseconds);
-			
-			return NextTickTime;
 		}
 
 		/// <summary>
diff --git a/Engine/Source/Programs/Horde/Horde.Build/Fleet/PoolUpdateService.cs b/Engine/Source/Programs/Horde/Horde.Build/Fleet/PoolUpdateService.cs
--- a/Engine/Source/Programs/Horde/Horde.Build/Fleet/PoolUpdateService.cs
+++ b/Engine/Source/Programs/Horde/Horde.Build/Fleet/PoolUpdateService.cs
@@ -34,7 +34,7 @@ namespace HordeServer.Services
 			this.Pools = Pools;
 			this.Streams = Streams;
 			this.Logger = Logger;
-			this.Ticker = Clock.CreateTicker(TimeSpan.FromSeconds(30.0), UpdatePoolsAsync, Logger);
+			this.Ticker = Clock.AddTicker(TimeSpan.FromSeconds(30.0), UpdatePoolsAsync, Logger);
 		}
 
 		/// <inheritdoc/>
diff --git a/Engine/Source/Programs/Horde/Horde.Build/Fleet/TelemetryService.cs b/Engine/Source/Programs/Horde/Horde.Build/Fleet/TelemetryService.cs
--- a/Engine/Source/Programs/Horde/Horde.Build/Fleet/TelemetryService.cs
+++ b/Engine/Source/Programs/Horde/Horde.Build/Fleet/TelemetryService.cs
@@ -1,10 +1,12 @@
 // Copyright Epic Games, Inc. All Rights Reserved.
 
 using Google.Protobuf.WellKnownTypes;
+using HordeCommon;
 using HordeCommon.Rpc.Tasks;
 using HordeServer.Collections;
 using HordeServer.Models;
 using HordeServer.Utilities;
+using Microsoft.Extensions.Hosting;
 using Microsoft.Extensions.Logging;
 using MongoDB.Bson;
 using System;
@@ -20,57 +22,43 @@ namespace HordeServer.Services
 	/// <summary>
 	/// Service which updates telemetry periodically
 	/// </summary>
-	public class TelemetryService : ElectedBackgroundService
+	public sealed class TelemetryService : IHostedService, IDisposable
 	{
-		/// <summary>
-		/// Collection of telemetry documents
-		/// </summary>
 		ITelemetryCollection TelemetryCollection;
-
-		/// <summary>
-		/// Collection of agent documents
-		/// </summary>
 		IAgentCollection AgentCollection;
-
-		/// <summary>
-		/// Collection of lease documents
-		/// </summary>
 		ILeaseCollection LeaseCollection;
-
-		/// <summary>
-		/// Collection of pool documents
-		/// </summary>
 		IPoolCollection PoolCollection;
-
-		/// <summary>
-		/// Logger for diagnostic messages
-		/// </summary>
+		IClock Clock;
+		ITicker Tick;
 		ILogger Logger;
 
 		/// <summary>
 		/// Constructor
 		/// </summary>
-		/// <param name="DatabaseService"></param>
-		/// <param name="TelemetryCollection"></param>
-		/// <param name="AgentCollection"></param>
-		/// <param name="LeaseCollection"></param>
-		/// <param name="PoolCollection"></param>
-		/// <param name="Logger"></param>
-		public TelemetryService(DatabaseService DatabaseService, ITelemetryCollection TelemetryCollection, IAgentCollection AgentCollection, ILeaseCollection LeaseCollection, IPoolCollection PoolCollection, ILogger<TelemetryService> Logger)
-			: base(DatabaseService, new ObjectId("600082cb9b822d7afee6f0d1"), Logger)
+		public TelemetryService(ITelemetryCollection TelemetryCollection, IAgentCollection AgentCollection, ILeaseCollection LeaseCollection, IPoolCollection PoolCollection, IClock Clock, ILogger<TelemetryService> Logger)
 		{
 			this.TelemetryCollection = TelemetryCollection;
 			this.AgentCollection = AgentCollection;
 			this.LeaseCollection = LeaseCollection;
 			this.PoolCollection = PoolCollection;
+			this.Clock = Clock;
+			this.Tick = Clock.AddSharedTicker<TelemetryService>(TimeSpan.FromMinutes(10.0), TickLeaderAsync, Logger);
 			this.Logger = Logger;
 		}
 
 		/// <inheritdoc/>
-		protected override async Task<DateTime> TickLeaderAsync(CancellationToken StoppingToken)
+		public Task StartAsync(CancellationToken CancellationToken) => Tick.StartAsync();
+
+		/// <inheritdoc/>
+		public Task StopAsync(CancellationToken CancellationToken) => Tick.StopAsync();
+
+		/// <inheritdoc/>
+		public void Dispose() => Tick.Dispose();
+
+		/// <inheritdoc/>
+		async ValueTask TickLeaderAsync(CancellationToken StoppingToken)
 		{
-			DateTime CurrentTime = DateTime.UtcNow;
-			DateTime NextTickTime = CurrentTime + TimeSpan.FromMinutes(10.0);
+			DateTime CurrentTime = Clock.UtcNow;
 
 			// Find the last time that we need telemetry for
 			DateTime MaxTime = CurrentTime.Date + TimeSpan.FromHours(CurrentTime.Hour);
@@ -138,8 +126,6 @@ namespace HordeServer.Services
 
 				BucketMinTime = BucketMaxTime;
 			}
-
-			return NextTickTime;
 		}
 	}
 }
diff --git a/Engine/Source/Programs/Horde/Horde.Build/Issues/Impl/IssueService.cs b/Engine/Source/Programs/Horde/Horde.Build/Issues/Impl/IssueService.cs
--- a/Engine/Source/Programs/Horde/Horde.Build/Issues/Impl/IssueService.cs
+++ b/Engine/Source/Programs/Horde/Horde.Build/Issues/Impl/IssueService.cs
@@ -130,7 +130,7 @@ namespace HordeServer.Services.Impl
 	/// <summary>
 	/// Wraps funtionality for manipulating build health issues
 	/// </summary>
-	public class IssueService : ElectedBackgroundService, IIssueService
+	public sealed class IssueService : IIssueService, IHostedService, IDisposable
 	{
 		/// <summary>
 		/// Maximum number of changes to query from Perforce in one go
@@ -145,6 +145,8 @@ namespace HordeServer.Services.Impl
 		IUserCollection UserCollection;
 		IPerforceService Perforce;
 		ILogFileService LogFileService;
+		IClock Clock;
+		ITicker Ticker;
 
 		/// <summary>
 		/// 
@@ -184,8 +186,7 @@ namespace HordeServer.Services.Impl
 		/// <summary>
 		/// Constructor
 		/// </summary>
-		public IssueService(DatabaseService DatabaseService, RedisService RedisService, IIssueCollection IssueCollection, IJobCollection Jobs, IJobStepRefCollection JobStepRefs, StreamService Streams, IUserCollection UserCollection, IPerforceService Perforce, ILogFileService LogFileService, ILogger<IssueService> Logger)
-			: base(DatabaseService, ObjectId.Parse("609542152fb0794700a6c3df"), Logger)
+		public IssueService(RedisService RedisService, IIssueCollection IssueCollection, IJobCollection Jobs, IJobStepRefCollection JobStepRefs, StreamService Streams, IUserCollection UserCollection, IPerforceService Perforce, ILogFileService LogFileService, IClock Clock, ILogger<IssueService> Logger)
 		{
 			this.RedisService = RedisService;
 
@@ -203,6 +204,8 @@ namespace HordeServer.Services.Impl
 			this.UserCollection = UserCollection;
 			this.Perforce = Perforce;
 			this.LogFileService = LogFileService;
+			this.Clock = Clock;
+			this.Ticker = Clock.AddSharedTicker<IssueService>(TimeSpan.FromMinutes(1.0), TickLeaderAsync, Logger);
 			this.Logger = Logger;
 
 			// Create all the issue factories
@@ -221,14 +224,23 @@ namespace HordeServer.Services.Impl
 			TypeToHandler = Matchers.ToDictionary(x => x.Type, x => x, StringComparer.Ordinal);
 		}
 
+		/// <inheritdoc/>
+		public Task StartAsync(CancellationToken cancellationToken) => Ticker.StartAsync();
+
+		/// <inheritdoc/>
+		public Task StopAsync(CancellationToken cancellationToken) => Ticker.StopAsync();
+
+		/// <inheritdoc/>
+		public void Dispose() => Ticker.Dispose();
+
 		/// <summary>
 		/// Periodically update the list of cached open issues
 		/// </summary>
 		/// <param name="StoppingToken">Token to indicate that the service should stop</param>
 		/// <returns>Async task</returns>
-		protected override async Task<DateTime> TickLeaderAsync(CancellationToken StoppingToken)
+		async ValueTask TickLeaderAsync(CancellationToken StoppingToken)
 		{
-			DateTime UtcNow = DateTime.UtcNow;
+			DateTime UtcNow = Clock.UtcNow;
 
 			Dictionary<StreamId, HashSet<TemplateRefId>> NewCachedDesktopAlerts = new Dictionary<StreamId, HashSet<TemplateRefId>>();
 
@@ -263,8 +275,6 @@ namespace HordeServer.Services.Impl
 				NewCachedOpenIssues.Add(await GetIssueDetailsAsync(OpenIssue));
 			}
 			CachedIssues = NewCachedOpenIssues;
-
-			return UtcNow + TimeSpan.FromMinutes(1.0);
 		}
 
 		/// <inheritdoc/>
diff --git a/Engine/Source/Programs/Horde/Horde.Build/Jobs/Impl/JobTaskSource.cs b/Engine/Source/Programs/Horde/Horde.Build/Jobs/Impl/JobTaskSource.cs
--- a/Engine/Source/Programs/Horde/Horde.Build/Jobs/Impl/JobTaskSource.cs
+++ b/Engine/Source/Programs/Horde/Horde.Build/Jobs/Impl/JobTaskSource.cs
@@ -243,7 +243,7 @@ namespace HordeServer.Tasks.Impl
 			this.LogFileService = LogFileService;
 			this.PerforceLoadBalancer = PerforceLoadBalancer;
 			this.ApplicationLifetime = ApplicationLifetime;
-			this.Ticker = Clock.CreateTicker(RefreshInterval, TickAsync, Logger);
+			this.Ticker = Clock.AddTicker(RefreshInterval, TickAsync, Logger);
 			this.Settings = Settings;
 			this.Logger = Logger;
 
diff --git a/Engine/Source/Programs/Horde/Horde.Build/Logs/Impl/LogFileService.cs b/Engine/Source/Programs/Horde/Horde.Build/Logs/Impl/LogFileService.cs
--- a/Engine/Source/Programs/Horde/Horde.Build/Logs/Impl/LogFileService.cs
+++ b/Engine/Source/Programs/Horde/Horde.Build/Logs/Impl/LogFileService.cs
@@ -507,7 +507,7 @@ namespace HordeServer.Services
 			this.LogFileCache = new MemoryCache(new MemoryCacheOptions());
 			this.Builder = Builder;
 			this.Storage = Storage;
-			this.Ticker = Clock.CreateTicker(TimeSpan.FromSeconds(30.0), TickAsync, Logger);
+			this.Ticker = Clock.AddTicker(TimeSpan.FromSeconds(30.0), TickAsync, Logger);
 			this.Logger = Logger;
 		}
 
diff --git a/Engine/Source/Programs/Horde/Horde.Build/Perforce/Impl/CommitService.cs b/Engine/Source/Programs/Horde/Horde.Build/Perforce/Impl/CommitService.cs
--- a/Engine/Source/Programs/Horde/Horde.Build/Perforce/Impl/CommitService.cs
+++ b/Engine/Source/Programs/Horde/Horde.Build/Perforce/Impl/CommitService.cs
@@ -29,6 +29,7 @@ using EpicGames.Redis;
 using System.Diagnostics;
 using System.IO.Compression;
 using System.IO;
+using HordeCommon;
 
 namespace HordeServer.Commits.Impl
 {
@@ -73,16 +74,15 @@ namespace HordeServer.Commits.Impl
 		IStreamCollection StreamCollection;
 		IPerforceService PerforceService;
 		IUserCollection UserCollection;
-		ISingletonDocument<Globals> GlobalsDocument;
 		ILogger<CommitService> Logger;
-		ElectedTick UpdateCommitsTicker;
+		ITicker UpdateCommitsTicker;
 
 		const int MaxBackgroundTasks = 2;
 
 		/// <summary>
 		/// Constructor
 		/// </summary>
-		public CommitService(IDatabase Redis, ICommitCollection CommitCollection, IBlobCollection BlobCollection, IObjectCollection ObjectCollection, IRefCollection RefCollection, IStreamCollection StreamCollection, IPerforceService PerforceService, IUserCollection UserCollection, DatabaseService DatabaseService, ILogger<CommitService> Logger)
+		public CommitService(IDatabase Redis, ICommitCollection CommitCollection, IBlobCollection BlobCollection, IObjectCollection ObjectCollection, IRefCollection RefCollection, IStreamCollection StreamCollection, IPerforceService PerforceService, IUserCollection UserCollection, IClock Clock, ILogger<CommitService> Logger)
 		{
 			this.Redis = Redis;
 			this.RedisDirtyStreams = new RedisSet<StreamId>(Redis, RedisBaseKey.Append("streams"));
@@ -95,23 +95,19 @@ namespace HordeServer.Commits.Impl
 			this.StreamCollection = StreamCollection;
 			this.PerforceService = PerforceService;
 			this.UserCollection = UserCollection;
-			this.GlobalsDocument = new SingletonDocument<Globals>(DatabaseService);
 			this.Logger = Logger;
-			this.UpdateCommitsTicker = new ElectedTick(DatabaseService, new ObjectId("60f866c49e7268f71803b6ef"), UpdateCommitsAsync, TimeSpan.FromSeconds(30.0), Logger);
+			this.UpdateCommitsTicker = Clock.AddSharedTicker<CommitService>(TimeSpan.FromSeconds(30.0), UpdateCommitsAsync, Logger);
 		}
 
 		/// <inheritdoc/>
-		public void Dispose()
-		{
-			UpdateCommitsTicker.Dispose();
-		}
+		public void Dispose() => UpdateCommitsTicker.Dispose();
 
 		/// <inheritdoc/>
 		public async Task StartAsync(CancellationToken CancellationToken)
 		{
 			if (EnableUpdates)
 			{
-				UpdateCommitsTicker.Start();
+				await UpdateCommitsTicker.StartAsync();
 				await StartTreeUpdatesAsync();
 			}
 		}
@@ -133,7 +129,7 @@ namespace HordeServer.Commits.Impl
 		/// </summary>
 		/// <param name="CancellationToken"></param>
 		/// <returns></returns>
-		async Task UpdateCommitsAsync(CancellationToken CancellationToken)
+		async ValueTask UpdateCommitsAsync(CancellationToken CancellationToken)
 		{
 			List<IStream> Streams = await StreamCollection.FindAllAsync();
 			foreach (IGrouping<string, IStream> Group in Streams.GroupBy(x => x.ClusterName, StringComparer.OrdinalIgnoreCase))
diff --git a/Engine/Source/Programs/Horde/Horde.Build/Perforce/PerforceLoadBalancer.cs b/Engine/Source/Programs/Horde/Horde.Build/Perforce/PerforceLoadBalancer.cs
--- a/Engine/Source/Programs/Horde/Horde.Build/Perforce/PerforceLoadBalancer.cs
+++ b/Engine/Source/Programs/Horde/Horde.Build/Perforce/PerforceLoadBalancer.cs
@@ -1,6 +1,7 @@
 // Copyright Epic Games, Inc. All Rights Reserved.
 
 using Google.Protobuf.WellKnownTypes;
+using HordeCommon;
 using HordeCommon.Rpc.Tasks;
 using HordeServer.Collections;
 using HordeServer.Models;
@@ -140,43 +141,28 @@ namespace HordeServer.Services
 		SingletonDocument<PerforceServerList> ServerListSingleton;
 		Random Random = new Random();
 		ILogger Logger;
-		ElectedTick Ticker;
-		Task InternalTickTask = Task.CompletedTask;
+		ITicker Ticker;
 
 		/// <summary>
 		/// Constructor
 		/// </summary>
-		/// <param name="DatabaseService"></param>
-		/// <param name="LeaseCollection"></param>
-		/// <param name="Logger"></param>
-		public PerforceLoadBalancer(DatabaseService DatabaseService, ILeaseCollection LeaseCollection, ILogger<PerforceLoadBalancer> Logger)
+		public PerforceLoadBalancer(DatabaseService DatabaseService, ILeaseCollection LeaseCollection, IClock Clock, ILogger<PerforceLoadBalancer> Logger)
 		{
 			this.DatabaseService = DatabaseService;
 			this.LeaseCollection = LeaseCollection;
 			this.ServerListSingleton = new SingletonDocument<PerforceServerList>(DatabaseService);
 			this.Logger = Logger;
-			this.Ticker = new ElectedTick(DatabaseService, new ObjectId("603fb20536e999974fd6ff6b"), TickAsync, TimeSpan.FromMinutes(1.0), Logger);
+			this.Ticker = Clock.AddSharedTicker<PerforceLoadBalancer>(TimeSpan.FromMinutes(1.0), TickInternalAsync, Logger);
 		}
 
 		/// <inheritdoc/>
-		public void Dispose()
-		{
-			Ticker.Dispose();
-		}
+		public Task StartAsync(CancellationToken cancellationToken) => Ticker.StartAsync();
 
 		/// <inheritdoc/>
-		public Task StartAsync(CancellationToken cancellationToken)
-		{
-			Ticker.Start();
-			return Task.CompletedTask;
-		}
+		public Task StopAsync(CancellationToken cancellationToken) => Ticker.StopAsync();
 
 		/// <inheritdoc/>
-		public async Task StopAsync(CancellationToken cancellationToken)
-		{
-			await Ticker.StopAsync();
-			await InternalTickTask;
-		}
+		public void Dispose() => Ticker.Dispose();
 
 		/// <summary>
 		/// Get the current server list
@@ -327,30 +313,7 @@ namespace HordeServer.Services
 		}
 
 		/// <inheritdoc/>
-		Task TickAsync(CancellationToken StoppingToken)
-		{
-			if (InternalTickTask.IsCompleted)
-			{
-				InternalTickTask = Task.Run(() => TickInternalGuardedAsync());
-			}
-			return Task.CompletedTask;
-		}
-
-		/// <inheritdoc/>
-		async Task TickInternalGuardedAsync()
-		{
-			try
-			{
-				await TickInternalAsync();
-			}
-			catch (Exception Ex)
-			{
-				Logger.LogError(Ex, "Error updating server state");
-			}
-		}
-
-		/// <inheritdoc/>
-		async Task TickInternalAsync()
+		async ValueTask TickInternalAsync(CancellationToken CancellationToken)
 		{
 			Globals Globals = await DatabaseService.GetGlobalsAsync();
 
diff --git a/Engine/Source/Programs/Horde/Horde.Build/Server/ConfigService.cs b/Engine/Source/Programs/Horde/Horde.Build/Server/ConfigService.cs
--- a/Engine/Source/Programs/Horde/Horde.Build/Server/ConfigService.cs
+++ b/Engine/Source/Programs/Horde/Horde.Build/Server/ConfigService.cs
@@ -24,13 +24,15 @@ using System.Threading;
 using System.Threading.Tasks;
 using System.Reflection;
 
-using PoolId = HordeServer.Utilities.StringId<HordeServer.Models.IPool>;
 using System.Globalization;
 using HordeServer.Storage;
 using System.Text.Json.Serialization;
+using HordeCommon;
+using Microsoft.Extensions.Hosting;
 
 namespace HordeServer.Services
 {
+	using PoolId = StringId<IPool>;
 	using ProjectId = StringId<IProject>;
 	using StreamId = StringId<IStream>;
 	using NamespaceId = StringId<INamespace>;
@@ -39,7 +41,7 @@ namespace HordeServer.Services
 	/// <summary>
 	/// Polls Perforce for stream config changes
 	/// </summary>
-	public class ConfigService : ElectedBackgroundService
+	public sealed class ConfigService : IHostedService, IDisposable
 	{
 		const string FileScheme = "file";
 		const string PerforceScheme = "p4-cluster";
@@ -49,77 +51,23 @@ namespace HordeServer.Services
 		/// </summary>
 		const int Version = 5;
 
-		/// <summary>
-		/// Database service
-		/// </summary>
 		DatabaseService DatabaseService;
-
-		/// <summary>
-		/// Project service
-		/// </summary>
 		ProjectService ProjectService;
-
-		/// <summary>
-		/// Stream service
-		/// </summary>
 		StreamService StreamService;
-
-		/// <summary>
-		/// Instance of the perforce service
-		/// </summary>
 		IPerforceService PerforceService;
-
-		/// <summary>
-		/// Load balancer instance
-		/// </summary>
-		private PerforceLoadBalancer PerforceLoadBalancer;
-
-		/// <summary>
-		/// Instance of the notification service
-		/// </summary>
+		PerforceLoadBalancer PerforceLoadBalancer;
 		INotificationService NotificationService;
-
-		/// <summary>
-		/// The namespace collection
-		/// </summary>
 		INamespaceCollection NamespaceCollection;
-
-		/// <summary>
-		/// The bucket collection
-		/// </summary>
 		IBucketCollection BucketCollection;
-
-		/// <summary>
-		/// Singleton instance of the pool service
-		/// </summary>
-		private readonly PoolService PoolService;
-
-		/// <summary>
-		/// The server settings
-		/// </summary>
+		PoolService PoolService;
 		IOptionsMonitor<ServerSettings> Settings;
-
-		/// <summary>
-		/// Logging device
-		/// </summary>
+		ITicker Ticker;
 		ILogger Logger;
 
 		/// <summary>
 		/// Constructor
 		/// </summary>
-		/// <param name="DatabaseService"></param>
-		/// <param name="PerforceService"></param>
-		/// <param name="PerforceLoadBalancer"></param>
-		/// <param name="ProjectService"></param>
-		/// <param name="StreamService"></param>
-		/// <param name="NotificationService"></param>
-		/// <param name="NamespaceCollection"></param>
-		/// <param name="BucketCollection"></param>
-		/// <param name="PoolService"></param>
-		/// <param name="Settings"></param>
-		/// <param name="Logger"></param>
-		public ConfigService(DatabaseService DatabaseService, IPerforceService PerforceService, ProjectService ProjectService, StreamService StreamService, INotificationService NotificationService, INamespaceCollection NamespaceCollection, IBucketCollection BucketCollection, PoolService PoolService, PerforceLoadBalancer PerforceLoadBalancer, IOptionsMonitor<ServerSettings> Settings, ILogger<ConfigService> Logger)
-			: base(DatabaseService, new ObjectId("5ff60549e7632b15e64ac2f7"), Logger)
+		public ConfigService(DatabaseService DatabaseService, IPerforceService PerforceService, ProjectService ProjectService, StreamService StreamService, INotificationService NotificationService, INamespaceCollection NamespaceCollection, IBucketCollection BucketCollection, PoolService PoolService, PerforceLoadBalancer PerforceLoadBalancer, IClock Clock, IOptionsMonitor<ServerSettings> Settings, ILogger<ConfigService> Logger)
 		{
 			this.DatabaseService = DatabaseService;
 			this.PerforceService = PerforceService;
@@ -131,12 +79,23 @@ namespace HordeServer.Services
 			this.BucketCollection = BucketCollection;
 			this.PoolService = PoolService;
 			this.Settings = Settings;
+			this.Ticker = Clock.AddSharedTicker<ConfigService>(TimeSpan.FromMinutes(1.0), TickLeaderAsync, Logger);
 			this.Logger = Logger;
 
 			// This will trigger if the local Horde.json user configuration is changed
 			this.Settings.OnChange(OnUserConfigUpdated);
 		}
 
+		/// <inheritdoc/>
+		public Task StartAsync(CancellationToken cancellationToken) => Ticker.StartAsync();
+
+		/// <inheritdoc/>
+		public Task StopAsync(CancellationToken cancellationToken) => Ticker.StopAsync();
+
+		/// <inheritdoc/>
+		public void Dispose() => Ticker.Dispose();
+
+
 		GlobalConfig? CachedGlobalConfig;
 		string? CachedGlobalConfigRevision;
 		Dictionary<ProjectId, (ProjectConfig Config, string Revision)> CachedProjectConfigs = new Dictionary<ProjectId, (ProjectConfig, string)>();
@@ -501,7 +460,7 @@ namespace HordeServer.Services
 		}
 
 		/// <inheritdoc/>
-		protected override async Task<DateTime> TickLeaderAsync(CancellationToken StoppingToken)
+		async ValueTask TickLeaderAsync(CancellationToken StoppingToken)
 		{
 
 			Uri? ConfigUri = null;
@@ -521,8 +480,6 @@ namespace HordeServer.Services
 			{
 				await UpdateConfigAsync(ConfigUri);
 			}
-		
-			return DateTime.UtcNow + TimeSpan.FromMinutes(1.0);
 		}
 
 		//
diff --git a/Engine/Source/Programs/Horde/Horde.Build/Server/DowntimeService.cs b/Engine/Source/Programs/Horde/Horde.Build/Server/DowntimeService.cs
--- a/Engine/Source/Programs/Horde/Horde.Build/Server/DowntimeService.cs
+++ b/Engine/Source/Programs/Horde/Horde.Build/Server/DowntimeService.cs
@@ -64,7 +64,7 @@ namespace HordeServer.Services
 			// Ensure the initial value to be correct
 			TickAsync(CancellationToken.None).AsTask().Wait();
 
-			Ticker = Clock.CreateTicker(TimeSpan.FromMinutes(1.0), TickAsync, Logger);
+			Ticker = Clock.AddTicker(TimeSpan.FromMinutes(1.0), TickAsync, Logger);
 		}
 
 		/// <inheritdoc/>
diff --git a/Engine/Source/Programs/Horde/Horde.Build/Server/MetricService.cs b/Engine/Source/Programs/Horde/Horde.Build/Server/MetricService.cs
--- a/Engine/Source/Programs/Horde/Horde.Build/Server/MetricService.cs
+++ b/Engine/Source/Programs/Horde/Horde.Build/Server/MetricService.cs
@@ -25,7 +25,7 @@ namespace HordeServer.Services
 		public MetricService(IDogStatsd DogStatsd, IClock Clock, ILogger<MetricService> Logger)
 		{
 			this.DogStatsd = DogStatsd;
-			this.Ticker = Clock.CreateTicker(TimeSpan.FromSeconds(20.0), TickAsync, Logger);
+			this.Ticker = Clock.AddTicker(TimeSpan.FromSeconds(20.0), TickAsync, Logger);
 		}
 
 		/// <inheritdoc/>
diff --git a/Engine/Source/Programs/Horde/Horde.Build/Services/ConsistencyService.cs b/Engine/Source/Programs/Horde/Horde.Build/Services/ConsistencyService.cs
--- a/Engine/Source/Programs/Horde/Horde.Build/Services/ConsistencyService.cs
+++ b/Engine/Source/Programs/Horde/Horde.Build/Services/ConsistencyService.cs
@@ -4,6 +4,7 @@ using HordeCommon;
 using HordeServer.Collections;
 using HordeServer.Models;
 using HordeServer.Utilities;
+using Microsoft.Extensions.Hosting;
 using Microsoft.Extensions.Logging;
 using MongoDB.Bson;
 using System;
@@ -19,36 +20,40 @@ namespace HordeServer.Services
 	/// <summary>
 	/// Service which checks the database for consistency and fixes up any errors
 	/// </summary>
-	class ConsistencyService : ElectedBackgroundService
+	class ConsistencyService : IHostedService, IDisposable
 	{
 		ISessionCollection SessionCollection;
 		ILeaseCollection LeaseCollection;
+		ITicker Ticker;
 		ILogger<ConsistencyService> Logger;
 
 		/// <summary>
 		/// Constructor
 		/// </summary>
-		/// <param name="DatabaseService">Database service</param>
-		/// <param name="SessionCollection"></param>
-		/// <param name="LeaseCollection"></param>
-		/// <param name="Logger">Logging device</param>
-		public ConsistencyService(DatabaseService DatabaseService, ISessionCollection SessionCollection, ILeaseCollection LeaseCollection, ILogger<ConsistencyService> Logger)
-			: base(DatabaseService, new ObjectId("5fdb9d9f39b1291512821ef8"), Logger)
+		public ConsistencyService(ISessionCollection SessionCollection, ILeaseCollection LeaseCollection, IClock Clock, ILogger<ConsistencyService> Logger)
 		{
 			this.SessionCollection = SessionCollection;
 			this.LeaseCollection = LeaseCollection;
+			this.Ticker = Clock.AddSharedTicker<ConsistencyService>(TimeSpan.FromMinutes(20.0), TickLeaderAsync, Logger);
 			this.Logger = Logger;
 		}
 
+		/// <inheritdoc/>
+		public Task StartAsync(CancellationToken cancellationToken) => Ticker.StartAsync();
+
+		/// <inheritdoc/>
+		public Task StopAsync(CancellationToken cancellationToken) => Ticker.StopAsync();
+
+		/// <inheritdoc/>
+		public void Dispose() => Ticker.Dispose();
+
 		/// <summary>
 		/// Poll for inconsistencies in the database
 		/// </summary>
 		/// <param name="StoppingToken">Stopping token</param>
 		/// <returns>Async task</returns>
-		protected override async Task<DateTime> TickLeaderAsync(CancellationToken StoppingToken)
+		async ValueTask TickLeaderAsync(CancellationToken StoppingToken)
 		{
-			DateTime NextTickTime = DateTime.UtcNow + TimeSpan.FromMinutes(20.0);
-
 			List<ISession> Sessions = await SessionCollection.FindActiveSessionsAsync();
 			Dictionary<SessionId, ISession> SessionIdToInstance = Sessions.ToDictionary(x => x.Id, x => x);
 
@@ -64,8 +69,6 @@ namespace HordeServer.Services
 					await LeaseCollection.TrySetOutcomeAsync(Lease.Id, FinishTime, LeaseOutcome.Cancelled, null);
 				}
 			}
-
-			return NextTickTime;
 		}
 	}
 }
diff --git a/Engine/Source/Programs/Horde/Horde.Build/Utilities/Clock.cs b/Engine/Source/Programs/Horde/Horde.Build/Utilities/Clock.cs
--- a/Engine/Source/Programs/Horde/Horde.Build/Utilities/Clock.cs
+++ b/Engine/Source/Programs/Horde/Horde.Build/Utilities/Clock.cs
@@ -1,7 +1,10 @@
 // Copyright Epic Games, Inc. All Rights Reserved.
 
 using EpicGames.Core;
+using EpicGames.Redis.Utility;
+using HordeServer.Services;
 using Microsoft.Extensions.Logging;
+using StackExchange.Redis;
 using System;
 using System.Collections.Generic;
 using System.Diagnostics;
@@ -44,7 +47,17 @@ namespace HordeCommon
 		/// <param name="TickAsync">Callback for the tick. Returns the time interval until the next tick, or null to cancel the tick.</param>
 		/// <param name="Logger">Logger for error messages</param>
 		/// <returns>Handle to the event</returns>
-		ITicker CreateTicker(TimeSpan Interval, Func<CancellationToken, ValueTask<TimeSpan?>> TickAsync, ILogger Logger);
+		ITicker AddTicker(TimeSpan Interval, Func<CancellationToken, ValueTask<TimeSpan?>> TickAsync, ILogger Logger);
+
+		/// <summary>
+		/// Create a ticker shared between all server pods
+		/// </summary>
+		/// <param name="Name">Name of the event</param>
+		/// <param name="Interval">Time after which the event will trigger</param>
+		/// <param name="TickAsync">Callback for the tick. Returns the time interval until the next tick, or null to cancel the tick.</param>
+		/// <param name="Logger">Logger for error messages</param>
+		/// <returns>New ticker instance</returns>
+		ITicker AddSharedTicker(string Name, TimeSpan Interval, Func<CancellationToken, ValueTask> TickAsync, ILogger Logger);
 	}
 
 	/// <summary>
@@ -60,7 +73,7 @@ namespace HordeCommon
 		/// <param name="TickAsync">Trigger callback</param>
 		/// <param name="Logger">Logger for any error messages</param>
 		/// <returns>Handle to the event</returns>
-		public static ITicker CreateTicker(this IClock Clock, TimeSpan Interval, Func<CancellationToken, ValueTask> TickAsync, ILogger Logger)
+		public static ITicker AddTicker(this IClock Clock, TimeSpan Interval, Func<CancellationToken, ValueTask> TickAsync, ILogger Logger)
 		{
 			Func<CancellationToken, ValueTask<TimeSpan?>> WrappedTrigger = async Token =>
 			{
@@ -68,8 +81,18 @@ namespace HordeCommon
 				await TickAsync(Token);
 				return Interval - Timer.Elapsed;
 			};
-			return Clock.CreateTicker(Interval, WrappedTrigger, Logger);
+			return Clock.AddTicker(Interval, WrappedTrigger, Logger);
 		}
+
+		/// <summary>
+		/// Create a ticker shared between all server pods
+		/// </summary>
+		/// <param name="Clock">Clock to schedule the event on</param>
+		/// <param name="Interval">Time after which the event will trigger</param>
+		/// <param name="TickAsync">Callback for the tick. Returns the time interval until the next tick, or null to cancel the tick.</param>
+		/// <param name="Logger">Logger for error messages</param>
+		/// <returns>New ticker instance</returns>
+		public static ITicker AddSharedTicker<T>(this IClock Clock, TimeSpan Interval, Func<CancellationToken, ValueTask> TickAsync, ILogger Logger) => Clock.AddSharedTicker(typeof(T).Name, Interval, TickAsync, Logger);
 	}
 
 	/// <summary>
@@ -151,22 +174,44 @@ namespace HordeCommon
 			}
 		}
 
+		RedisService Redis;
+
 		/// <inheritdoc/>
 		public DateTime UtcNow => DateTime.UtcNow;
 
 		/// <summary>
-		/// Create an event that will trigger after the given time for one agent in the cluster
+		/// Constructor
 		/// </summary>
-		/// <param name="Delay">Time after which the event will trigger</param>
-		/// <param name="TriggerAsync">Callback for the event triggering</param>
-		/// <param name="Logger">Logger for error messages</param>
-		/// <returns>Handle to the event</returns>
-		public ITicker CreateTicker(TimeSpan Delay, Func<CancellationToken, ValueTask<TimeSpan?>> TriggerAsync, ILogger Logger)
+		public Clock(RedisService Redis)
 		{
-			return new TickerImpl(Delay, TriggerAsync, Logger);
+			this.Redis = Redis;
+		}
+
+		/// <inheritdoc/>
+		public ITicker AddTicker(TimeSpan Delay, Func<CancellationToken, ValueTask<TimeSpan?>> TickAsync, ILogger Logger)
+		{
+			return new TickerImpl(Delay, TickAsync, Logger);
+		}
+
+		/// <inheritdoc/>
+		public ITicker AddSharedTicker(string Name, TimeSpan Delay, Func<CancellationToken, ValueTask> TickAsync, ILogger Logger)
+		{
+			RedisKey Key = new RedisKey($"tick/{Name}");
+			return ClockExtensions.AddTicker(this, Delay / 4, Token => TriggerSharedAsync(Key, Delay, TickAsync, Token), Logger);
+		}
+
+		async ValueTask TriggerSharedAsync(RedisKey Key, TimeSpan Interval, Func<CancellationToken, ValueTask> TickAsync, CancellationToken CancellationToken)
+		{
+			using (RedisLock Lock = new RedisLock(Redis.Database, Key))
+			{
+				if (await Lock.AcquireAsync(Interval, false))
+				{
+					await TickAsync(CancellationToken);
+				}
+			}
 		}
 	}
-	
+
 	/// <summary>
 	/// Fake clock that doesn't advance by wall block time
 	/// Requires manual ticking to progress. Used in tests.
@@ -263,9 +308,15 @@ namespace HordeCommon
 		}
 
 		/// <inheritdoc/>
-		public ITicker CreateTicker(TimeSpan Interval, Func<CancellationToken, ValueTask<TimeSpan?>> TriggerAsync, ILogger Logger)
+		public ITicker AddTicker(TimeSpan Interval, Func<CancellationToken, ValueTask<TimeSpan?>> TickAsync, ILogger Logger)
+		{
+			return new TickerImpl(this, Interval, TickAsync);
+		}
+
+		/// <inheritdoc/>
+		public ITicker AddSharedTicker(string Name, TimeSpan Interval, Func<CancellationToken, ValueTask> TickAsync, ILogger Logger)
 		{
-			return new TickerImpl(this, Interval, TriggerAsync);
+			return AddTicker(Interval, async Token => { await TickAsync(Token); return Interval; }, Logger);
 		}
 	}
 }
\ No newline at end of file
diff --git a/Engine/Source/Programs/Horde/Horde.Build/Utilities/ElectedBackgroundService.cs b/Engine/Source/Programs/Horde/Horde.Build/Utilities/ElectedBackgroundService.cs
deleted file mode 100644
diff --git a/Engine/Source/Programs/Horde/Horde.Build/Utilities/ElectedTick.cs b/Engine/Source/Programs/Horde/Horde.Build/Utilities/ElectedTick.cs
deleted file mode 100644
diff --git a/Engine/Source/Programs/Horde/Horde.Build/Utilities/TickedBackgroundService.cs b/Engine/Source/Programs/Horde/Horde.Build/Utilities/TickedBackgroundService.cs
deleted file mode 100644
