Horde: Enable code analysis warnings for Horde.Build (currently IDE only due to long build times, but also enabled in Docker builds)

diff --git a/Engine/Source/Programs/Horde/Horde.Build/.editorconfig b/Engine/Source/Programs/Horde/Horde.Build/.editorconfig
new file mode 100644
--- /dev/null
+++ b/Engine/Source/Programs/Horde/Horde.Build/.editorconfig
@@ -0,0 +1,818 @@
+# NOTE: Requires **VS2019 16.3** or later
+
+# All Rules Enabled as build warnings
+# Description: All Rules are enabled as build warnings. Rules with IsEnabledByDefault = false are force enabled as build warnings.
+
+# Code files
+[*.{cs,vb}]
+
+
+# CA1000: Do not declare static members on generic types
+dotnet_diagnostic.CA1000.severity = warning
+
+# CA1001: Types that own disposable fields should be disposable
+dotnet_diagnostic.CA1001.severity = warning
+
+# CA1002: Do not expose generic lists
+dotnet_diagnostic.CA1002.severity = none
+
+# CA1003: Use generic event handler instances
+dotnet_diagnostic.CA1003.severity = none
+
+# CA1005: Avoid excessive parameters on generic types
+dotnet_diagnostic.CA1005.severity = warning
+
+# CA1008: Enums should have zero value
+dotnet_diagnostic.CA1008.severity = warning
+
+# CA1010: Generic interface should also be implemented
+dotnet_diagnostic.CA1010.severity = warning
+
+# CA1012: Abstract types should not have public constructors
+dotnet_diagnostic.CA1012.severity = warning
+
+# CA1014: Mark assemblies with CLSCompliant
+dotnet_diagnostic.CA1014.severity = warning
+
+# CA1016: Mark assemblies with assembly version
+dotnet_diagnostic.CA1016.severity = warning
+
+# CA1017: Mark assemblies with ComVisible
+dotnet_diagnostic.CA1017.severity = warning
+
+# CA1018: Mark attributes with AttributeUsageAttribute
+dotnet_diagnostic.CA1018.severity = warning
+
+# CA1019: Define accessors for attribute arguments
+dotnet_diagnostic.CA1019.severity = warning
+
+# CA1021: Avoid out parameters
+dotnet_diagnostic.CA1021.severity = none
+
+# CA1024: Use properties where appropriate
+dotnet_diagnostic.CA1024.severity = none
+
+# CA1027: Mark enums with FlagsAttribute
+dotnet_diagnostic.CA1027.severity = warning
+
+# CA1028: Enum Storage should be Int32
+dotnet_diagnostic.CA1028.severity = warning
+
+# CA1030: Use events where appropriate
+dotnet_diagnostic.CA1030.severity = warning
+
+# CA1031: Do not catch general exception types
+dotnet_diagnostic.CA1031.severity = none
+
+# CA1032: Implement standard exception constructors
+dotnet_diagnostic.CA1032.severity = none
+
+# CA1033: Interface methods should be callable by child types
+dotnet_diagnostic.CA1033.severity = warning
+
+# CA1034: Nested types should not be visible
+dotnet_diagnostic.CA1034.severity = warning
+
+# CA1036: Override methods on comparable types
+dotnet_diagnostic.CA1036.severity = warning
+
+# CA1040: Avoid empty interfaces
+dotnet_diagnostic.CA1040.severity = warning
+
+# CA1041: Provide ObsoleteAttribute message
+dotnet_diagnostic.CA1041.severity = warning
+
+# CA1043: Use Integral Or String Argument For Indexers
+dotnet_diagnostic.CA1043.severity = warning
+
+# CA1044: Properties should not be write only
+dotnet_diagnostic.CA1044.severity = warning
+
+# CA1045: Do not pass types by reference
+dotnet_diagnostic.CA1045.severity = warning
+
+# CA1046: Do not overload equality operator on reference types
+dotnet_diagnostic.CA1046.severity = warning
+
+# CA1047: Do not declare protected member in sealed type
+dotnet_diagnostic.CA1047.severity = warning
+
+# CA1050: Declare types in namespaces
+dotnet_diagnostic.CA1050.severity = warning
+
+# CA1051: Do not declare visible instance fields
+dotnet_diagnostic.CA1051.severity = warning
+
+# CA1052: Static holder types should be Static or NotInheritable
+dotnet_diagnostic.CA1052.severity = warning
+
+# CA1054: URI-like parameters should not be strings
+dotnet_diagnostic.CA1054.severity = warning
+
+# CA1055: URI-like return values should not be strings
+dotnet_diagnostic.CA1055.severity = warning
+
+# CA1056: URI-like properties should not be strings
+dotnet_diagnostic.CA1056.severity = warning
+
+# CA1058: Types should not extend certain base types
+dotnet_diagnostic.CA1058.severity = warning
+
+# CA1060: Move pinvokes to native methods class
+dotnet_diagnostic.CA1060.severity = none
+
+# CA1061: Do not hide base class methods
+dotnet_diagnostic.CA1061.severity = warning
+
+# CA1062: Validate arguments of public methods
+dotnet_diagnostic.CA1062.severity = none
+
+# CA1063: Implement IDisposable Correctly
+dotnet_diagnostic.CA1063.severity = warning
+
+# CA1064: Exceptions should be public
+dotnet_diagnostic.CA1064.severity = warning
+
+# CA1065: Do not raise exceptions in unexpected locations
+dotnet_diagnostic.CA1065.severity = none
+
+# CA1066: Implement IEquatable when overriding Object.Equals
+dotnet_diagnostic.CA1066.severity = warning
+
+# CA1067: Override Object.Equals(object) when implementing IEquatable<T>
+dotnet_diagnostic.CA1067.severity = warning
+
+# CA1068: CancellationToken parameters must come last
+dotnet_diagnostic.CA1068.severity = warning
+
+# CA1069: Enums values should not be duplicated
+dotnet_diagnostic.CA1069.severity = warning
+
+# CA1070: Do not declare event fields as virtual
+dotnet_diagnostic.CA1070.severity = warning
+
+# CA1200: Avoid using cref tags with a prefix
+dotnet_diagnostic.CA1200.severity = warning
+
+# CA1303: Do not pass literals as localized parameters
+dotnet_diagnostic.CA1303.severity = warning
+
+# CA1304: Specify CultureInfo
+dotnet_diagnostic.CA1304.severity = none
+
+# CA1305: Specify IFormatProvider
+dotnet_diagnostic.CA1305.severity = none
+
+# CA1307: Specify StringComparison for clarity
+dotnet_diagnostic.CA1307.severity = warning
+
+# CA1308: Normalize strings to uppercase
+dotnet_diagnostic.CA1308.severity = warning
+
+# CA1309: Use ordinal string comparison
+dotnet_diagnostic.CA1309.severity = warning
+
+# CA1310: Specify StringComparison for correctness
+dotnet_diagnostic.CA1310.severity = warning
+
+# CA1401: P/Invokes should not be visible
+dotnet_diagnostic.CA1401.severity = warning
+
+# CA1416: Validate platform compatibility
+dotnet_diagnostic.CA1416.severity = warning
+
+# CA1417: Do not use 'OutAttribute' on string parameters for P/Invokes
+dotnet_diagnostic.CA1417.severity = warning
+
+# CA1418: Use valid platform string
+dotnet_diagnostic.CA1418.severity = warning
+
+# CA1419: Provide a parameterless constructor that is as visible as the containing type for concrete types derived from 'System.Runtime.InteropServices.SafeHandle'
+dotnet_diagnostic.CA1419.severity = warning
+
+# CA1501: Avoid excessive inheritance
+dotnet_diagnostic.CA1501.severity = warning
+
+# CA1502: Avoid excessive complexity
+dotnet_diagnostic.CA1502.severity = none
+
+# CA1505: Avoid unmaintainable code
+dotnet_diagnostic.CA1505.severity = warning
+
+# CA1506: Avoid excessive class coupling
+dotnet_diagnostic.CA1506.severity = none
+
+# CA1507: Use nameof to express symbol names
+dotnet_diagnostic.CA1507.severity = warning
+
+# CA1508: Avoid dead conditional code
+dotnet_diagnostic.CA1508.severity = warning
+
+# CA1509: Invalid entry in code metrics rule specification file
+dotnet_diagnostic.CA1509.severity = warning
+
+# CA1700: Do not name enum values 'Reserved'
+dotnet_diagnostic.CA1700.severity = warning
+
+# CA1707: Identifiers should not contain underscores
+dotnet_diagnostic.CA1707.severity = warning
+
+# CA1708: Identifiers should differ by more than case
+dotnet_diagnostic.CA1708.severity = warning
+
+# CA1710: Identifiers should have correct suffix
+dotnet_diagnostic.CA1710.severity = warning
+
+# CA1711: Identifiers should not have incorrect suffix
+dotnet_diagnostic.CA1711.severity = none
+
+# CA1712: Do not prefix enum values with type name
+dotnet_diagnostic.CA1712.severity = warning
+
+# CA1713: Events should not have 'Before' or 'After' prefix
+dotnet_diagnostic.CA1713.severity = warning
+
+# CA1715: Identifiers should have correct prefix
+dotnet_diagnostic.CA1715.severity = warning
+
+# CA1716: Identifiers should not match keywords
+dotnet_diagnostic.CA1716.severity = warning
+
+# CA1720: Identifier contains type name
+dotnet_diagnostic.CA1720.severity = none
+
+# CA1721: Property names should not match get methods
+dotnet_diagnostic.CA1721.severity = warning
+
+# CA1724: Type names should not match namespaces
+dotnet_diagnostic.CA1724.severity = warning
+
+# CA1725: Parameter names should match base declaration
+dotnet_diagnostic.CA1725.severity = none
+
+# CA1727: Use PascalCase for named placeholders
+dotnet_diagnostic.CA1727.severity = warning
+
+# CA1802: Use literals where appropriate
+dotnet_diagnostic.CA1802.severity = warning
+
+# CA1805: Do not initialize unnecessarily
+dotnet_diagnostic.CA1805.severity = none
+
+# CA1806: Do not ignore method results
+dotnet_diagnostic.CA1806.severity = warning
+
+# CA1810: Initialize reference type static fields inline
+dotnet_diagnostic.CA1810.severity = warning
+
+# CA1812: Avoid uninstantiated internal classes
+dotnet_diagnostic.CA1812.severity = warning
+
+# CA1813: Avoid unsealed attributes
+dotnet_diagnostic.CA1813.severity = warning
+
+# CA1814: Prefer jagged arrays over multidimensional
+dotnet_diagnostic.CA1814.severity = warning
+
+# CA1815: Override equals and operator equals on value types
+dotnet_diagnostic.CA1815.severity = none
+
+# CA1816: Dispose methods should call SuppressFinalize
+dotnet_diagnostic.CA1816.severity = warning
+
+# CA1819: Properties should not return arrays
+dotnet_diagnostic.CA1819.severity = warning
+
+# CA1820: Test for empty strings using string length
+dotnet_diagnostic.CA1820.severity = warning
+
+# CA1821: Remove empty Finalizers
+dotnet_diagnostic.CA1821.severity = warning
+
+# CA1822: Mark members as static
+dotnet_diagnostic.CA1822.severity = warning
+
+# CA1823: Avoid unused private fields
+dotnet_diagnostic.CA1823.severity = warning
+
+# CA1824: Mark assemblies with NeutralResourcesLanguageAttribute
+dotnet_diagnostic.CA1824.severity = warning
+
+# CA1825: Avoid zero-length array allocations
+dotnet_diagnostic.CA1825.severity = warning
+
+# CA1826: Do not use Enumerable methods on indexable collections
+dotnet_diagnostic.CA1826.severity = warning
+
+# CA1827: Do not use Count() or LongCount() when Any() can be used
+dotnet_diagnostic.CA1827.severity = warning
+
+# CA1828: Do not use CountAsync() or LongCountAsync() when AnyAsync() can be used
+dotnet_diagnostic.CA1828.severity = warning
+
+# CA1829: Use Length/Count property instead of Count() when available
+dotnet_diagnostic.CA1829.severity = warning
+
+# CA1830: Prefer strongly-typed Append and Insert method overloads on StringBuilder
+dotnet_diagnostic.CA1830.severity = warning
+
+# CA1831: Use AsSpan or AsMemory instead of Range-based indexers when appropriate
+dotnet_diagnostic.CA1831.severity = warning
+
+# CA1832: Use AsSpan or AsMemory instead of Range-based indexers when appropriate
+dotnet_diagnostic.CA1832.severity = warning
+
+# CA1833: Use AsSpan or AsMemory instead of Range-based indexers when appropriate
+dotnet_diagnostic.CA1833.severity = warning
+
+# CA1834: Consider using 'StringBuilder.Append(char)' when applicable
+dotnet_diagnostic.CA1834.severity = warning
+
+# CA1835: Prefer the 'Memory'-based overloads for 'ReadAsync' and 'WriteAsync'
+dotnet_diagnostic.CA1835.severity = none
+
+# CA1836: Prefer IsEmpty over Count
+dotnet_diagnostic.CA1836.severity = warning
+
+# CA1837: Use 'Environment.ProcessId'
+dotnet_diagnostic.CA1837.severity = warning
+
+# CA1838: Avoid 'StringBuilder' parameters for P/Invokes
+dotnet_diagnostic.CA1838.severity = warning
+
+# CA1839: Use 'Environment.ProcessPath'
+dotnet_diagnostic.CA1839.severity = warning
+
+# CA1840: Use 'Environment.CurrentManagedThreadId'
+dotnet_diagnostic.CA1840.severity = warning
+
+# CA1841: Prefer Dictionary.Contains methods
+dotnet_diagnostic.CA1841.severity = warning
+
+# CA1842: Do not use 'WhenAll' with a single task
+dotnet_diagnostic.CA1842.severity = warning
+
+# CA1843: Do not use 'WaitAll' with a single task
+dotnet_diagnostic.CA1843.severity = warning
+
+# CA1844: Provide memory-based overrides of async methods when subclassing 'Stream'
+dotnet_diagnostic.CA1844.severity = warning
+
+# CA1845: Use span-based 'string.Concat'
+dotnet_diagnostic.CA1845.severity = none
+
+# CA1846: Prefer 'AsSpan' over 'Substring'
+dotnet_diagnostic.CA1846.severity = none
+
+# CA1847: Use char literal for a single character lookup
+dotnet_diagnostic.CA1847.severity = warning
+
+# CA1848: Use the LoggerMessage delegates
+dotnet_diagnostic.CA1848.severity = none
+
+# CA1849: Call async methods when in an async method
+dotnet_diagnostic.CA1849.severity = warning
+
+# CA2000: Dispose objects before losing scope
+dotnet_diagnostic.CA2000.severity = warning
+
+# CA2002: Do not lock on objects with weak identity
+dotnet_diagnostic.CA2002.severity = warning
+
+# CA2007: Consider calling ConfigureAwait on the awaited task
+dotnet_diagnostic.CA2007.severity = none
+
+# CA2008: Do not create tasks without passing a TaskScheduler
+dotnet_diagnostic.CA2008.severity = warning
+
+# CA2009: Do not call ToImmutableCollection on an ImmutableCollection value
+dotnet_diagnostic.CA2009.severity = warning
+
+# CA2011: Avoid infinite recursion
+dotnet_diagnostic.CA2011.severity = warning
+
+# CA2012: Use ValueTasks correctly
+dotnet_diagnostic.CA2012.severity = warning
+
+# CA2013: Do not use ReferenceEquals with value types
+dotnet_diagnostic.CA2013.severity = warning
+
+# CA2014: Do not use stackalloc in loops
+dotnet_diagnostic.CA2014.severity = warning
+
+# CA2015: Do not define finalizers for types derived from MemoryManager<T>
+dotnet_diagnostic.CA2015.severity = warning
+
+# CA2016: Forward the 'CancellationToken' parameter to methods
+dotnet_diagnostic.CA2016.severity = warning
+
+# CA2017: Parameter count mismatch
+dotnet_diagnostic.CA2017.severity = warning
+
+# CA2018: 'Buffer.BlockCopy' expects the number of bytes to be copied for the 'count' argument
+dotnet_diagnostic.CA2018.severity = warning
+
+# CA2100: Review SQL queries for security vulnerabilities
+dotnet_diagnostic.CA2100.severity = warning
+
+# CA2101: Specify marshaling for P/Invoke string arguments
+dotnet_diagnostic.CA2101.severity = none
+
+# CA2109: Review visible event handlers
+dotnet_diagnostic.CA2109.severity = warning
+
+# CA2119: Seal methods that satisfy private interfaces
+dotnet_diagnostic.CA2119.severity = warning
+
+# CA2153: Do Not Catch Corrupted State Exceptions
+dotnet_diagnostic.CA2153.severity = warning
+
+# CA2200: Rethrow to preserve stack details
+dotnet_diagnostic.CA2200.severity = warning
+
+# CA2201: Do not raise reserved exception types
+dotnet_diagnostic.CA2201.severity = none
+
+# CA2207: Initialize value type static fields inline
+dotnet_diagnostic.CA2207.severity = warning
+
+# CA2208: Instantiate argument exceptions correctly
+dotnet_diagnostic.CA2208.severity = warning
+
+# CA2211: Non-constant fields should not be visible
+dotnet_diagnostic.CA2211.severity = warning
+
+# CA2213: Disposable fields should be disposed
+dotnet_diagnostic.CA2213.severity = warning
+
+# CA2214: Do not call overridable methods in constructors
+dotnet_diagnostic.CA2214.severity = warning
+
+# CA2215: Dispose methods should call base class dispose
+dotnet_diagnostic.CA2215.severity = warning
+
+# CA2216: Disposable types should declare finalizer
+dotnet_diagnostic.CA2216.severity = warning
+
+# CA2217: Do not mark enums with FlagsAttribute
+dotnet_diagnostic.CA2217.severity = warning
+
+# CA2218: Override GetHashCode on overriding Equals
+dotnet_diagnostic.CA2218.severity = warning
+
+# CA2219: Do not raise exceptions in finally clauses
+dotnet_diagnostic.CA2219.severity = warning
+
+# CA2224: Override Equals on overloading operator equals
+dotnet_diagnostic.CA2224.severity = warning
+
+# CA2225: Operator overloads have named alternates
+dotnet_diagnostic.CA2225.severity = none
+
+# CA2226: Operators should have symmetrical overloads
+dotnet_diagnostic.CA2226.severity = warning
+
+# CA2227: Collection properties should be read only
+dotnet_diagnostic.CA2227.severity = warning
+
+# CA2229: Implement serialization constructors
+dotnet_diagnostic.CA2229.severity = warning
+
+# CA2231: Overload operator equals on overriding value type Equals
+dotnet_diagnostic.CA2231.severity = warning
+
+# CA2234: Pass system uri objects instead of strings
+dotnet_diagnostic.CA2234.severity = warning
+
+# CA2235: Mark all non-serializable fields
+dotnet_diagnostic.CA2235.severity = warning
+
+# CA2237: Mark ISerializable types with serializable
+dotnet_diagnostic.CA2237.severity = warning
+
+# CA2241: Provide correct arguments to formatting methods
+dotnet_diagnostic.CA2241.severity = warning
+
+# CA2242: Test for NaN correctly
+dotnet_diagnostic.CA2242.severity = warning
+
+# CA2243: Attribute string literals should parse correctly
+dotnet_diagnostic.CA2243.severity = warning
+
+# CA2244: Do not duplicate indexed element initializations
+dotnet_diagnostic.CA2244.severity = warning
+
+# CA2245: Do not assign a property to itself
+dotnet_diagnostic.CA2245.severity = warning
+
+# CA2246: Assigning symbol and its member in the same statement
+dotnet_diagnostic.CA2246.severity = warning
+
+# CA2247: Argument passed to TaskCompletionSource constructor should be TaskCreationOptions enum instead of TaskContinuationOptions enum
+dotnet_diagnostic.CA2247.severity = warning
+
+# CA2248: Provide correct 'enum' argument to 'Enum.HasFlag'
+dotnet_diagnostic.CA2248.severity = warning
+
+# CA2249: Consider using 'string.Contains' instead of 'string.IndexOf'
+dotnet_diagnostic.CA2249.severity = warning
+
+# CA2250: Use 'ThrowIfCancellationRequested'
+dotnet_diagnostic.CA2250.severity = warning
+
+# CA2251: Use 'string.Equals'
+dotnet_diagnostic.CA2251.severity = warning
+
+# CA2252: This API requires opting into preview features
+dotnet_diagnostic.CA2252.severity = warning
+
+# CA2253: Named placeholders should not be numeric values
+dotnet_diagnostic.CA2253.severity = warning
+
+# CA2254: Template should be a static expression
+dotnet_diagnostic.CA2254.severity = warning
+
+# CA2255: The 'ModuleInitializer' attribute should not be used in libraries
+dotnet_diagnostic.CA2255.severity = warning
+
+# CA2256: All members declared in parent interfaces must have an implementation in a DynamicInterfaceCastableImplementation-attributed interface
+dotnet_diagnostic.CA2256.severity = warning
+
+# CA2257: Members defined on an interface with the 'DynamicInterfaceCastableImplementationAttribute' should be 'static'
+dotnet_diagnostic.CA2257.severity = warning
+
+# CA2258: Providing a 'DynamicInterfaceCastableImplementation' interface in Visual Basic is unsupported
+dotnet_diagnostic.CA2258.severity = warning
+
+# CA2300: Do not use insecure deserializer BinaryFormatter
+dotnet_diagnostic.CA2300.severity = warning
+
+# CA2301: Do not call BinaryFormatter.Deserialize without first setting BinaryFormatter.Binder
+dotnet_diagnostic.CA2301.severity = warning
+
+# CA2302: Ensure BinaryFormatter.Binder is set before calling BinaryFormatter.Deserialize
+dotnet_diagnostic.CA2302.severity = warning
+
+# CA2305: Do not use insecure deserializer LosFormatter
+dotnet_diagnostic.CA2305.severity = warning
+
+# CA2310: Do not use insecure deserializer NetDataContractSerializer
+dotnet_diagnostic.CA2310.severity = warning
+
+# CA2311: Do not deserialize without first setting NetDataContractSerializer.Binder
+dotnet_diagnostic.CA2311.severity = warning
+
+# CA2312: Ensure NetDataContractSerializer.Binder is set before deserializing
+dotnet_diagnostic.CA2312.severity = warning
+
+# CA2315: Do not use insecure deserializer ObjectStateFormatter
+dotnet_diagnostic.CA2315.severity = warning
+
+# CA2321: Do not deserialize with JavaScriptSerializer using a SimpleTypeResolver
+dotnet_diagnostic.CA2321.severity = warning
+
+# CA2322: Ensure JavaScriptSerializer is not initialized with SimpleTypeResolver before deserializing
+dotnet_diagnostic.CA2322.severity = warning
+
+# CA2326: Do not use TypeNameHandling values other than None
+dotnet_diagnostic.CA2326.severity = warning
+
+# CA2327: Do not use insecure JsonSerializerSettings
+dotnet_diagnostic.CA2327.severity = warning
+
+# CA2328: Ensure that JsonSerializerSettings are secure
+dotnet_diagnostic.CA2328.severity = warning
+
+# CA2329: Do not deserialize with JsonSerializer using an insecure configuration
+dotnet_diagnostic.CA2329.severity = warning
+
+# CA2330: Ensure that JsonSerializer has a secure configuration when deserializing
+dotnet_diagnostic.CA2330.severity = warning
+
+# CA2350: Do not use DataTable.ReadXml() with untrusted data
+dotnet_diagnostic.CA2350.severity = warning
+
+# CA2351: Do not use DataSet.ReadXml() with untrusted data
+dotnet_diagnostic.CA2351.severity = warning
+
+# CA2352: Unsafe DataSet or DataTable in serializable type can be vulnerable to remote code execution attacks
+dotnet_diagnostic.CA2352.severity = warning
+
+# CA2353: Unsafe DataSet or DataTable in serializable type
+dotnet_diagnostic.CA2353.severity = warning
+
+# CA2354: Unsafe DataSet or DataTable in deserialized object graph can be vulnerable to remote code execution attacks
+dotnet_diagnostic.CA2354.severity = warning
+
+# CA2355: Unsafe DataSet or DataTable type found in deserializable object graph
+dotnet_diagnostic.CA2355.severity = warning
+
+# CA2356: Unsafe DataSet or DataTable type in web deserializable object graph
+dotnet_diagnostic.CA2356.severity = warning
+
+# CA2361: Ensure auto-generated class containing DataSet.ReadXml() is not used with untrusted data
+dotnet_diagnostic.CA2361.severity = warning
+
+# CA2362: Unsafe DataSet or DataTable in auto-generated serializable type can be vulnerable to remote code execution attacks
+dotnet_diagnostic.CA2362.severity = warning
+
+# CA3001: Review code for SQL injection vulnerabilities
+dotnet_diagnostic.CA3001.severity = warning
+
+# CA3002: Review code for XSS vulnerabilities
+dotnet_diagnostic.CA3002.severity = warning
+
+# CA3003: Review code for file path injection vulnerabilities
+dotnet_diagnostic.CA3003.severity = warning
+
+# CA3004: Review code for information disclosure vulnerabilities
+dotnet_diagnostic.CA3004.severity = warning
+
+# CA3005: Review code for LDAP injection vulnerabilities
+dotnet_diagnostic.CA3005.severity = warning
+
+# CA3006: Review code for process command injection vulnerabilities
+dotnet_diagnostic.CA3006.severity = warning
+
+# CA3007: Review code for open redirect vulnerabilities
+dotnet_diagnostic.CA3007.severity = warning
+
+# CA3008: Review code for XPath injection vulnerabilities
+dotnet_diagnostic.CA3008.severity = warning
+
+# CA3009: Review code for XML injection vulnerabilities
+dotnet_diagnostic.CA3009.severity = warning
+
+# CA3010: Review code for XAML injection vulnerabilities
+dotnet_diagnostic.CA3010.severity = warning
+
+# CA3011: Review code for DLL injection vulnerabilities
+dotnet_diagnostic.CA3011.severity = warning
+
+# CA3012: Review code for regex injection vulnerabilities
+dotnet_diagnostic.CA3012.severity = warning
+
+# CA3061: Do Not Add Schema By URL
+dotnet_diagnostic.CA3061.severity = warning
+
+# CA3075: Insecure DTD processing in XML
+dotnet_diagnostic.CA3075.severity = warning
+
+# CA3076: Insecure XSLT script processing.
+dotnet_diagnostic.CA3076.severity = warning
+
+# CA3077: Insecure Processing in API Design, XmlDocument and XmlTextReader
+dotnet_diagnostic.CA3077.severity = warning
+
+# CA3147: Mark Verb Handlers With Validate Antiforgery Token
+dotnet_diagnostic.CA3147.severity = warning
+
+# CA5350: Do Not Use Weak Cryptographic Algorithms
+dotnet_diagnostic.CA5350.severity = warning
+
+# CA5351: Do Not Use Broken Cryptographic Algorithms
+dotnet_diagnostic.CA5351.severity = warning
+
+# CA5358: Review cipher mode usage with cryptography experts
+dotnet_diagnostic.CA5358.severity = warning
+
+# CA5359: Do Not Disable Certificate Validation
+dotnet_diagnostic.CA5359.severity = warning
+
+# CA5360: Do Not Call Dangerous Methods In Deserialization
+dotnet_diagnostic.CA5360.severity = warning
+
+# CA5361: Do Not Disable SChannel Use of Strong Crypto
+dotnet_diagnostic.CA5361.severity = warning
+
+# CA5362: Potential reference cycle in deserialized object graph
+dotnet_diagnostic.CA5362.severity = warning
+
+# CA5363: Do Not Disable Request Validation
+dotnet_diagnostic.CA5363.severity = warning
+
+# CA5364: Do Not Use Deprecated Security Protocols
+dotnet_diagnostic.CA5364.severity = warning
+
+# CA5365: Do Not Disable HTTP Header Checking
+dotnet_diagnostic.CA5365.severity = warning
+
+# CA5366: Use XmlReader for 'DataSet.ReadXml()'
+dotnet_diagnostic.CA5366.severity = warning
+
+# CA5367: Do Not Serialize Types With Pointer Fields
+dotnet_diagnostic.CA5367.severity = warning
+
+# CA5368: Set ViewStateUserKey For Classes Derived From Page
+dotnet_diagnostic.CA5368.severity = warning
+
+# CA5369: Use XmlReader for 'XmlSerializer.Deserialize()'
+dotnet_diagnostic.CA5369.severity = warning
+
+# CA5370: Use XmlReader for XmlValidatingReader constructor
+dotnet_diagnostic.CA5370.severity = warning
+
+# CA5371: Use XmlReader for 'XmlSchema.Read()'
+dotnet_diagnostic.CA5371.severity = warning
+
+# CA5372: Use XmlReader for XPathDocument constructor
+dotnet_diagnostic.CA5372.severity = warning
+
+# CA5373: Do not use obsolete key derivation function
+dotnet_diagnostic.CA5373.severity = warning
+
+# CA5374: Do Not Use XslTransform
+dotnet_diagnostic.CA5374.severity = warning
+
+# CA5375: Do Not Use Account Shared Access Signature
+dotnet_diagnostic.CA5375.severity = warning
+
+# CA5376: Use SharedAccessProtocol HttpsOnly
+dotnet_diagnostic.CA5376.severity = warning
+
+# CA5377: Use Container Level Access Policy
+dotnet_diagnostic.CA5377.severity = warning
+
+# CA5378: Do not disable ServicePointManagerSecurityProtocols
+dotnet_diagnostic.CA5378.severity = warning
+
+# CA5379: Ensure Key Derivation Function algorithm is sufficiently strong
+dotnet_diagnostic.CA5379.severity = warning
+
+# CA5380: Do Not Add Certificates To Root Store
+dotnet_diagnostic.CA5380.severity = warning
+
+# CA5381: Ensure Certificates Are Not Added To Root Store
+dotnet_diagnostic.CA5381.severity = warning
+
+# CA5382: Use Secure Cookies In ASP.NET Core
+dotnet_diagnostic.CA5382.severity = warning
+
+# CA5383: Ensure Use Secure Cookies In ASP.NET Core
+dotnet_diagnostic.CA5383.severity = warning
+
+# CA5384: Do Not Use Digital Signature Algorithm (DSA)
+dotnet_diagnostic.CA5384.severity = warning
+
+# CA5385: Use Rivest–Shamir–Adleman (RSA) Algorithm With Sufficient Key Size
+dotnet_diagnostic.CA5385.severity = warning
+
+# CA5386: Avoid hardcoding SecurityProtocolType value
+dotnet_diagnostic.CA5386.severity = warning
+
+# CA5387: Do Not Use Weak Key Derivation Function With Insufficient Iteration Count
+dotnet_diagnostic.CA5387.severity = warning
+
+# CA5388: Ensure Sufficient Iteration Count When Using Weak Key Derivation Function
+dotnet_diagnostic.CA5388.severity = warning
+
+# CA5389: Do Not Add Archive Item's Path To The Target File System Path
+dotnet_diagnostic.CA5389.severity = warning
+
+# CA5390: Do not hard-code encryption key
+dotnet_diagnostic.CA5390.severity = warning
+
+# CA5391: Use antiforgery tokens in ASP.NET Core MVC controllers
+dotnet_diagnostic.CA5391.severity = warning
+
+# CA5392: Use DefaultDllImportSearchPaths attribute for P/Invokes
+dotnet_diagnostic.CA5392.severity = none
+
+# CA5393: Do not use unsafe DllImportSearchPath value
+dotnet_diagnostic.CA5393.severity = warning
+
+# CA5394: Do not use insecure randomness
+dotnet_diagnostic.CA5394.severity = none
+
+# CA5395: Miss HttpVerb attribute for action methods
+dotnet_diagnostic.CA5395.severity = warning
+
+# CA5396: Set HttpOnly to true for HttpCookie
+dotnet_diagnostic.CA5396.severity = warning
+
+# CA5397: Do not use deprecated SslProtocols values
+dotnet_diagnostic.CA5397.severity = warning
+
+# CA5398: Avoid hardcoded SslProtocols values
+dotnet_diagnostic.CA5398.severity = warning
+
+# CA5399: HttpClients should enable certificate revocation list checks
+dotnet_diagnostic.CA5399.severity = warning
+
+# CA5400: Ensure HttpClient certificate revocation list check is not disabled
+dotnet_diagnostic.CA5400.severity = warning
+
+# CA5401: Do not use CreateEncryptor with non-default IV
+dotnet_diagnostic.CA5401.severity = warning
+
+# CA5402: Use CreateEncryptor with the default IV 
+dotnet_diagnostic.CA5402.severity = warning
+
+# CA5403: Do not hard-code certificate
+dotnet_diagnostic.CA5403.severity = warning
+
+# CA5404: Do not disable token validation checks
+dotnet_diagnostic.CA5404.severity = none
+
+# CA5405: Do not always skip token validation in delegates
+dotnet_diagnostic.CA5405.severity = warning
diff --git a/Engine/Source/Programs/Horde/Horde.Build/Agents/ITaskSource.cs b/Engine/Source/Programs/Horde/Horde.Build/Agents/ITaskSource.cs
--- a/Engine/Source/Programs/Horde/Horde.Build/Agents/ITaskSource.cs
+++ b/Engine/Source/Programs/Horde/Horde.Build/Agents/ITaskSource.cs
@@ -149,7 +149,9 @@ namespace HordeServer.Tasks
 			{
 				Arguments[Idx + 2] = OnLeaseStartedProperties.Accessors[Idx](Payload);
 			}
+#pragma warning disable CA2254 // Template should be a static expression
 			Logger.LogInformation($"Lease {{LeaseId}} started (Type={{Type}}{OnLeaseStartedProperties.FormatString})", Arguments);
+#pragma warning restore CA2254 // Template should be a static expression
 			return Task.CompletedTask;
 		}
 
diff --git a/Engine/Source/Programs/Horde/Horde.Build/Agents/Impl/AgentService.cs b/Engine/Source/Programs/Horde/Horde.Build/Agents/Impl/AgentService.cs
--- a/Engine/Source/Programs/Horde/Horde.Build/Agents/Impl/AgentService.cs
+++ b/Engine/Source/Programs/Horde/Horde.Build/Agents/Impl/AgentService.cs
@@ -361,15 +361,14 @@ namespace HordeServer.Services
 
 						for (int Idx = 0; Idx < Tasks.Count; Idx++)
 						{
-							Task<(ITaskSource, AgentLease)?> Task = Tasks[Idx];
-							if (!Task.IsCompleted)
+							(ITaskSource, AgentLease)? TaskResult;
+							if (!Tasks[Idx].TryGetResult(out TaskResult))
 							{
 								continue;
 							}
 
 							Tasks.RemoveAt(Idx--);
 
-							(ITaskSource, AgentLease)? TaskResult = Task.Result;
 							if (!TaskResult.HasValue)
 							{
 								continue;
diff --git a/Engine/Source/Programs/Horde/Horde.Build/Agents/Impl/AgentSoftwareCollection.cs b/Engine/Source/Programs/Horde/Horde.Build/Agents/Impl/AgentSoftwareCollection.cs
--- a/Engine/Source/Programs/Horde/Horde.Build/Agents/Impl/AgentSoftwareCollection.cs
+++ b/Engine/Source/Programs/Horde/Horde.Build/Agents/Impl/AgentSoftwareCollection.cs
@@ -134,7 +134,7 @@ namespace HordeServer.Collections.Impl
             {
                 for (; ; )
                 {
-					int BytesRead = ChunkStream.Read(Buffer, 0, ChunkSize);
+					int BytesRead = await ChunkStream.ReadAsync(Buffer, 0, ChunkSize);
 					Chunks.Add(new ChunkDocument(ObjectId.GenerateNewId().ToString(), Version, Buffer.Take(BytesRead).ToArray()));
 					TotalRead += BytesRead;
 					if (TotalRead == Data.Length)
diff --git a/Engine/Source/Programs/Horde/Horde.Build/Agents/Impl/AgentSoftwareService.cs b/Engine/Source/Programs/Horde/Horde.Build/Agents/Impl/AgentSoftwareService.cs
--- a/Engine/Source/Programs/Horde/Horde.Build/Agents/Impl/AgentSoftwareService.cs
+++ b/Engine/Source/Programs/Horde/Horde.Build/Agents/Impl/AgentSoftwareService.cs
@@ -339,9 +339,9 @@ namespace HordeServer.Services
 				DirectoryReference.CreateDirectory(AgentHashFile.Directory);
 			}
 
-			FileReference.WriteAllText(AgentHashFile, AgentHash);
+			await FileReference.WriteAllTextAsync(AgentHashFile, AgentHash);
 
-			Logger.LogInformation($"Updated default agent software to {AgentHash}");
+			Logger.LogInformation("Updated default agent software to {AgentHash}", AgentHash);
 
 		}
 	}
diff --git a/Engine/Source/Programs/Horde/Horde.Build/Api/StreamTabs.cs b/Engine/Source/Programs/Horde/Horde.Build/Api/StreamTabs.cs
--- a/Engine/Source/Programs/Horde/Horde.Build/Api/StreamTabs.cs
+++ b/Engine/Source/Programs/Horde/Horde.Build/Api/StreamTabs.cs
@@ -133,7 +133,7 @@ namespace HordeServer.Api
 		/// Constructor
 		/// </summary>
 		/// <param name="Title">Title of this page</param>
-		public GetStreamTabResponse(string Title)//, StreamPageType Type)
+		protected GetStreamTabResponse(string Title)//, StreamPageType Type)
 		{
 			this.Title = Title;
 		}
@@ -160,7 +160,7 @@ namespace HordeServer.Api
 		/// </summary>
 		/// <param name="Heading">Heading for this column</param>
 		/// <param name="RelativeWidth">Relative width of this column</param>
-		public GetJobsTabColumnResponse(string Heading, int? RelativeWidth)
+		protected GetJobsTabColumnResponse(string Heading, int? RelativeWidth)
 		{
 			this.Heading = Heading;
 			this.RelativeWidth = RelativeWidth;
diff --git a/Engine/Source/Programs/Horde/Horde.Build/Auditing/Impl/AuditLog.cs b/Engine/Source/Programs/Horde/Horde.Build/Auditing/Impl/AuditLog.cs
--- a/Engine/Source/Programs/Horde/Horde.Build/Auditing/Impl/AuditLog.cs
+++ b/Engine/Source/Programs/Horde/Horde.Build/Auditing/Impl/AuditLog.cs
@@ -85,10 +85,12 @@ namespace HordeServer.Collections.Impl
 				AuditLogMessage Message = new AuditLogMessage(Subject, Time, LogLevel, Data);
 				Outer.MessageChannel.Writer.TryWrite(Message);
 
+#pragma warning disable CA2254 // Template should be a static expression
 				using (IDisposable _ = Outer.Logger.BeginScope($"Subject: {{{Outer.SubjectProperty}}}", Subject))
 				{
 					Outer.Logger.Log(LogLevel, EventId, State, Exception, Formatter);
 				}
+#pragma warning restore CA2254 // Template should be a static expression
 			}
 
 			public IAsyncEnumerable<IAuditLogMessage> FindAsync(DateTime? MinTime, DateTime? MaxTime, int? Index, int? Count) => Outer.FindAsync(Subject, MinTime, MaxTime, Index, Count);
diff --git a/Engine/Source/Programs/Horde/Horde.Build/Controllers/DashboardController.cs b/Engine/Source/Programs/Horde/Horde.Build/Controllers/DashboardController.cs
--- a/Engine/Source/Programs/Horde/Horde.Build/Controllers/DashboardController.cs
+++ b/Engine/Source/Programs/Horde/Horde.Build/Controllers/DashboardController.cs
@@ -12,72 +12,75 @@ using HordeServer;
 using HordeServer.Controllers;
 using Microsoft.Extensions.Options;
 
-/// <summary>	
-/// Dashboard authorization challenge controller	
-/// </summary>	
-[ApiController]
-[Route("[controller]")]
-public class DashboardChallengeController : Controller
+namespace Horde.Build
 {
-	/// <summary>
-	/// Authentication scheme in use
-	/// </summary>
-	string AuthenticationScheme;
-
-	/// <summary>
-	/// Constructor
-	/// </summary>
-	/// <param name="ServerSettings">Server settings</param>
-	public DashboardChallengeController(IOptionsMonitor<ServerSettings> ServerSettings)
-	{
-		AuthenticationScheme = AccountController.GetAuthScheme(ServerSettings.CurrentValue.AuthMethod);
-	}
-	
 	/// <summary>	
-	/// Challenge endpoint for the dashboard, using cookie authentication scheme	
+	/// Dashboard authorization challenge controller	
 	/// </summary>	
-	/// <returns>Ok on authorized, otherwise will 401</returns>	
-	[HttpGet]
-	[Authorize]
-	[Route("/api/v1/dashboard/challenge")]
-	public StatusCodeResult GetChallenge()
+	[ApiController]
+	[Route("[controller]")]
+	public class DashboardChallengeController : Controller
 	{
-		return Ok();
-	}
+		/// <summary>
+		/// Authentication scheme in use
+		/// </summary>
+		string AuthenticationScheme;
 
-	/// <summary>
-	/// Login to server, redirecting to the specified URL on success
-	/// </summary>
-	/// <param name="Redirect"></param>
-	/// <returns></returns>
-	[HttpGet]
-	[Route("/api/v1/dashboard/login")]
-	public IActionResult Login([FromQuery] string? Redirect)
-	{
-		return new ChallengeResult(AuthenticationScheme, new AuthenticationProperties { RedirectUri = Redirect ?? "/" });
-	}
+		/// <summary>
+		/// Constructor
+		/// </summary>
+		/// <param name="ServerSettings">Server settings</param>
+		public DashboardChallengeController(IOptionsMonitor<ServerSettings> ServerSettings)
+		{
+			AuthenticationScheme = AccountController.GetAuthScheme(ServerSettings.CurrentValue.AuthMethod);
+		}
 
-	/// <summary>
-	/// Logout of the current account
-	/// </summary>
-	/// <returns></returns>
-	[HttpGet]
-	[Route("/api/v1/dashboard/logout")]
-	public async Task<StatusCodeResult> Logout()
-	{
-		await HttpContext.SignOutAsync(CookieAuthenticationDefaults.AuthenticationScheme);
-		try
+		/// <summary>	
+		/// Challenge endpoint for the dashboard, using cookie authentication scheme	
+		/// </summary>	
+		/// <returns>Ok on authorized, otherwise will 401</returns>	
+		[HttpGet]
+		[Authorize]
+		[Route("/api/v1/dashboard/challenge")]
+		public StatusCodeResult GetChallenge()
 		{
-			await HttpContext.SignOutAsync(AuthenticationScheme);
+			return Ok();
 		}
-#pragma warning disable CA1031 // Do not catch general exception types
-		catch
-#pragma warning restore CA1031 
+
+		/// <summary>
+		/// Login to server, redirecting to the specified URL on success
+		/// </summary>
+		/// <param name="Redirect"></param>
+		/// <returns></returns>
+		[HttpGet]
+		[Route("/api/v1/dashboard/login")]
+		public IActionResult Login([FromQuery] string? Redirect)
 		{
+			return new ChallengeResult(AuthenticationScheme, new AuthenticationProperties { RedirectUri = Redirect ?? "/" });
 		}
 
-		return Ok();
-	}
+		/// <summary>
+		/// Logout of the current account
+		/// </summary>
+		/// <returns></returns>
+		[HttpGet]
+		[Route("/api/v1/dashboard/logout")]
+		public async Task<StatusCodeResult> Logout()
+		{
+			await HttpContext.SignOutAsync(CookieAuthenticationDefaults.AuthenticationScheme);
+			try
+			{
+				await HttpContext.SignOutAsync(AuthenticationScheme);
+			}
+#pragma warning disable CA1031 // Do not catch general exception types
+			catch
+#pragma warning restore CA1031
+			{
+			}
+
+			return Ok();
+		}
 
 
-}
\ No newline at end of file
+	}
+}
diff --git a/Engine/Source/Programs/Horde/Horde.Build/Controllers/DevicesController.cs b/Engine/Source/Programs/Horde/Horde.Build/Controllers/DevicesController.cs
--- a/Engine/Source/Programs/Horde/Horde.Build/Controllers/DevicesController.cs
+++ b/Engine/Source/Programs/Horde/Horde.Build/Controllers/DevicesController.cs
@@ -706,7 +706,6 @@ namespace HordeServer.Controllers
 			List<IDevicePool> Pools = await DeviceService.GetPoolsAsync();
 			List<IDevicePlatform> Platforms = await DeviceService.GetPlatformsAsync();
 
-			string Message;
 			string? PoolId = Request.PoolId;
 
 			// @todo: Remove this once all streams are updated to provide jobid
@@ -725,8 +724,8 @@ namespace HordeServer.Controllers
 
 			if (string.IsNullOrEmpty(PoolId))
 			{
-				Message = $"No pool specified, defaulting to UE4" + Details;
-				Logger.LogError(Message + $" JobId: {Request.JobId}, StepId: {Request.StepId}");
+				Logger.LogError("No pool specified, defaulting to UE4 {Details} JobId: {JobId}, StepId: {StepId}", Details, Request.JobId, Request.StepId);
+				string Message = $"No pool specified, defaulting to UE4" + Details;
 				await DeviceService.NotifyDeviceServiceAsync(Message, null, Request.JobId, Request.StepId);
 				PoolId = "ue4";
 				//return BadRequest(Message);
@@ -736,8 +735,8 @@ namespace HordeServer.Controllers
 			IDevicePool? Pool = Pools.FirstOrDefault(x => x.Id == PoolIdValue);
 			if (Pool == null)
 			{
-				Message = $"Unknown pool {PoolId} " + Details;
-				Logger.LogError(Message);
+				Logger.LogError("Unknown pool {PoolId} {Details}", PoolId, Details);
+				string Message = $"Unknown pool {PoolId} " + Details;
 				await DeviceService.NotifyDeviceServiceAsync(Message, null, Request.JobId, Request.StepId);
 				return BadRequest(Message);
 			}
@@ -752,7 +751,7 @@ namespace HordeServer.Controllers
 				string PlatformName = DeviceType;
 				string PerfSpecName = "Unspecified";
 
-				if (DeviceType.Contains(":", StringComparison.OrdinalIgnoreCase))
+				if (DeviceType.Contains(':', StringComparison.Ordinal))
 				{
 					string[] Tokens = DeviceType.Split(":");
 					PlatformName = Tokens[0];
@@ -767,8 +766,8 @@ namespace HordeServer.Controllers
 
 				if (!MapV1.PlatformMap.TryGetValue(PlatformName, out PlatformId))
 				{
-					Message = $"Unknown platform {PlatformName}" + Details;
-					Logger.LogError(Message);
+					string Message = $"Unknown platform {PlatformName}" + Details;
+					Logger.LogError("Unknown platform {PlatformName} {Details}", PlatformName, Details);
 					await DeviceService.NotifyDeviceServiceAsync(Message, null, Request.JobId, Request.StepId);
 					return BadRequest(Message);
 				}
@@ -777,8 +776,8 @@ namespace HordeServer.Controllers
 
 				if (Platform == null)
 				{
-					Message = $"Unknown platform {PlatformId}" + Details;
-					Logger.LogError(Message);
+					string Message = $"Unknown platform {PlatformId}" + Details;
+					Logger.LogError("Unknown platform {PlatformId} {Details}", PlatformId, Details);
 					await DeviceService.NotifyDeviceServiceAsync(Message, null, Request.JobId, Request.StepId);
 					return BadRequest(Message);
 				}
@@ -843,7 +842,7 @@ namespace HordeServer.Controllers
 
 			if (Reservation == null)
 			{
-				Logger.LogError($"Unable to find reservation for legacy guid {ReservationGuid}");
+				Logger.LogError("Unable to find reservation for legacy guid {ReservationGuid}", ReservationGuid);
 				return BadRequest();
 			}
 
@@ -851,7 +850,7 @@ namespace HordeServer.Controllers
 
 			if (!Updated)
 			{
-				Logger.LogError($"Unable to find reservation for reservation {Reservation.Id}");
+				Logger.LogError("Unable to find reservation for reservation {ReservationId}", Reservation.Id);
 				return BadRequest();
 			}
 
@@ -978,7 +977,7 @@ namespace HordeServer.Controllers
 
 			if (Device == null)
 			{
-				Logger.LogError($"Device error reported for unknown device {DeviceName}");
+				Logger.LogError("Device error reported for unknown device {DeviceName}", DeviceName);
 				return BadRequest($"Unknown device {DeviceName}");
 			}
 
@@ -1008,7 +1007,7 @@ namespace HordeServer.Controllers
 				}
 			}
 
-			Logger.LogError(Message);
+			Logger.LogError("{DeviceError}", Message);
 
 			await DeviceService.NotifyDeviceServiceAsync(Message, Device.Id, JobId, StepId);
 
diff --git a/Engine/Source/Programs/Horde/Horde.Build/Controllers/LogsController.cs b/Engine/Source/Programs/Horde/Horde.Build/Controllers/LogsController.cs
--- a/Engine/Source/Programs/Horde/Horde.Build/Controllers/LogsController.cs
+++ b/Engine/Source/Programs/Horde/Horde.Build/Controllers/LogsController.cs
@@ -221,7 +221,7 @@ namespace HordeServer.Controllers
 							}
 							else if (Result[Offset] == (byte)'\n')
 							{
-								Stream.Write(Result, StartOffset, Offset - StartOffset);
+								await Stream.WriteAsync(Result.AsMemory(StartOffset, Offset - StartOffset));
 								Offset++;
 								break;
 							}
diff --git a/Engine/Source/Programs/Horde/Horde.Build/Devices/Impl/DeviceService.cs b/Engine/Source/Programs/Horde/Horde.Build/Devices/Impl/DeviceService.cs
--- a/Engine/Source/Programs/Horde/Horde.Build/Devices/Impl/DeviceService.cs
+++ b/Engine/Source/Programs/Horde/Horde.Build/Devices/Impl/DeviceService.cs
@@ -404,7 +404,7 @@ namespace HordeServer.Services
 			}
 			catch (Exception Ex)
 			{
-                Logger.LogError($"Error on device notification {Ex.Message}");
+                Logger.LogError(Ex, "Error on device notification {Message}", Ex.Message);
             }
         }
 
@@ -454,7 +454,7 @@ namespace HordeServer.Services
 
 				if (Projects.Where(x => x.Id == ProjectId).FirstOrDefault() == null)
 				{
-					Logger.LogWarning($"Device pool authorization references missing project id {ProjectId}");
+					Logger.LogWarning("Device pool authorization references missing project id {ProjectId}", ProjectId);
 					continue;
 				}
 
diff --git a/Engine/Source/Programs/Horde/Horde.Build/Horde.Build.csproj b/Engine/Source/Programs/Horde/Horde.Build/Horde.Build.csproj
--- a/Engine/Source/Programs/Horde/Horde.Build/Horde.Build.csproj
+++ b/Engine/Source/Programs/Horde/Horde.Build/Horde.Build.csproj
@@ -7,6 +7,7 @@
     <AssemblyName>Horde.Build</AssemblyName>
     <DocumentationFile>Horde.Build.xml</DocumentationFile>
     <ApplicationIcon>HordeServer.ico</ApplicationIcon>
+    <RunAnalyzersDuringBuild>False</RunAnalyzersDuringBuild>
   </PropertyGroup>
 
   <ItemGroup>
@@ -25,14 +26,10 @@
     </PackageReference>
     <PackageReference Include="JetBrains.Profiler.SelfApi" Version="2.1.2" />
     <PackageReference Include="Microsoft.AspNetCore.Authentication.Cookies" Version="2.2.0" />
-    <PackageReference Include="Microsoft.AspNetCore.Authentication.JwtBearer" Version="3.1.5" />
+    <PackageReference Include="Microsoft.AspNetCore.Authentication.JwtBearer" Version="6.0.0" />
     <PackageReference Include="Microsoft.AspNetCore.Authentication.OpenIdConnect" Version="3.1.5" />
     <PackageReference Include="Microsoft.AspNetCore.DataProtection.StackExchangeRedis" Version="3.1.7" />
     <PackageReference Include="Microsoft.AspNetCore.SpaServices.Extensions" Version="3.1.17" />
-    <PackageReference Include="Microsoft.CodeAnalysis.FxCopAnalyzers" Version="3.0.0">
-      <PrivateAssets>all</PrivateAssets>
-      <IncludeAssets>runtime; build; native; contentfiles; analyzers; buildtransitive</IncludeAssets>
-    </PackageReference>
     <PackageReference Include="Microsoft.Extensions.Caching.Memory" Version="5.0.0" />
     <PackageReference Include="MongoDB.Driver" Version="2.11.6" />
     <PackageReference Include="OpenTracing" Version="0.12.1" />
diff --git a/Engine/Source/Programs/Horde/Horde.Build/Issues/Impl/IssueService.cs b/Engine/Source/Programs/Horde/Horde.Build/Issues/Impl/IssueService.cs
--- a/Engine/Source/Programs/Horde/Horde.Build/Issues/Impl/IssueService.cs
+++ b/Engine/Source/Programs/Horde/Horde.Build/Issues/Impl/IssueService.cs
@@ -212,11 +212,8 @@ namespace HordeServer.Services.Impl
 			Type[] MatcherTypes = Assembly.GetExecutingAssembly().GetTypes().Where(x => !x.IsAbstract && typeof(IIssueHandler).IsAssignableFrom(x)).ToArray();
 			foreach (Type MatcherType in MatcherTypes)
 			{
-				IIssueHandler? Matcher = (IIssueHandler?)Activator.CreateInstance(MatcherType);
-				if (Matcher != null)
-				{
-					Matchers.Add(Matcher);
-				}
+				IIssueHandler Matcher = (IIssueHandler)Activator.CreateInstance(MatcherType)!;
+				Matchers.Add(Matcher);
 			}
 			Matchers.SortBy(x => -x.Priority);
 
diff --git a/Engine/Source/Programs/Horde/Horde.Build/Jobs/Impl/JobTaskSource.cs b/Engine/Source/Programs/Horde/Horde.Build/Jobs/Impl/JobTaskSource.cs
--- a/Engine/Source/Programs/Horde/Horde.Build/Jobs/Impl/JobTaskSource.cs
+++ b/Engine/Source/Programs/Horde/Horde.Build/Jobs/Impl/JobTaskSource.cs
@@ -507,6 +507,7 @@ namespace HordeServer.Tasks.Impl
 		/// </summary>
 		/// <param name="Item">The queue item</param>
 		/// <returns></returns>
+		[SuppressMessage("Maintainability", "CA1508:Avoid dead conditional code", Justification = "<Pending>")]
 		void AssignQueueItemToAnyWaiter(QueueItem Item)
 		{
 			if (Item.AssignTask == null && Item.Batch.SessionId == null)
@@ -630,9 +631,9 @@ namespace HordeServer.Tasks.Impl
 			lock (LockObject)
 			{
 				AssignAnyQueueItemToWaiter(Waiter);
-				if (Waiter.LeaseSource.Task.IsCompleted)
+				if (Waiter.LeaseSource.Task.TryGetResult(out AgentLease? Result))
 				{
-					return Waiter.Task.Result;
+					return Result;
 				}
 				Waiters.Add(Waiter);
 			}
diff --git a/Engine/Source/Programs/Horde/Horde.Build/Jobs/SubResourceId.cs b/Engine/Source/Programs/Horde/Horde.Build/Jobs/SubResourceId.cs
--- a/Engine/Source/Programs/Horde/Horde.Build/Jobs/SubResourceId.cs
+++ b/Engine/Source/Programs/Horde/Horde.Build/Jobs/SubResourceId.cs
@@ -77,8 +77,7 @@ namespace HordeServer.Models
 		/// <inheritdoc/>
 		public override bool Equals(object? Other)
 		{
-			SubResourceId? OtherId = Other as SubResourceId?;
-			return OtherId != null && Equals(OtherId.Value);
+			return Other is SubResourceId OtherId && Equals(OtherId);
 		}
 
 		/// <inheritdoc/>
diff --git a/Engine/Source/Programs/Horde/Horde.Build/Logs/Impl/LogFileService.cs b/Engine/Source/Programs/Horde/Horde.Build/Logs/Impl/LogFileService.cs
--- a/Engine/Source/Programs/Horde/Horde.Build/Logs/Impl/LogFileService.cs
+++ b/Engine/Source/Programs/Horde/Horde.Build/Logs/Impl/LogFileService.cs
@@ -228,7 +228,7 @@ namespace HordeServer.Services
 				while(Length > 0)
 				{
 					// Add more data to the buffer
-					int ReadBytes = await Stream.ReadAsync(ReadBuffer, ReadBufferLength, ReadBuffer.Length - ReadBufferLength);
+					int ReadBytes = await Stream.ReadAsync(ReadBuffer.AsMemory(ReadBufferLength, ReadBuffer.Length - ReadBufferLength));
 					ReadBufferLength += ReadBytes;
 
 					// Copy as many lines as possible to the output
@@ -248,7 +248,7 @@ namespace HordeServer.Services
 						if (Offset < WriteBufferLength)
 						{
 							int WriteLength = (int)Math.Min((long)WriteBufferLength - Offset, Length);
-							await OutputStream.WriteAsync(WriteBuffer, (int)Offset, WriteLength);
+							await OutputStream.WriteAsync(WriteBuffer.AsMemory((int)Offset, WriteLength));
 							Length -= WriteLength;
 						}
 						Offset = Math.Max(Offset - WriteBufferLength, 0);
@@ -433,6 +433,12 @@ namespace HordeServer.Services
 
 			/// <inheritdoc/>
 			public override async Task<int> ReadAsync(byte[] Buffer, int Offset, int Length, CancellationToken CancellationToken)
+			{
+				return await ReadAsync(Buffer.AsMemory(Offset, Length), CancellationToken);
+			}
+
+			/// <inheritdoc/>
+			public override async ValueTask<int> ReadAsync(Memory<byte> Buffer, CancellationToken CancellationToken)
 			{
 				int ReadBytes = 0;
 				while (ReadBytes < Length)
@@ -440,8 +446,8 @@ namespace HordeServer.Services
 					if (SourcePos < SourceEnd)
 					{
 						// Try to copy from the current buffer
-						int BlockSize = Math.Min(SourceEnd - SourcePos, Length - ReadBytes);
-						SourceBuffer.Slice(SourcePos, BlockSize).Span.CopyTo(Buffer.AsSpan(Offset + ReadBytes));
+						int BlockSize = Math.Min(SourceEnd - SourcePos, Buffer.Length - ReadBytes);
+						SourceBuffer.Slice(SourcePos, BlockSize).Span.CopyTo(Buffer.Slice(ReadBytes).Span);
 						CurrentOffset += BlockSize;
 						ReadBytes += BlockSize;
 						SourcePos += BlockSize;
diff --git a/Engine/Source/Programs/Horde/Horde.Build/Notifications/Impl/SlackNotificationSink.cs b/Engine/Source/Programs/Horde/Horde.Build/Notifications/Impl/SlackNotificationSink.cs
--- a/Engine/Source/Programs/Horde/Horde.Build/Notifications/Impl/SlackNotificationSink.cs
+++ b/Engine/Source/Programs/Horde/Horde.Build/Notifications/Impl/SlackNotificationSink.cs
@@ -1290,7 +1290,7 @@ namespace HordeServer.Notifications.Impl
 			using FormUrlEncodedContent Content = new FormUrlEncodedContent(Array.Empty<KeyValuePair<string?, string?>>());
 			HttpResponseMessage Response = await Client.PostAsync(new Uri("https://slack.com/api/apps.connections.open"), Content, StoppingToken);
 
-			byte[] ResponseData = await Response.Content.ReadAsByteArrayAsync();
+			byte[] ResponseData = await Response.Content.ReadAsByteArrayAsync(StoppingToken);
 
 			SocketResponse SocketResponse = JsonSerializer.Deserialize<SocketResponse>(ResponseData)!;
 			if (!SocketResponse.Ok)
diff --git a/Engine/Source/Programs/Horde/Horde.Build/Perforce/Impl/CommitService.cs b/Engine/Source/Programs/Horde/Horde.Build/Perforce/Impl/CommitService.cs
--- a/Engine/Source/Programs/Horde/Horde.Build/Perforce/Impl/CommitService.cs
+++ b/Engine/Source/Programs/Horde/Horde.Build/Perforce/Impl/CommitService.cs
@@ -1081,7 +1081,7 @@ namespace HordeServer.Commits.Impl
 
 			foreach ((Utf8String Name, StreamFile File) in Tree.NameToFile)
 			{
-				Logger.LogInformation($"{Prefix}/{Name} = {File.Path}#{File.Revision}");
+				Logger.LogInformation("{Prefix}/{Name} = {Path}#{Revision}", Prefix, Name, File.Path, File.Revision);
 			}
 			foreach ((Utf8String Name, StreamTreeRef ChildTreeRef) in Tree.NameToTree)
 			{
diff --git a/Engine/Source/Programs/Horde/Horde.Build/Perforce/Impl/PerforceService.cs b/Engine/Source/Programs/Horde/Horde.Build/Perforce/Impl/PerforceService.cs
--- a/Engine/Source/Programs/Horde/Horde.Build/Perforce/Impl/PerforceService.cs
+++ b/Engine/Source/Programs/Horde/Horde.Build/Perforce/Impl/PerforceService.cs
@@ -1383,7 +1383,7 @@ namespace HordeServer.Services
 				}
 				catch
 				{
-					Logger.LogError($"Unable to delete client {ClientName}");
+					Logger.LogError("Unable to delete client {ClientName}", ClientName);
 				}
 
 			}
diff --git a/Engine/Source/Programs/Horde/Horde.Build/Projects/IStream.cs b/Engine/Source/Programs/Horde/Horde.Build/Projects/IStream.cs
--- a/Engine/Source/Programs/Horde/Horde.Build/Projects/IStream.cs
+++ b/Engine/Source/Programs/Horde/Horde.Build/Projects/IStream.cs
@@ -459,7 +459,7 @@ namespace HordeServer.Models
 			}
 			else if (ChangeTemplateRefId != null)
 			{
-				ChangeRequest = new ChangeQueryRequest { TemplateId = ChangeTemplateRefId?.ToString() };
+				ChangeRequest = new ChangeQueryRequest { TemplateId = ChangeTemplateRefId.ToString() };
 			}
 
 			return new DefaultPreflightRequest { TemplateId = TemplateRefId?.ToString(), Change = ChangeRequest, ChangeTemplateId = ChangeRequest?.TemplateId };
@@ -612,7 +612,7 @@ namespace HordeServer.Models
 
 				// Get the workspace identifier
 				string Identifier;
-				if (WorkspaceType.Identifier != null && !String.IsNullOrEmpty(WorkspaceType.Identifier))
+				if (WorkspaceType.Identifier != null)
 				{
 					Identifier = WorkspaceType.Identifier;
 				}
diff --git a/Engine/Source/Programs/Horde/Horde.Build/Projects/StreamTab.cs b/Engine/Source/Programs/Horde/Horde.Build/Projects/StreamTab.cs
--- a/Engine/Source/Programs/Horde/Horde.Build/Projects/StreamTab.cs
+++ b/Engine/Source/Programs/Horde/Horde.Build/Projects/StreamTab.cs
@@ -36,7 +36,7 @@ namespace HordeServer.Models
 		/// Constructor
 		/// </summary>
 		/// <param name="Title">Title of this page</param>
-		public StreamTab(string Title)
+		protected StreamTab(string Title)
 		{
 			this.Title = Title;
 		}
@@ -45,7 +45,7 @@ namespace HordeServer.Models
 		/// Constructor
 		/// </summary>
 		/// <param name="Request">Public request object</param>
-		public StreamTab(CreateStreamTabRequest Request)
+		protected StreamTab(CreateStreamTabRequest Request)
 		{
 			Title = Request.Title;
 		}
@@ -103,7 +103,7 @@ namespace HordeServer.Models
 		/// </summary>
 		/// <param name="Heading">Heading for this column</param>
 		/// <param name="RelativeWidth">Relative width of this column.</param>
-		public JobsTabColumn(string Heading, int? RelativeWidth)
+		protected JobsTabColumn(string Heading, int? RelativeWidth)
 		{
 			this.Heading = Heading;
 			this.RelativeWidth = RelativeWidth;
@@ -113,7 +113,7 @@ namespace HordeServer.Models
 		/// Constructor
 		/// </summary>
 		/// <param name="Request">Public request object</param>
-		public JobsTabColumn(CreateJobsTabColumnRequest Request)
+		protected JobsTabColumn(CreateJobsTabColumnRequest Request)
 		{
 			this.Heading = Request.Heading;
 			this.RelativeWidth = Request.RelativeWidth;
diff --git a/Engine/Source/Programs/Horde/Horde.Build/Server/ConfigService.cs b/Engine/Source/Programs/Horde/Horde.Build/Server/ConfigService.cs
--- a/Engine/Source/Programs/Horde/Horde.Build/Server/ConfigService.cs
+++ b/Engine/Source/Programs/Horde/Horde.Build/Server/ConfigService.cs
@@ -517,12 +517,12 @@ namespace HordeServer.Services
 				// write out the projects
 				foreach (KeyValuePair<string, string> Project in Request.ProjectsJson)
 				{
-					FileReference.WriteAllText(FileReference.Combine(GlobalConfigDirectory, Project.Key), Project.Value);
+					await FileReference.WriteAllTextAsync(FileReference.Combine(GlobalConfigDirectory, Project.Key), Project.Value);
 
 					if (Request.ProjectLogo != null)
 					{
 						Byte[] bytes = Convert.FromBase64String(Request.ProjectLogo);
-						FileReference.WriteAllBytes(FileReference.Combine(GlobalConfigDirectory, Project.Key.Replace(".json", ".png", StringComparison.OrdinalIgnoreCase)), bytes);
+						await FileReference.WriteAllBytesAsync(FileReference.Combine(GlobalConfigDirectory, Project.Key.Replace(".json", ".png", StringComparison.OrdinalIgnoreCase)), bytes);
 					}
 				}
 
@@ -536,7 +536,7 @@ namespace HordeServer.Services
 				// write out the streams
 				foreach (KeyValuePair<string, string> Stream in Request.StreamsJson)
 				{
-					FileReference.WriteAllText(FileReference.Combine(GlobalConfigDirectory, Stream.Key), Stream.Value);
+					await FileReference.WriteAllTextAsync(FileReference.Combine(GlobalConfigDirectory, Stream.Key), Stream.Value);
 				}
 			}
 
@@ -545,7 +545,7 @@ namespace HordeServer.Services
 			{
 				ConfigDirty = true;
 
-				FileReference.WriteAllText(GlobalConfigFile, Request.GlobalsJson);
+				await FileReference.WriteAllTextAsync(GlobalConfigFile, Request.GlobalsJson);
 			}
 
 			if (ConfigDirty == true)
@@ -634,14 +634,14 @@ namespace HordeServer.Services
 					if (Property == null)
 					{
 						Response.Errors.Add($"Horde configuration property {Pair.Key} does not exist when reading server settings");
-						Logger.LogError(Response.Errors.Last());
+						Logger.LogError("Horde configuration property {Key} does not exist when reading server settings", Pair.Key);
 						continue;
 					}
 
 					if (Property.SetMethod == null)
 					{
 						Response.Errors.Add($"Horde configuration property {Pair.Key} does not have a set method");
-						Logger.LogError(Response.Errors.Last());
+						Logger.LogError("Horde configuration property {Key} does not have a set method", Pair.Key);
 						continue;
 					}
 
@@ -725,7 +725,7 @@ namespace HordeServer.Services
 					if (Property == null)
 					{
 						Response.Errors.Add($"Horde configuration property {Name} does not exist when writing server settings");
-						Logger.LogError(Response.Errors.Last());
+						Logger.LogError("Horde configuration property {Name} does not exist when writing server settings", Name);
 
 						continue;
 					}
@@ -733,7 +733,7 @@ namespace HordeServer.Services
 					if (Property.GetMethod == null)
 					{
 						Response.Errors.Add($"Horde configuration property {Name} does not have a get method");
-						Logger.LogError(Response.Errors.Last());
+						Logger.LogError("Horde configuration property {Name} does not have a get method", Name);
 						continue;
 					}
 
@@ -742,7 +742,7 @@ namespace HordeServer.Services
 					if (Result == null)
 					{
 						Response.Errors.Add($"Horde configuration property {Name} was null while writing and should have already been filtered out");
-						Logger.LogError(Response.Errors.Last());
+						Logger.LogError("Horde configuration property {Name} was null while writing and should have already been filtered out", Name);
 						continue;
 					}
 
diff --git a/Engine/Source/Programs/Horde/Horde.Build/Server/DatabaseService.cs b/Engine/Source/Programs/Horde/Horde.Build/Server/DatabaseService.cs
--- a/Engine/Source/Programs/Horde/Horde.Build/Server/DatabaseService.cs
+++ b/Engine/Source/Programs/Horde/Horde.Build/Server/DatabaseService.cs
@@ -234,7 +234,7 @@ namespace HordeServer.Services
 			}
 			if(MongoProcessGroup != null)
 			{
-				MongoProcessGroup?.Dispose();
+				MongoProcessGroup.Dispose();
 				MongoProcessGroup = null;
 			}
 		}
@@ -345,11 +345,11 @@ namespace HordeServer.Services
 							ChannelLogger = LoggerFactory.CreateLogger($"MongoDB.{Match.Groups[2].Value}");
 							ChannelLoggers.Add(Match.Groups[2].Value, ChannelLogger);
 						}
-						ChannelLogger.Log(ParseMongoLogLevel(Match.Groups[1].Value), Match.Groups[3].Value.TrimEnd());
+						ChannelLogger.Log(ParseMongoLogLevel(Match.Groups[1].Value), "{Message}", Match.Groups[3].Value.TrimEnd());
 					}
 					else
 					{
-						MongoLogger.Log(LogLevel.Information, Line);
+						MongoLogger.Log(LogLevel.Information, "{Message}", Line);
 					}
 				}
 			}
@@ -390,7 +390,9 @@ namespace HordeServer.Services
 				}
 			}
 
+#pragma warning disable CA2254 // Template should be a static expression
 			Logger.LogTrace($"MongoDB: {String.Join(", ", Params)}", Values.ToArray());
+#pragma warning restore CA2254 // Template should be a static expression
 		}
 
 		/// <summary>
diff --git a/Engine/Source/Programs/Horde/Horde.Build/Server/LifetimeService.cs b/Engine/Source/Programs/Horde/Horde.Build/Server/LifetimeService.cs
--- a/Engine/Source/Programs/Horde/Horde.Build/Server/LifetimeService.cs
+++ b/Engine/Source/Programs/Horde/Horde.Build/Server/LifetimeService.cs
@@ -108,9 +108,9 @@ namespace HordeServer.Services
 			StoppingTaskCompletionSource.TrySetResult(true);
 
 			int ShutdownDelayMs = 30 * 1000;
-			Logger.LogInformation($"Delaying shutdown by sleeping {ShutdownDelayMs} ms...");
+			Logger.LogInformation("Delaying shutdown by sleeping {ShutdownDelayMs} ms...", ShutdownDelayMs);
 			Thread.Sleep(ShutdownDelayMs);
-			Logger.LogInformation($"Server process now shutting down...");
+			Logger.LogInformation("Server process now shutting down...");
 			
 			/*
 			if (PreStoppingTaskCompletionSource.TrySetResult(true))
diff --git a/Engine/Source/Programs/Horde/Horde.Build/Server/RedisService.cs b/Engine/Source/Programs/Horde/Horde.Build/Server/RedisService.cs
--- a/Engine/Source/Programs/Horde/Horde.Build/Server/RedisService.cs
+++ b/Engine/Source/Programs/Horde/Horde.Build/Server/RedisService.cs
@@ -203,7 +203,7 @@ namespace HordeServer.Services
 				}
 				if (Line.Length > 0)
 				{
-					RedisLogger.Log(LogLevel.Information, Line);
+					RedisLogger.Log(LogLevel.Information, "{Output}", Line);
 				}
 			}
 			RedisLogger.LogInformation("Exit code {ExitCode}", RedisProcess.ExitCode);
diff --git a/Engine/Source/Programs/Horde/Horde.Build/Startup.cs b/Engine/Source/Programs/Horde/Horde.Build/Startup.cs
--- a/Engine/Source/Programs/Horde/Horde.Build/Startup.cs
+++ b/Engine/Source/Programs/Horde/Horde.Build/Startup.cs
@@ -159,7 +159,9 @@ namespace HordeServer
 				}
 				catch (StructuredRpcException Ex)
 				{
+#pragma warning disable CA2254 // Template should be a static expression
 					Logger.LogError(Ex, Ex.Format, Ex.Args);
+#pragma warning restore CA2254 // Template should be a static expression
 					throw;
 				}
 				catch (Exception Ex)
diff --git a/Engine/Source/Programs/Horde/Horde.Build/Users/Impl/UserCollectionV2.cs b/Engine/Source/Programs/Horde/Horde.Build/Users/Impl/UserCollectionV2.cs
--- a/Engine/Source/Programs/Horde/Horde.Build/Users/Impl/UserCollectionV2.cs
+++ b/Engine/Source/Programs/Horde/Horde.Build/Users/Impl/UserCollectionV2.cs
@@ -260,7 +260,7 @@ namespace HordeServer.Collections.Impl
 
 			if (Name == null)
 			{
-				Update = Update.SetOnInsert(x => x.Name, Name ?? Login);
+				Update = Update.SetOnInsert(x => x.Name, Login);
 			}
 			else
 			{
diff --git a/Engine/Source/Programs/Horde/Horde.Build/Utilities/ArtifactChunkStream.cs b/Engine/Source/Programs/Horde/Horde.Build/Utilities/ArtifactChunkStream.cs
--- a/Engine/Source/Programs/Horde/Horde.Build/Utilities/ArtifactChunkStream.cs
+++ b/Engine/Source/Programs/Horde/Horde.Build/Utilities/ArtifactChunkStream.cs
@@ -88,9 +88,15 @@ namespace HordeServer.Utilities
 
 		/// <inheritdoc/>
 		public override async Task<int> ReadAsync(byte[] Buffer, int Offset, int Count, CancellationToken CancellationToken)
+		{
+			return await ReadAsync(Buffer.AsMemory(Offset, Count), CancellationToken);
+		}
+
+		/// <inheritdoc/>
+		public override async ValueTask<int> ReadAsync(Memory<byte> Buffer, CancellationToken CancellationToken)
 		{
 			int BytesRead = 0;
-			while (BytesRead < Count && StreamPosition < StreamLength)
+			while (BytesRead < Buffer.Length && StreamPosition < StreamLength)
 			{
 				if (Request == null || RequestPos == Request.Data.Length)
 				{
@@ -106,8 +112,8 @@ namespace HordeServer.Utilities
 				else
 				{
 					// Copy data from the current request object
-					int NumBytesToCopy = Math.Min(Count - BytesRead, Request.Data.Length - RequestPos);
-					Request.Data.Span.Slice(RequestPos, NumBytesToCopy).CopyTo(Buffer.AsSpan(BytesRead, NumBytesToCopy));
+					int NumBytesToCopy = Math.Min(Buffer.Length - BytesRead, Request.Data.Length - RequestPos);
+					Request.Data.Span.Slice(RequestPos, NumBytesToCopy).CopyTo(Buffer.Slice(BytesRead, NumBytesToCopy).Span);
 					RequestPos += NumBytesToCopy;
 					BytesRead += NumBytesToCopy;
 					StreamPosition += NumBytesToCopy;
diff --git a/Engine/Source/Programs/Horde/Horde.Build/Utilities/AsyncCachedValue.cs b/Engine/Source/Programs/Horde/Horde.Build/Utilities/AsyncCachedValue.cs
--- a/Engine/Source/Programs/Horde/Horde.Build/Utilities/AsyncCachedValue.cs
+++ b/Engine/Source/Programs/Horde/Horde.Build/Utilities/AsyncCachedValue.cs
@@ -114,7 +114,7 @@ namespace HordeServer.Utilities
 			if (CurrentState.Next != null && CurrentState.Next.IsCompleted)
 			{
 				_ = Interlocked.CompareExchange(ref Current, CurrentState.Next, CurrentCopy);
-				return CurrentState.Next.Result.Value;
+				return (await CurrentState.Next).Value;
 			}
 
 			return CurrentState.Value;
diff --git a/Engine/Source/Programs/Horde/Horde.Build/Utilities/JwtUtils.cs b/Engine/Source/Programs/Horde/Horde.Build/Utilities/JwtUtils.cs
--- a/Engine/Source/Programs/Horde/Horde.Build/Utilities/JwtUtils.cs
+++ b/Engine/Source/Programs/Horde/Horde.Build/Utilities/JwtUtils.cs
@@ -26,7 +26,7 @@ namespace HordeServer.Utilities
 		public static bool TryGetBearerToken(HttpRequest Request, string BearerPrefix, [NotNullWhen(true)] out string? Token)
 		{
 			// Get the authorization header
-			string? Authorization = Request.Headers["Authorization"];
+			string? Authorization = Request.Headers.Authorization;
 			if (String.IsNullOrEmpty(Authorization))
 			{
 				Token = null;
diff --git a/Engine/Source/Programs/Horde/Horde.Build/Utilities/SingletonDocument.cs b/Engine/Source/Programs/Horde/Horde.Build/Utilities/SingletonDocument.cs
--- a/Engine/Source/Programs/Horde/Horde.Build/Utilities/SingletonDocument.cs
+++ b/Engine/Source/Programs/Horde/Horde.Build/Utilities/SingletonDocument.cs
@@ -15,7 +15,7 @@ namespace HordeServer.Utilities
 	/// Attribute specifying the unique id for a singleton document
 	/// </summary>
 	[AttributeUsage(AttributeTargets.Class)]
-	public class SingletonDocumentAttribute : Attribute
+	public sealed class SingletonDocumentAttribute : Attribute
 	{
 		/// <summary>
 		/// Unique id for the singleton document
diff --git a/Engine/Source/Programs/Horde/Horde.Build/Utilities/ViewMap.cs b/Engine/Source/Programs/Horde/Horde.Build/Utilities/ViewMap.cs
--- a/Engine/Source/Programs/Horde/Horde.Build/Utilities/ViewMap.cs
+++ b/Engine/Source/Programs/Horde/Horde.Build/Utilities/ViewMap.cs
@@ -313,7 +313,7 @@ namespace HordeServer.Utilities
 			StringBuilder Builder = new StringBuilder();
 			if (!Include)
 			{
-				Builder.Append("-");
+				Builder.Append('-');
 			}
 			Builder.Append(CultureInfo.InvariantCulture, $"{SourcePrefix}{Wildcard}{SourceSuffix} {TargetPrefix}{Wildcard}{TargetSuffix}");
 			return Builder.ToString();
diff --git a/Engine/Source/Programs/Shared/EpicGames.Core/AsyncUtils.cs b/Engine/Source/Programs/Shared/EpicGames.Core/AsyncUtils.cs
--- a/Engine/Source/Programs/Shared/EpicGames.Core/AsyncUtils.cs
+++ b/Engine/Source/Programs/Shared/EpicGames.Core/AsyncUtils.cs
@@ -23,6 +23,27 @@ namespace EpicGames.Core
 			return Task.Delay(-1, Token).ContinueWith(x => { });
 		}
 
+		/// <summary>
+		/// Attempts to get the result of a task, if it has finished
+		/// </summary>
+		/// <typeparam name="T"></typeparam>
+		/// <param name="Task"></param>
+		/// <param name="Result"></param>
+		/// <returns></returns>
+		public static bool TryGetResult<T>(this Task<T> Task, out T Result)
+		{
+			if (Task.IsCompleted)
+			{
+				Result = Task.Result;
+				return true;
+			}
+			else
+			{
+				Result = default!;
+				return false;
+			}
+		}
+
 		/// <summary>
 		/// Waits for a time period to elapse or the task to be cancelled, without throwing an cancellation exception
 		/// </summary>
