Horde: Fix a code analysis warning.

diff --git a/Engine/Source/Programs/Horde/Horde.Build/Perforce/PerforceLoadBalancer.cs b/Engine/Source/Programs/Horde/Horde.Build/Perforce/PerforceLoadBalancer.cs
--- a/Engine/Source/Programs/Horde/Horde.Build/Perforce/PerforceLoadBalancer.cs
+++ b/Engine/Source/Programs/Horde/Horde.Build/Perforce/PerforceLoadBalancer.cs
@@ -345,7 +345,7 @@ namespace HordeServer.Services
 			List<Task> Tasks = new List<Task>();
 			foreach (PerforceServerEntry Entry in List.Servers)
 			{
-				Tasks.Add(Task.Run(() => UpdateHealthAsync(Entry)));
+				Tasks.Add(Task.Run(() => UpdateHealthAsync(Entry), CancellationToken.None));
 			}
 			await Task.WhenAll(Tasks);
 		}
