Horde Storage - Proper fix for handling blobs not existing when resolving their references

diff --git a/Engine/Source/Programs/Horde/HordeStorage/Horde.Storage/Controllers/ObjectController.cs b/Engine/Source/Programs/Horde/HordeStorage/Horde.Storage/Controllers/ObjectController.cs
--- a/Engine/Source/Programs/Horde/HordeStorage/Horde.Storage/Controllers/ObjectController.cs
+++ b/Engine/Source/Programs/Horde/HordeStorage/Horde.Storage/Controllers/ObjectController.cs
@@ -159,7 +159,16 @@ namespace Horde.Storage.Controllers
                 return Forbid();
             }
 
-            BlobContents blob = await _storage.GetObject(ns, id);
+            BlobContents blob;
+            try
+            {
+                blob = await _storage.GetObject(ns, id);
+            }
+            catch (BlobNotFoundException e)
+            {
+                return NotFound(new ValidationProblemDetails {Title = $"Object {e.Blob} not found"});
+            }
+           
             byte[] blobContents = await blob.Stream.ToByteArray();
             CompactBinaryObject compactBinaryObject = CompactBinaryObject.Load(blobContents);
 
@@ -172,10 +181,7 @@ namespace Horde.Storage.Controllers
             {
                 return BadRequest(new ValidationProblemDetails {Title = $"Object {id} is missing blobs", Detail = $"Following blobs are missing: {string.Join(",", e.UnresolvedReferences)}"});
             }
-            catch (BlobNotFoundException e)
-            {
-                return NotFound(new ValidationProblemDetails {Title = $"Object {e.Blob} not found"});
-            }
+
         }
 
 
