Horde Storage - Added option (and set it by default) to disable the parallel requests in the blob service. We suspect this is causing a increase in response times from s3 due to the large increase in number of requests.

diff --git a/Engine/Source/Programs/Horde/HordeStorage/Horde.Storage/Implementation/Blob/IBlobService.cs b/Engine/Source/Programs/Horde/HordeStorage/Horde.Storage/Implementation/Blob/IBlobService.cs
--- a/Engine/Source/Programs/Horde/HordeStorage/Horde.Storage/Implementation/Blob/IBlobService.cs
+++ b/Engine/Source/Programs/Horde/HordeStorage/Horde.Storage/Implementation/Blob/IBlobService.cs
@@ -46,6 +46,7 @@ public interface IBlobService
 public class BlobService : IBlobService
 {
     private List<IBlobStore> _blobStores;
+    private readonly IOptionsMonitor<HordeStorageSettings> _settings;
     private readonly IBlobIndex _blobIndex;
     private readonly ILogger _logger = Log.ForContext<BlobService>();
 
@@ -58,6 +59,7 @@ public class BlobService : IBlobService
     public BlobService(IServiceProvider provider, IOptionsMonitor<HordeStorageSettings> settings, IBlobIndex blobIndex)
     {
         _blobStores = GetBlobStores(provider, settings).ToList();
+        _settings = settings;
         _blobIndex = blobIndex;
     }
 
@@ -144,30 +146,134 @@ public class BlobService : IBlobService
 
     private async Task<BlobIdentifier> PutObjectToStores(NamespaceId ns, IBufferedPayload bufferedPayload, BlobIdentifier identifier)
     {
-        await Parallel.ForEachAsync(_blobStores, async (store, cancellationToken) =>
+        if (_settings.CurrentValue.BlobStoreParallel)
         {
-            using Scope scope = Tracer.Instance.StartActive("put_blob_to_store");
-            scope.Span.ResourceName = identifier.ToString();
-            scope.Span.SetTag("store", store.ToString());
+            await Parallel.ForEachAsync(_blobStores, async (store, cancellationToken) =>
+            {
+                using Scope scope = Tracer.Instance.StartActive("put_blob_to_store");
+                scope.Span.ResourceName = identifier.ToString();
+                scope.Span.SetTag("store", store.ToString());
+
+                await using Stream s = bufferedPayload.GetStream();
+                await store.PutObject(ns, s, identifier);
+            });
+        }
+        else
+        {
+            foreach (IBlobStore store in _blobStores)
+            {
+                using Scope scope = Tracer.Instance.StartActive("put_blob_to_store");
+                scope.Span.ResourceName = identifier.ToString();
+                scope.Span.SetTag("store", store.ToString());
+
+                await using Stream s = bufferedPayload.GetStream();
+                await store.PutObject(ns, s, identifier);
+            }
+        }
 
-            await using Stream s = bufferedPayload.GetStream();
-            await store.PutObject(ns, s, identifier);
-        });
 
         return identifier;
     }
 
     private async Task<BlobIdentifier> PutObjectToStores(NamespaceId ns, byte[] payload, BlobIdentifier identifier)
     {
-        await Parallel.ForEachAsync(_blobStores, async (store, cancellationToken) =>
+        if (_settings.CurrentValue.BlobStoreParallel)
+        {
+            await Parallel.ForEachAsync(_blobStores,
+                async (store, cancellationToken) => { await store.PutObject(ns, payload, identifier); });
+        }
+        else
         {
-            await store.PutObject(ns, payload, identifier);
-        });
+            foreach (IBlobStore store in _blobStores)
+            {
+                await store.PutObject(ns, payload, identifier);
+            }
+        }
 
         return identifier;
     }
 
-    public async Task<BlobContents> GetObject(NamespaceId ns, BlobIdentifier blob)
+    public Task<BlobContents> GetObject(NamespaceId ns, BlobIdentifier blob)
+    {
+        if (_settings.CurrentValue.BlobStoreParallel)
+        {
+            return GetObjectParallel(ns, blob);
+        }
+        else
+        {
+            return GetObjectSequential(ns, blob);
+        }
+    }
+
+    private async Task<BlobContents> GetObjectSequential(NamespaceId ns, BlobIdentifier blob)
+    {
+        bool seenBlobNotFound = false;
+        bool seenNamespaceNotFound = false;
+        int numStoreMisses = 0;
+        BlobContents? blobContents = null;
+        foreach (IBlobStore store in _blobStores)
+        {
+            using Scope scope = Tracer.Instance.StartActive("HierarchicalStore.GetObject");
+            scope.Span.SetTag("BlobStore", store.GetType().Name);
+            scope.Span.SetTag("ObjectFound", false.ToString());
+            try
+            {
+                blobContents = await store.GetObject(ns, blob);
+                scope.Span.SetTag("ObjectFound", true.ToString());
+                break;
+            }
+            catch (BlobNotFoundException)
+            {
+                seenBlobNotFound = true;
+                numStoreMisses++;
+            }
+            catch (NamespaceNotFoundException)
+            {
+                seenNamespaceNotFound = true;
+            }
+        }
+
+        if (seenBlobNotFound && blobContents == null)
+        {
+            throw new BlobNotFoundException(ns, blob);
+        }
+
+        if (seenNamespaceNotFound && blobContents == null)
+        {
+            throw new NamespaceNotFoundException(ns);
+        }
+        
+        if (blobContents == null)
+        {
+            // Should not happen but exists to safeguard against the null pointer
+            throw new Exception("blobContents is null");
+        }
+
+        if (numStoreMisses >= 1)
+        {
+            using Scope _ = Tracer.Instance.StartActive("HierarchicalStore.Populate");
+            await using MemoryStream tempStream = new MemoryStream();
+            await blobContents.Stream.CopyToAsync(tempStream);
+            byte[] data = tempStream.ToArray();
+            
+            // Don't populate the last store, as that is where we got the hit
+            for (int i = 0; i < numStoreMisses; i++)
+            {
+                var blobStore = _blobStores[i];
+                using Scope scope = Tracer.Instance.StartActive("HierarchicalStore.PopulateStore");
+                scope.Span.SetTag("BlobStore", blobStore.GetType().Name);
+                // Populate each store traversed that did not have the content found lower in the hierarchy
+                await blobStore.PutObject(ns, data, blob);
+            }
+
+            blobContents = new BlobContents(new MemoryStream(data), data.Length);
+        }
+        
+        return blobContents;
+    }
+
+
+    private async Task<BlobContents> GetObjectParallel(NamespaceId ns, BlobIdentifier blob)
     {
         bool seenBlobNotFound = false;
         bool seenNamespaceNotFound = false;
@@ -266,7 +372,37 @@ public class BlobService : IBlobService
         return blobContents;
     }
 
-    public async Task<bool> Exists(NamespaceId ns, BlobIdentifier blob)
+    public Task<bool> Exists(NamespaceId ns, BlobIdentifier blob)
+    {
+        if (_settings.CurrentValue.BlobStoreParallel)
+        {
+            return ExistsParallel(ns, blob);
+        }
+        else
+        {
+            return ExistsSequential(ns, blob);
+        }
+    }
+
+    private async Task<bool> ExistsSequential(NamespaceId ns, BlobIdentifier blob)
+    {
+        foreach (IBlobStore store in _blobStores)
+        {
+            using Scope scope = Tracer.Instance.StartActive("HierarchicalStore.ObjectExists");
+            scope.Span.SetTag("BlobStore", store.GetType().Name);
+            if (await store.Exists(ns, blob))
+            {
+                scope.Span.SetTag("ObjectFound", true.ToString());
+                return true;
+            }
+            scope.Span.SetTag("ObjectFound", false.ToString());
+        }
+
+        return false;
+    }
+
+
+    private async Task<bool> ExistsParallel(NamespaceId ns, BlobIdentifier blob)
     {
         List<Task<bool>> existsTasks = new();
 
@@ -310,17 +446,13 @@ public class BlobService : IBlobService
         bool deletedAtLeastOnce = false;
         Task removeFromIndexTask = _blobIndex.RemoveBlobFromIndex(ns, blob);
 
-        List<Task> tasksToWaitFor = new List<Task>();
-        tasksToWaitFor.Add(removeFromIndexTask);
-
         foreach (IBlobStore store in _blobStores)
         {
             try
             {
                 using Scope scope = Tracer.Instance.StartActive("HierarchicalStore.DeleteObject");
                 scope.Span.SetTag("BlobStore", store.GetType().Name);
-                Task deleteTask = store.DeleteObject(ns, blob);
-                tasksToWaitFor.Add(deleteTask);
+                await store.DeleteObject(ns, blob);
                 deletedAtLeastOnce = true;
             }
             catch (NamespaceNotFoundException)
@@ -333,7 +465,7 @@ public class BlobService : IBlobService
             }
         }
 
-        await Task.WhenAll(tasksToWaitFor);
+        await removeFromIndexTask;
 
         if (deletedAtLeastOnce)
             return;
diff --git a/Engine/Source/Programs/Horde/HordeStorage/Horde.Storage/Settings.cs b/Engine/Source/Programs/Horde/HordeStorage/Horde.Storage/Settings.cs
--- a/Engine/Source/Programs/Horde/HordeStorage/Horde.Storage/Settings.cs
+++ b/Engine/Source/Programs/Horde/HordeStorage/Horde.Storage/Settings.cs
@@ -129,6 +129,7 @@ namespace Horde.Storage
         public int LastAccessRollupFrequencySeconds { get; set; } = 900; // 15 minutes
 
         public bool UseNewDDCEndpoints { get; set; } = false;
+        public bool BlobStoreParallel { get; set; } = false;
     }
 
     public class MongoSettings
