Horde Storage

Consider partial scylla objects as none existant, but log a warning so we can follow up on it.

diff --git a/Engine/Source/Programs/Horde/HordeStorage/Horde.Storage/Implementation/Objects/ScyllaReferencesStore.cs b/Engine/Source/Programs/Horde/HordeStorage/Horde.Storage/Implementation/Objects/ScyllaReferencesStore.cs
--- a/Engine/Source/Programs/Horde/HordeStorage/Horde.Storage/Implementation/Objects/ScyllaReferencesStore.cs
+++ b/Engine/Source/Programs/Horde/HordeStorage/Horde.Storage/Implementation/Objects/ScyllaReferencesStore.cs
@@ -10,6 +10,7 @@ using Cassandra.Mapping;
 using Datadog.Trace;
 using Jupiter.Implementation;
 using Microsoft.Extensions.Options;
+using Serilog;
 
 namespace Horde.Storage.Implementation
 {
@@ -18,6 +19,7 @@ namespace Horde.Storage.Implementation
         private readonly ISession _session;
         private readonly IMapper _mapper;
         private readonly IOptionsMonitor<ScyllaSettings> _settings;
+        private readonly ILogger _logger = Log.ForContext<ScyllaReferencesStore>();
 
         public ScyllaReferencesStore(IScyllaSessionManager scyllaSessionManager, IOptionsMonitor<ScyllaSettings> settings)
         {
@@ -76,7 +78,12 @@ namespace Horde.Storage.Implementation
             if (o == null)
                 throw new ObjectNotFoundException(ns, bucket, name);
 
-            o.ThrowIfRequiredFieldsAreMissing();
+            if (o.RequiredFieldsAreMissing())
+            {
+                _logger.Warning("Partial object found {Namespace} {Bucket} {Name} ignoring object", ns, bucket, name);
+                throw new ObjectNotFoundException(ns, bucket, name);
+            }
+
             return new ObjectRecord(new NamespaceId(o.Namespace!), new BucketId(o.Bucket!), new IoHashKey(o.Name!), o.LastAccessTime, o.InlinePayload, o.PayloadHash!.AsBlobIdentifier(), o.IsFinalized!.Value);
         }
 
@@ -446,18 +453,25 @@ namespace Horde.Storage.Implementation
         [Cassandra.Mapping.Attributes.Column("last_access_time")]
         public DateTime LastAccessTime { get; set; }
 
-        public void ThrowIfRequiredFieldsAreMissing()
+        public bool RequiredFieldsAreMissing()
         {
             if (string.IsNullOrEmpty(Namespace))
-                throw new ArgumentException("Namespace was not valid", nameof(Namespace));
+                //throw new ArgumentException("Namespace was not valid", nameof(Namespace));
+                return true;
             if (string.IsNullOrEmpty(Bucket))
-                throw new ArgumentException("Bucket was not valid", nameof(Bucket));
+                //throw new ArgumentException("Bucket was not valid", nameof(Bucket));
+                return true;
             if (string.IsNullOrEmpty(Name))
-                throw new ArgumentException("Name was not valid", nameof(Name));
+                //throw new ArgumentException("Name was not valid", nameof(Name));
+                return true;
             if (PayloadHash == null)
-                throw new ArgumentException("PayloadHash was not valid", nameof(PayloadHash));
+                //throw new ArgumentException("PayloadHash was not valid", nameof(PayloadHash));
+                return true;
             if (!IsFinalized.HasValue)
-                throw new ArgumentException("IsFinalized was not valid", nameof(IsFinalized));
+                //throw new ArgumentException("IsFinalized was not valid", nameof(IsFinalized));
+                return true;
+
+            return false;
         }
     }
 
