Horde Storage - Add last access tracking only for objects which exists, as otherwise we can get a partial object which only has last access tracking but no payload.

Also make sure to log a exception for the last access tracking, as we do not join the task itself.

diff --git a/Engine/Source/Programs/Horde/HordeStorage/Horde.Storage/Implementation/ObjectService.cs b/Engine/Source/Programs/Horde/HordeStorage/Horde.Storage/Implementation/ObjectService.cs
--- a/Engine/Source/Programs/Horde/HordeStorage/Horde.Storage/Implementation/ObjectService.cs
+++ b/Engine/Source/Programs/Horde/HordeStorage/Horde.Storage/Implementation/ObjectService.cs
@@ -10,6 +10,7 @@ using Datadog.Trace;
 using Horde.Storage.Implementation.Blob;
 using Jupiter.Implementation;
 using Jupiter.Utils;
+using Serilog;
 
 namespace Horde.Storage.Implementation
 {
@@ -34,6 +35,7 @@ namespace Horde.Storage.Implementation
         private readonly IReplicationLog _replicationLog;
         private readonly IBlobIndex _blobIndex;
         private readonly ILastAccessTracker<LastAccessRecord> _lastAccessTracker;
+        private readonly ILogger _logger = Log.ForContext<ObjectService>();
 
         public ObjectService(IReferencesStore referencesStore, IBlobService blobService, IReferenceResolver referenceResolver, IReplicationLog replicationLog, IBlobIndex blobIndex, ILastAccessTracker<LastAccessRecord> lastAccessTracker)
         {
@@ -47,9 +49,6 @@ namespace Horde.Storage.Implementation
 
         public async Task<(ObjectRecord, BlobContents?)> Get(NamespaceId ns, BucketId bucket, IoHashKey key, string[]? fields = null)
         {
-            // we do not wait for the last access tracking as it does not matter when it completes
-            Task lastAccessTask = _lastAccessTracker.TrackUsed(new LastAccessRecord(ns, bucket, key));
-
             // if no field filtering is being used we assume everything is needed
             IReferencesStore.FieldFlags flags = IReferencesStore.FieldFlags.All;
             if (fields != null)
@@ -66,6 +65,15 @@ namespace Horde.Storage.Implementation
             }
             ObjectRecord o = await _referencesStore.Get(ns, bucket, key, flags);
 
+            // we do not wait for the last access tracking as it does not matter when it completes
+            Task lastAccessTask = _lastAccessTracker.TrackUsed(new LastAccessRecord(ns, bucket, key)).ContinueWith(task =>
+            {
+                if (task.Exception != null)
+                {
+                    _logger.Error(task.Exception, "Exception when tracking last access record");
+                }
+            });;
+
             BlobContents? blobContents = null;
             if ((flags & IReferencesStore.FieldFlags.IncludePayload) != 0)
             {
