Horde: Fix tests.

diff --git a/Engine/Source/Programs/Horde/Horde.Build.Tests/IssueServiceTests.cs b/Engine/Source/Programs/Horde/Horde.Build.Tests/IssueServiceTests.cs
--- a/Engine/Source/Programs/Horde/Horde.Build.Tests/IssueServiceTests.cs
+++ b/Engine/Source/Programs/Horde/Horde.Build.Tests/IssueServiceTests.cs
@@ -73,10 +73,7 @@ namespace HordeServerTests
 
 			protected override void WriteFormattedEvent(LogLevel Level, int LineIndex, int LineCount, byte[] Line)
 			{
-				if (LineIndex == 0)
-				{
-					Events.Add((Level, Line));
-				}
+				Events.Add((Level, Line));
 			}
 
 			private async Task WriteAsync(LogLevel Level, byte[] Line)
@@ -1162,6 +1159,7 @@ namespace HordeServerTests
 				ILogFile? Log = await LogFileService.GetLogFileAsync(Job.Batches[0].Steps[0].LogId!.Value);
 				List<ILogEvent> Events = await LogFileService.FindLogEventsAsync(Log!);
 				Assert.AreEqual(1, Events.Count);
+				Assert.AreEqual(40, Events[0].LineCount);
 				ILogEventData EventData = await LogFileService.GetEventDataAsync(Log!, Events[0].LineIndex, Events[0].LineCount);
 
 				List<IIssue> Issues = await IssueService.FindIssuesAsync();
