Horde: Include hourly rate for using an agent in returned agent and job documents.

Global config file now includes a list of conditions which are evaluated against an agent's properties. The first condition that matches specifies the rate for the agent. This should give us plenty of flexibility in pricing out esoteric configurations dependent on factors like instance type, OS, etc...

Rates are cached in Redis and re-evaluated every 5 minutes.

diff --git a/Engine/Source/Programs/Horde/Horde.Build.Tests/TestSetup.cs b/Engine/Source/Programs/Horde/Horde.Build.Tests/TestSetup.cs
--- a/Engine/Source/Programs/Horde/Horde.Build.Tests/TestSetup.cs
+++ b/Engine/Source/Programs/Horde/Horde.Build.Tests/TestSetup.cs
@@ -226,7 +226,7 @@ namespace HordeServerTests
         {
 			ILogger<JobsController> Logger = ServiceProvider.GetRequiredService<ILogger<JobsController>>();
 			JobsController JobsCtrl = new JobsController(AclService, GraphCollection, PerforceService, StreamService, JobService,
-		        TemplateCollection, ArtifactCollection, UserCollection, NotificationService, Logger);
+		        TemplateCollection, ArtifactCollection, UserCollection, NotificationService, AgentService, Logger);
 	        JobsCtrl.ControllerContext = GetControllerContext();
 	        return JobsCtrl;
         }
diff --git a/Engine/Source/Programs/Horde/Horde.Build/Agents/Impl/AgentService.cs b/Engine/Source/Programs/Horde/Horde.Build/Agents/Impl/AgentService.cs
--- a/Engine/Source/Programs/Horde/Horde.Build/Agents/Impl/AgentService.cs
+++ b/Engine/Source/Programs/Horde/Horde.Build/Agents/Impl/AgentService.cs
@@ -1,7 +1,8 @@
 // Copyright Epic Games, Inc. All Rights Reserved.
 
 using EpicGames.Core;
-using Google.Protobuf;
+using EpicGames.Redis;
+using EpicGames.Serialization;
 using Google.Protobuf.WellKnownTypes;
 using HordeCommon;
 using HordeServer.Collections;
@@ -12,6 +13,7 @@ using Microsoft.Extensions.Hosting;
 using Microsoft.Extensions.Logging;
 using MongoDB.Bson;
 using MongoDB.Driver;
+using StackExchange.Redis;
 using StatsdClient;
 using System;
 using System.Collections.Generic;
@@ -28,6 +30,19 @@ namespace HordeServer.Services
 	using PoolId = StringId<IPool>;
 	using SessionId = ObjectId<ISession>;
 
+	/// <summary>
+	/// Singleton used to store agent costs
+	/// </summary>
+	[RedisConverter(typeof(RedisCbConverter<>))]
+	public class AgentRateTable
+	{
+		/// <summary>
+		/// List of costs for different agent types
+		/// </summary>
+		[CbField]
+		public List<AgentRateConfig> Entries { get; set; } = new List<AgentRateConfig>();
+	}
+
 	/// <summary>
 	/// Wraps funtionality for manipulating agents
 	/// </summary>
@@ -65,29 +80,38 @@ namespace HordeServer.Services
 		ILogger Logger;
 		ITicker Ticker;
 
+		RedisString<AgentRateTable> AgentRateTableData;
+		RedisService RedisService;
+
+		// Lazily updated costs for different agent types
+		AsyncCachedValue<AgentRateTable> AgentRateTable;
+
 		// Lazily updated list of current pools
 		AsyncCachedValue<List<IPool>> PoolsList;
 
 		// All the agents currently performing a long poll for work on this server
 		Dictionary<AgentId, CancellationTokenSource> WaitingAgents = new Dictionary<AgentId, CancellationTokenSource>();
-
+		
 		// Subscription for update events
 		IDisposable Subscription;
 
 		/// <summary>
 		/// Constructor
 		/// </summary>
-		public AgentService(IAgentCollection Agents, ILeaseCollection Leases, ISessionCollection Sessions, AclService AclService, IDowntimeService DowntimeService, IPoolCollection PoolCollection, IDogStatsd DogStatsd, IEnumerable<ITaskSource> TaskSources, IHostApplicationLifetime ApplicationLifetime, IClock Clock, ILogger<AgentService> Logger)
+		public AgentService(IAgentCollection Agents, ILeaseCollection Leases, ISessionCollection Sessions, AclService AclService, IDowntimeService DowntimeService, IPoolCollection PoolCollection, IDogStatsd DogStatsd, IEnumerable<ITaskSource> TaskSources, RedisService RedisService, IHostApplicationLifetime ApplicationLifetime, IClock Clock, ILogger<AgentService> Logger)
 		{
 			this.Agents = Agents;
 			this.Leases = Leases;
 			this.Sessions = Sessions;
 			this.AclService = AclService;
 			this.DowntimeService = DowntimeService;
+			this.AgentRateTableData = new RedisString<AgentRateTable>(RedisService.Database, "agent-rates");
+			this.AgentRateTable = new AsyncCachedValue<AgentRateTable>(() => AgentRateTableData.GetAsync(), TimeSpan.FromSeconds(2.0));//.FromMinutes(5.0));
 			this.PoolsList = new AsyncCachedValue<List<IPool>>(() => PoolCollection.GetAsync(), TimeSpan.FromSeconds(30.0));
 			this.DogStatsd = DogStatsd;
 			this.TaskSources = TaskSources.ToArray();
 			this.ApplicationLifetime = ApplicationLifetime;
+			this.RedisService = RedisService;
 			this.Clock = Clock;
 			this.Ticker = Clock.AddTicker(TimeSpan.FromSeconds(30.0), TickAsync, Logger);
 			this.Logger = Logger;
@@ -763,6 +787,67 @@ namespace HordeServer.Services
 			return Sessions.FindAsync(AgentId, StartTime, FinishTime, Index, Count);
 		}
 
+		/// <summary>
+		/// Update the rate table for different agent types
+		/// </summary>
+		/// <param name="Entries">New entries for the rate table</param>
+		/// <returns></returns>
+		public async Task UpdateRateTableAsync(List<AgentRateConfig> Entries)
+		{
+			await AgentRateTableData.SetAsync(new AgentRateTable { Entries = Entries });
+		}
+
+		/// <summary>
+		/// Gets the rate for the given agent
+		/// </summary>
+		/// <param name="AgentId">Agent id to query</param>
+		/// <returns>Hourly rate of running the given agent</returns>
+		public async ValueTask<double?> GetRateAsync(AgentId AgentId)
+		{
+			RedisKey Key = $"agent-rate/{AgentId}";
+
+			// Try to get the current value
+			RedisValue Value = await RedisService.Database.StringGetAsync(Key);
+			if (!Value.IsNull)
+			{
+				double Rate = (double)Value;
+				if (Rate == 0.0)
+				{
+					return null;
+				}
+				else
+				{
+					return Rate;
+				}
+			}
+			else
+			{
+				double Rate = 0.0;
+
+				// Get the rate table
+				AgentRateTable RateTable = await AgentRateTable.GetAsync();
+				if (RateTable.Entries.Count > 0)
+				{
+					IAgent? Agent = await GetAgentAsync(AgentId);
+					if (Agent != null)
+					{
+						foreach (AgentRateConfig Config in RateTable.Entries)
+						{
+							if (Config.Condition != null && Config.Condition.Evaluate(x => Agent.GetPropertyValues(x)))
+							{
+								Rate = Config.Rate;
+								break;
+							}
+						}
+					}
+				}
+
+				// Cache it for future reference
+				await RedisService.Database.StringSetAsync(Key, Rate, TimeSpan.FromMinutes(5.0), flags: CommandFlags.FireAndForget);
+				return Rate;
+			}
+		}
+
 		/// <summary>
 		/// Terminate any sessions for agents that are offline
 		/// </summary>
diff --git a/Engine/Source/Programs/Horde/Horde.Build/Api/Agents.cs b/Engine/Source/Programs/Horde/Horde.Build/Api/Agents.cs
--- a/Engine/Source/Programs/Horde/Horde.Build/Api/Agents.cs
+++ b/Engine/Source/Programs/Horde/Horde.Build/Api/Agents.cs
@@ -327,6 +327,11 @@ namespace HordeServer.Api
 		/// </summary>
 		public bool Enabled { get; set; }
 
+		/// <summary>
+		/// Cost estimate per-hour for this agent
+		/// </summary>
+		public double? Rate { get; set; }
+
 		/// <summary>
 		/// The current session id
 		/// </summary>
@@ -441,12 +446,14 @@ namespace HordeServer.Api
 		/// Constructor
 		/// </summary>
 		/// <param name="Agent">The agent to construct from</param>
+		/// <param name="Rate">Rate for this agent</param>
 		/// <param name="bIncludeAcl">Whether to include the ACL in the response</param>
-		public GetAgentResponse(IAgent Agent, bool bIncludeAcl)
+		public GetAgentResponse(IAgent Agent, double? Rate, bool bIncludeAcl)
 		{
 			this.Id = Agent.Id.ToString();
 			this.Name = Agent.Id.ToString();
 			this.Enabled = Agent.Enabled;
+			this.Rate = Rate;
 			this.Properties = new List<string>(Agent.Properties);
 			this.Resources = new Dictionary<string, int>(Agent.Resources);
 			this.SessionId = Agent.SessionId?.ToString();
diff --git a/Engine/Source/Programs/Horde/Horde.Build/Api/Jobs.cs b/Engine/Source/Programs/Horde/Horde.Build/Api/Jobs.cs
--- a/Engine/Source/Programs/Horde/Horde.Build/Api/Jobs.cs
+++ b/Engine/Source/Programs/Horde/Horde.Build/Api/Jobs.cs
@@ -727,6 +727,11 @@ namespace HordeServer.Api
 		/// </summary>
 		public string? AgentId { get; set; }
 
+		/// <summary>
+		/// Rate for using this agent (per hour)
+		/// </summary>
+		public double? AgentRate { get; set; }
+
 		/// <summary>
 		/// The agent session holding this lease
 		/// </summary>
@@ -752,13 +757,19 @@ namespace HordeServer.Api
 		/// </summary>
 		public DateTimeOffset? FinishTime { get; set; }
 
+		/// <summary>
+		/// Estimated cost of running this batch
+		/// </summary>
+		public double? Cost { get; set; }
+
 		/// <summary>
 		/// Converts this batch into a public response object
 		/// </summary>
 		/// <param name="Batch">The batch to construct from</param>
 		/// <param name="Steps">Steps in this batch</param>
+		/// <param name="AgentRate">Rate for this agent</param>
 		/// <returns>Response instance</returns>
-		public GetBatchResponse(IJobStepBatch Batch, List<GetStepResponse> Steps)
+		public GetBatchResponse(IJobStepBatch Batch, List<GetStepResponse> Steps, double? AgentRate)
 		{
 			this.Id = Batch.Id.ToString();
 			this.LogId = Batch.LogId?.ToString();
@@ -767,6 +778,7 @@ namespace HordeServer.Api
 			this.Error = Batch.Error;
 			this.Steps = Steps;
 			this.AgentId = Batch.AgentId?.ToString();
+			this.AgentRate = AgentRate;
 			this.SessionId = Batch.SessionId?.ToString();
 			this.LeaseId = Batch.LeaseId?.ToString();
 			this.WeightedPriority = Batch.SchedulePriority;
diff --git a/Engine/Source/Programs/Horde/Horde.Build/Controllers/AgentsController.cs b/Engine/Source/Programs/Horde/Horde.Build/Controllers/AgentsController.cs
--- a/Engine/Source/Programs/Horde/Horde.Build/Controllers/AgentsController.cs
+++ b/Engine/Source/Programs/Horde/Horde.Build/Controllers/AgentsController.cs
@@ -98,8 +98,9 @@ namespace HordeServer.Controllers
 			List<object> Responses = new List<object>();
 			foreach (IAgent Agent in Agents)
 			{
+				double? Rate = await AgentService.GetRateAsync(Agent.Id);
 				bool bIncludeAcl = await AgentService.AuthorizeAsync(Agent, AclAction.ViewPermissions, User, PermissionsCache);
-				Responses.Add(new GetAgentResponse(Agent, bIncludeAcl).ApplyFilter(Filter));
+				Responses.Add(new GetAgentResponse(Agent, Rate, bIncludeAcl).ApplyFilter(Filter));
 			}
 
 			return Responses;
@@ -123,8 +124,9 @@ namespace HordeServer.Controllers
 			}
 
 			GlobalPermissionsCache Cache = new GlobalPermissionsCache();
+			double? Rate = await AgentService.GetRateAsync(Agent.Id);
 			bool bIncludeAcl = await AgentService.AuthorizeAsync(Agent, AclAction.ViewPermissions, User, Cache);
-			return new GetAgentResponse(Agent, bIncludeAcl).ApplyFilter(Filter);
+			return new GetAgentResponse(Agent, Rate, bIncludeAcl).ApplyFilter(Filter);
 		}
 
 		/// <summary>
diff --git a/Engine/Source/Programs/Horde/Horde.Build/Controllers/JobsController.cs b/Engine/Source/Programs/Horde/Horde.Build/Controllers/JobsController.cs
--- a/Engine/Source/Programs/Horde/Horde.Build/Controllers/JobsController.cs
+++ b/Engine/Source/Programs/Horde/Horde.Build/Controllers/JobsController.cs
@@ -51,12 +51,13 @@ namespace HordeServer.Controllers
 		private readonly IArtifactCollection ArtifactCollection;
 		private readonly IUserCollection UserCollection;
 		private readonly INotificationService NotificationService;
+		private readonly AgentService AgentService;
 		private ILogger<JobsController> Logger;
 
 		/// <summary>
 		/// Constructor
 		/// </summary>
-		public JobsController(AclService AclService, IGraphCollection Graphs, IPerforceService Perforce, StreamService StreamService, JobService JobService, ITemplateCollection TemplateCollection, IArtifactCollection ArtifactCollection, IUserCollection UserCollection, INotificationService NotificationService, ILogger<JobsController> Logger)
+		public JobsController(AclService AclService, IGraphCollection Graphs, IPerforceService Perforce, StreamService StreamService, JobService JobService, ITemplateCollection TemplateCollection, IArtifactCollection ArtifactCollection, IUserCollection UserCollection, INotificationService NotificationService, AgentService AgentService, ILogger<JobsController> Logger)
 		{
 			this.AclService = AclService;
 			this.Graphs = Graphs;
@@ -67,6 +68,7 @@ namespace HordeServer.Controllers
 			this.ArtifactCollection = ArtifactCollection;
 			this.UserCollection = UserCollection;
 			this.NotificationService = NotificationService;
+			this.AgentService = AgentService;
 			this.Logger = Logger;
 		}
 
@@ -478,7 +480,14 @@ namespace HordeServer.Controllers
 			{
 				Steps.Add(await CreateStepResponseAsync(Step));
 			}
-			return new GetBatchResponse(Batch, Steps);
+
+			double? AgentRate = null;
+			if (Batch.AgentId != null)
+			{
+				AgentRate = await AgentService.GetRateAsync(Batch.AgentId.Value);
+			}
+
+			return new GetBatchResponse(Batch, Steps, AgentRate);
 		}
 
 		/// <summary>
diff --git a/Engine/Source/Programs/Horde/Horde.Build/Server/Config.cs b/Engine/Source/Programs/Horde/Horde.Build/Server/Config.cs
--- a/Engine/Source/Programs/Horde/Horde.Build/Server/Config.cs
+++ b/Engine/Source/Programs/Horde/Horde.Build/Server/Config.cs
@@ -1,6 +1,8 @@
 // Copyright Epic Games, Inc. All Rights Reserved.
 
 using EpicGames.Core;
+using EpicGames.Horde.Common;
+using EpicGames.Serialization;
 using HordeServer.Api;
 using HordeServer.Utilities;
 using System;
@@ -41,6 +43,11 @@ namespace HordeServer.Models
 		/// </summary>
 		public List<PerforceCluster> PerforceClusters { get; set; } = new List<PerforceCluster>();
 
+		/// <summary>
+		/// List of costs of a particular agent type
+		/// </summary>
+		public List<AgentRateConfig> Rates { get; set; } = new List<AgentRateConfig>();
+
 		/// <summary>
 		/// Maximum number of conforms to run at once
 		/// </summary>
@@ -253,5 +260,23 @@ namespace HordeServer.Models
 		/// </summary>
 		public string Id { get; set; } = String.Empty;
 	}
+
+	/// <summary>
+	/// Describes the monetary cost of agents matching a particular criteris
+	/// </summary>
+	public class AgentRateConfig
+	{
+		/// <summary>
+		/// Condition string
+		/// </summary>
+		[CbField("c")]
+		public Condition? Condition { get; set; }
+
+		/// <summary>
+		/// Rate for this agent
+		/// </summary>
+		[CbField("r")]
+		public double Rate { get; set; }
+	}
 }
 
diff --git a/Engine/Source/Programs/Horde/Horde.Build/Server/ConfigService.cs b/Engine/Source/Programs/Horde/Horde.Build/Server/ConfigService.cs
--- a/Engine/Source/Programs/Horde/Horde.Build/Server/ConfigService.cs
+++ b/Engine/Source/Programs/Horde/Horde.Build/Server/ConfigService.cs
@@ -59,6 +59,7 @@ namespace HordeServer.Services
 		INotificationService NotificationService;
 		INamespaceCollection NamespaceCollection;
 		IBucketCollection BucketCollection;
+		AgentService AgentService;
 		PoolService PoolService;
 		IOptionsMonitor<ServerSettings> Settings;
 		ITicker Ticker;
@@ -67,7 +68,7 @@ namespace HordeServer.Services
 		/// <summary>
 		/// Constructor
 		/// </summary>
-		public ConfigService(DatabaseService DatabaseService, IPerforceService PerforceService, ProjectService ProjectService, StreamService StreamService, INotificationService NotificationService, INamespaceCollection NamespaceCollection, IBucketCollection BucketCollection, PoolService PoolService, PerforceLoadBalancer PerforceLoadBalancer, IClock Clock, IOptionsMonitor<ServerSettings> Settings, ILogger<ConfigService> Logger)
+		public ConfigService(DatabaseService DatabaseService, IPerforceService PerforceService, ProjectService ProjectService, StreamService StreamService, INotificationService NotificationService, INamespaceCollection NamespaceCollection, IBucketCollection BucketCollection, PoolService PoolService, AgentService AgentService, PerforceLoadBalancer PerforceLoadBalancer, IClock Clock, IOptionsMonitor<ServerSettings> Settings, ILogger<ConfigService> Logger)
 		{
 			this.DatabaseService = DatabaseService;
 			this.PerforceService = PerforceService;
@@ -78,6 +79,7 @@ namespace HordeServer.Services
 			this.NamespaceCollection = NamespaceCollection;
 			this.BucketCollection = BucketCollection;
 			this.PoolService = PoolService;
+			this.AgentService = AgentService;
 			this.Settings = Settings;
 			this.Ticker = Clock.AddSharedTicker<ConfigService>(TimeSpan.FromMinutes(1.0), TickLeaderAsync, Logger);
 			this.Logger = Logger;
@@ -151,6 +153,9 @@ namespace HordeServer.Services
 				}
 			}
 
+			// Update the agent rate table
+			await AgentService.UpdateRateTableAsync(GlobalConfig.Rates);
+
 			// Projects to remove
 			List<IProject> Projects = await ProjectService.GetProjectsAsync();
 
diff --git a/Engine/Source/Programs/Shared/EpicGames.Horde/Common/Condition.cs b/Engine/Source/Programs/Shared/EpicGames.Horde/Common/Condition.cs
--- a/Engine/Source/Programs/Shared/EpicGames.Horde/Common/Condition.cs
+++ b/Engine/Source/Programs/Shared/EpicGames.Horde/Common/Condition.cs
@@ -11,6 +11,8 @@ using System.Diagnostics.CodeAnalysis;
 using System.Globalization;
 using System.Linq;
 using System.Text;
+using System.Text.Json;
+using System.Text.Json.Serialization;
 using System.Text.RegularExpressions;
 
 namespace EpicGames.Horde.Common
@@ -32,7 +34,9 @@ namespace EpicGames.Horde.Common
 		Gte,
 		Lparen,
 		Rparen,
-		Regex
+		Regex,
+		True,
+		False,
 	}
 
 	class TokenReader
@@ -143,7 +147,22 @@ namespace EpicGames.Horde.Common
 						{
 							EndIdx++;
 						}
-						SetCurrent(EndIdx - Offset, TokenType.Identifier);
+
+						TokenType Type;
+						if (EndIdx == Offset + 4 && String.Compare(Input, Offset, "true", 0, 4, StringComparison.OrdinalIgnoreCase) == 0)
+						{
+							Type = TokenType.True;
+						}
+						else if (EndIdx == Offset + 5 && String.Compare(Input, Offset, "false", 0, 5, StringComparison.OrdinalIgnoreCase) == 0)
+						{
+							Type = TokenType.False;
+						}
+						else
+						{
+							Type = TokenType.Identifier;
+						}
+
+						SetCurrent(EndIdx - Offset, Type);
 						return;
 					default:
 						SetCurrent(1, TokenType.Error, $"Invalid character at offset {Offset}: '{Input[Offset]}'");
@@ -290,6 +309,7 @@ namespace EpicGames.Horde.Common
 	[JsonSchemaString]
 	[TypeConverter(typeof(ConditionTypeConverter))]
 	[CbConverter(typeof(ConditionCbConverter))]
+	[JsonConverter(typeof(ConditionJsonConverter))]
 	public class Condition
 	{
 		[DebuggerDisplay("{Type}")]
@@ -393,6 +413,11 @@ namespace EpicGames.Horde.Common
 					return ParseSubExpr(Reader);
 				case TokenType.Lparen:
 					return ParseSubExpr(Reader);
+				case TokenType.True:
+				case TokenType.False:
+					Tokens.Add(new Token(Reader.Type, 0));
+					Reader.MoveNext();
+					return null;
 				default:
 					return ParseComparisonExpr(Reader);
 			}
@@ -483,6 +508,10 @@ namespace EpicGames.Horde.Common
 			Token Token = Tokens[Idx++];
 			switch (Token.Type)
 			{
+				case TokenType.True:
+					return true;
+				case TokenType.False:
+					return false;
 				case TokenType.Not:
 					return !EvaluateCondition(ref Idx, GetPropertyValues);
 				case TokenType.Eq:
@@ -584,7 +613,7 @@ namespace EpicGames.Horde.Common
 	}
 
 	/// <summary>
-	/// Type converter from strings to PropertyFilter objects
+	/// Type converter from strings to condition objects
 	/// </summary>
 	sealed class ConditionTypeConverter : TypeConverter
 	{
@@ -600,4 +629,23 @@ namespace EpicGames.Horde.Common
 		/// <inheritdoc/>
 		public override object? ConvertTo(ITypeDescriptorContext Context, CultureInfo Culture, object Value, Type DestinationType) => ((Condition)Value).Text;
 	}
+
+	/// <summary>
+	/// Type converter from Json strings to condition objects
+	/// </summary>
+	sealed class ConditionJsonConverter : JsonConverter<Condition>
+	{
+		/// <inheritdoc/>
+		public override bool CanConvert(Type typeToConvert) => typeToConvert == typeof(Condition);
+
+		public override Condition Read(ref Utf8JsonReader reader, Type typeToConvert, JsonSerializerOptions options)
+		{
+			return Condition.Parse(reader.GetString());
+		}
+
+		public override void Write(Utf8JsonWriter writer, Condition value, JsonSerializerOptions options)
+		{
+			writer.WriteStringValue(value.ToString());
+		}
+	}
 }
diff --git a/Engine/Source/Programs/Shared/EpicGames.Serialization/CompactBinary.cs b/Engine/Source/Programs/Shared/EpicGames.Serialization/CompactBinary.cs
--- a/Engine/Source/Programs/Shared/EpicGames.Serialization/CompactBinary.cs
+++ b/Engine/Source/Programs/Shared/EpicGames.Serialization/CompactBinary.cs
@@ -1072,12 +1072,18 @@ namespace EpicGames.Serialization
 			}
 		}
 
+		/// <summary>
+		/// Access the field as a double.
+		/// </summary>
+		/// <returns>Value of the field</returns>
+		public double AsDouble() => AsDouble(0.0);
+
 		/// <summary>
 		/// Access the field as a double. Returns the provided default on error.
 		/// </summary>
 		/// <param name="Default">Default value</param>
 		/// <returns>Value of the field</returns>
-		public double AsDouble(double Default = 0.0f)
+		public double AsDouble(double Default)
 		{
 			switch (GetType())
 			{
diff --git a/Engine/Source/Programs/Shared/EpicGames.Serialization/CompactBinaryWriter.cs b/Engine/Source/Programs/Shared/EpicGames.Serialization/CompactBinaryWriter.cs
--- a/Engine/Source/Programs/Shared/EpicGames.Serialization/CompactBinaryWriter.cs
+++ b/Engine/Source/Programs/Shared/EpicGames.Serialization/CompactBinaryWriter.cs
@@ -521,6 +521,27 @@ namespace EpicGames.Serialization
 			WriteIntegerPayload(Value);
 		}
 
+		/// <summary>
+		/// Writes an unnamed double field
+		/// </summary>
+		/// <param name="Value">Value to be written</param>
+		public void WriteDoubleValue(double Value)
+		{
+			WriteFieldHeader(CbFieldType.Float64);
+			BinaryPrimitives.WriteInt64BigEndian(Allocate(sizeof(double)).Span, BitConverter.DoubleToInt64Bits(Value));
+		}
+
+		/// <summary>
+		/// Writes a named double field
+		/// </summary>
+		/// <param name="Name">Name of the field</param>
+		/// <param name="Value">Value to be written</param>
+		public void WriteDouble(Utf8String Name, double Value)
+		{
+			WriteFieldHeader(CbFieldType.Float64, Name);
+			BinaryPrimitives.WriteInt64BigEndian(Allocate(sizeof(double)).Span, BitConverter.DoubleToInt64Bits(Value));
+		}
+
 		/// <summary>
 		/// Writes the payload for a <see cref="DateTime"/> value
 		/// </summary>
diff --git a/Engine/Source/Programs/Shared/EpicGames.Serialization/Converters/CbDefaultConverter.cs b/Engine/Source/Programs/Shared/EpicGames.Serialization/Converters/CbDefaultConverter.cs
--- a/Engine/Source/Programs/Shared/EpicGames.Serialization/Converters/CbDefaultConverter.cs
+++ b/Engine/Source/Programs/Shared/EpicGames.Serialization/Converters/CbDefaultConverter.cs
@@ -102,6 +102,8 @@ namespace EpicGames.Serialization.Converters
 		static void WriteInt32(CbWriter Writer, Utf8String Name, int Value) => Writer.WriteInteger(Name, Value);
 		static void WriteInt64Value(CbWriter Writer, long Value) => Writer.WriteIntegerValue(Value);
 		static void WriteInt64(CbWriter Writer, Utf8String Name, long Value) => Writer.WriteInteger(Name, Value);
+		static void WriteDoubleValue(CbWriter Writer, double Value) => Writer.WriteDoubleValue(Value);
+		static void WriteDouble(CbWriter Writer, Utf8String Name, double Value) => Writer.WriteDouble(Name, Value);
 
 		static string ReadString(CbField Field) => Field.AsString().ToString();
 		static void WriteString(CbWriter Writer, Utf8String Name, string Value) => Writer.WriteString(Name, Value);
@@ -120,6 +122,7 @@ namespace EpicGames.Serialization.Converters
 			GetPodConverterInfo(x => x.AsBool(), (w, v) => w.WriteBoolValue(v), (w, n, v) => w.WriteBool(n, v), SkipIfNullOrZero),
 			GetPodConverterInfo(x => x.AsInt32(), (w, v) => WriteInt32Value(w, v), (w, n, v) => WriteInt32(w, n, v), SkipIfNullOrZero),
 			GetPodConverterInfo(x => x.AsInt64(), (w, v) => WriteInt64Value(w, v), (w, n, v) => WriteInt64(w, n, v), SkipIfNullOrZero),
+			GetPodConverterInfo(x => x.AsDouble(), (w, v) => WriteDoubleValue(w, v), (w, n, v)=> WriteDouble(w, n, v), SkipIfNullOrZero),
 			GetPodConverterInfo(x => x.AsString(), (w, v) => w.WriteStringValue(v), (w, n, v) => w.WriteString(n, v), null),
 			GetPodConverterInfo(x => x.AsHash(), (w, v) => w.WriteHashValue(v), (w, n, v) => w.WriteHash(n, v), null),
 			GetPodConverterInfo(x => x.AsObjectAttachment(), (w, v) => w.WriteObjectAttachmentValue(v.Hash), (w, n, v) => w.WriteObjectAttachment(n, v.Hash), null),
