Horde: Remove unused 'cost' field.

diff --git a/Engine/Source/Programs/Horde/Horde.Build/Api/Jobs.cs b/Engine/Source/Programs/Horde/Horde.Build/Api/Jobs.cs
--- a/Engine/Source/Programs/Horde/Horde.Build/Api/Jobs.cs
+++ b/Engine/Source/Programs/Horde/Horde.Build/Api/Jobs.cs
@@ -757,11 +757,6 @@ namespace HordeServer.Api
 		/// </summary>
 		public DateTimeOffset? FinishTime { get; set; }
 
-		/// <summary>
-		/// Estimated cost of running this batch
-		/// </summary>
-		public double? Cost { get; set; }
-
 		/// <summary>
 		/// Converts this batch into a public response object
 		/// </summary>
