Horde: Add an explicit permission for being able to see job costs on the dashboard. Currently only granted to admins.

diff --git a/Engine/Source/Programs/Horde/Horde.Build/Acls/Acl.cs b/Engine/Source/Programs/Horde/Horde.Build/Acls/Acl.cs
--- a/Engine/Source/Programs/Horde/Horde.Build/Acls/Acl.cs
+++ b/Engine/Source/Programs/Horde/Horde.Build/Acls/Acl.cs
@@ -160,7 +160,7 @@ namespace HordeServer.Models
 		DeleteAgent,
 
 		/// <summary>
-		/// View an available agents
+		/// View an agent
 		/// </summary>
 		ViewAgent,
 
@@ -381,12 +381,16 @@ namespace HordeServer.Models
 		/// </summary>
 		Impersonate,
 
+		/// <summary>
+		/// View estimated costs for particular operations
+		/// </summary>
+		ViewCosts,
 
 
 
 
 		//// PERMISSIONS ////
-		
+
 		/// <summary>
 		/// Ability to view permissions on an object
 		/// </summary>
diff --git a/Engine/Source/Programs/Horde/Horde.Build/Controllers/AgentsController.cs b/Engine/Source/Programs/Horde/Horde.Build/Controllers/AgentsController.cs
--- a/Engine/Source/Programs/Horde/Horde.Build/Controllers/AgentsController.cs
+++ b/Engine/Source/Programs/Horde/Horde.Build/Controllers/AgentsController.cs
@@ -98,7 +98,11 @@ namespace HordeServer.Controllers
 			List<object> Responses = new List<object>();
 			foreach (IAgent Agent in Agents)
 			{
-				double? Rate = await AgentService.GetRateAsync(Agent.Id);
+				double? Rate = null;
+				if (await AgentService.AuthorizeAsync(Agent, AclAction.ViewCosts, User, PermissionsCache))
+				{
+					Rate = await AgentService.GetRateAsync(Agent.Id);
+				}
 				bool bIncludeAcl = await AgentService.AuthorizeAsync(Agent, AclAction.ViewPermissions, User, PermissionsCache);
 				Responses.Add(new GetAgentResponse(Agent, Rate, bIncludeAcl).ApplyFilter(Filter));
 			}
@@ -124,7 +128,13 @@ namespace HordeServer.Controllers
 			}
 
 			GlobalPermissionsCache Cache = new GlobalPermissionsCache();
-			double? Rate = await AgentService.GetRateAsync(Agent.Id);
+
+			double? Rate = null;
+			if (await AgentService.AuthorizeAsync(Agent, AclAction.ViewCosts, User, Cache))
+			{
+				Rate = await AgentService.GetRateAsync(Agent.Id);
+			}
+
 			bool bIncludeAcl = await AgentService.AuthorizeAsync(Agent, AclAction.ViewPermissions, User, Cache);
 			return new GetAgentResponse(Agent, Rate, bIncludeAcl).ApplyFilter(Filter);
 		}
diff --git a/Engine/Source/Programs/Horde/Horde.Build/Controllers/JobsController.cs b/Engine/Source/Programs/Horde/Horde.Build/Controllers/JobsController.cs
--- a/Engine/Source/Programs/Horde/Horde.Build/Controllers/JobsController.cs
+++ b/Engine/Source/Programs/Horde/Horde.Build/Controllers/JobsController.cs
@@ -396,7 +396,8 @@ namespace HordeServer.Controllers
 
 			IGraph Graph = await JobService.GetGraphAsync(Job);
 			bool bIncludeAcl = await JobService.AuthorizeAsync(Job, AclAction.ViewPermissions, User, Cache);
-			return await CreateJobResponseAsync(Job, Graph, bIncludeAcl, Filter);
+			bool bIncludeCosts = await JobService.AuthorizeAsync(Job, AclAction.ViewCosts, User, Cache);
+			return await CreateJobResponseAsync(Job, Graph, bIncludeAcl, bIncludeCosts, Filter);
 		}
 
 		/// <summary>
@@ -405,17 +406,18 @@ namespace HordeServer.Controllers
 		/// <param name="Job">The job document</param>
 		/// <param name="Graph">The graph for this job</param>
 		/// <param name="bIncludeAcl">Whether to include the ACL in the response</param>
+		/// <param name="bIncludeCosts">Whether to include costs in the response</param>
 		/// <param name="Filter">Filter for the properties to return</param>
 		/// <returns>Object containing the requested properties</returns>
-		async Task<object> CreateJobResponseAsync(IJob Job, IGraph Graph, bool bIncludeAcl, PropertyFilter? Filter)
+		async Task<object> CreateJobResponseAsync(IJob Job, IGraph Graph, bool bIncludeAcl, bool bIncludeCosts, PropertyFilter? Filter)
 		{
 			if (Filter == null)
 			{
-				return await CreateJobResponseAsync(Job, Graph, true, true, bIncludeAcl);
+				return await CreateJobResponseAsync(Job, Graph, true, true, bIncludeAcl, bIncludeCosts);
 			}
 			else
 			{
-				return Filter.ApplyTo(await CreateJobResponseAsync(Job, Graph, Filter.Includes(nameof(GetJobResponse.Batches)), Filter.Includes(nameof(GetJobResponse.Labels)) || Filter.Includes(nameof(GetJobResponse.DefaultLabel)), bIncludeAcl));
+				return Filter.ApplyTo(await CreateJobResponseAsync(Job, Graph, Filter.Includes(nameof(GetJobResponse.Batches)), Filter.Includes(nameof(GetJobResponse.Labels)) || Filter.Includes(nameof(GetJobResponse.DefaultLabel)), bIncludeAcl, bIncludeCosts));
 			}
 		}
 
@@ -427,8 +429,9 @@ namespace HordeServer.Controllers
 		/// <param name="bIncludeBatches">Whether to include the job batches in the response</param>
 		/// <param name="bIncludeLabels">Whether to include the job aggregates in the response</param>
 		/// <param name="bIncludeAcl">Whether to include the ACL in the response</param>
+		/// <param name="bIncludeCosts">Whether to include costs of running particular agents</param>
 		/// <returns>The response object</returns>
-		async ValueTask<GetJobResponse> CreateJobResponseAsync(IJob Job, IGraph Graph, bool bIncludeBatches, bool bIncludeLabels, bool bIncludeAcl)
+		async ValueTask<GetJobResponse> CreateJobResponseAsync(IJob Job, IGraph Graph, bool bIncludeBatches, bool bIncludeLabels, bool bIncludeAcl, bool bIncludeCosts)
 		{
 			GetThinUserInfoResponse? StartedByUserInfo = null;
 			if (Job.StartedByUserId != null)
@@ -456,7 +459,7 @@ namespace HordeServer.Controllers
 					Response.Batches = new List<GetBatchResponse>();
 					foreach (IJobStepBatch Batch in Job.Batches)
 					{
-						Response.Batches.Add(await CreateBatchResponseAsync(Batch));
+						Response.Batches.Add(await CreateBatchResponseAsync(Batch, bIncludeCosts));
 					}
 				}
 				if (bIncludeLabels)
@@ -472,8 +475,9 @@ namespace HordeServer.Controllers
 		/// Get the response object for a batch
 		/// </summary>
 		/// <param name="Batch"></param>
+		/// <param name="bIncludeCosts"></param>
 		/// <returns></returns>
-		async ValueTask<GetBatchResponse> CreateBatchResponseAsync(IJobStepBatch Batch)
+		async ValueTask<GetBatchResponse> CreateBatchResponseAsync(IJobStepBatch Batch, bool bIncludeCosts)
 		{
 			List<GetStepResponse> Steps = new List<GetStepResponse>();
 			foreach (IJobStep Step in Batch.Steps)
@@ -482,7 +486,7 @@ namespace HordeServer.Controllers
 			}
 
 			double? AgentRate = null;
-			if (Batch.AgentId != null)
+			if (Batch.AgentId != null && bIncludeCosts)
 			{
 				AgentRate = await AgentService.GetRateAsync(Batch.AgentId.Value);
 			}
@@ -730,9 +734,15 @@ namespace HordeServer.Controllers
 						bIncludeAcl = await JobService.AuthorizeAsync(Job, AclAction.ViewPermissions, User, PermissionsCache);
 					}
 
+					bool bIncludeCosts;
+					using (IScope _ = GlobalTracer.Instance.BuildSpan("AuthorizeViewCosts").StartActive())
+					{
+						bIncludeCosts = await JobService.AuthorizeAsync(Job, AclAction.ViewCosts, User, PermissionsCache);
+					}
+
 					using (IScope _ = GlobalTracer.Instance.BuildSpan("CreateResponse").StartActive())
 					{
-						Responses.Add(await CreateJobResponseAsync(Job, Graph, bIncludeAcl, Filter));
+						Responses.Add(await CreateJobResponseAsync(Job, Graph, bIncludeAcl, bIncludeCosts, Filter));
 					}
 				}
 			}
@@ -804,9 +814,15 @@ namespace HordeServer.Controllers
 						bIncludeAcl = await JobService.AuthorizeAsync(Job, AclAction.ViewPermissions, User, PermissionsCache);
 					}
 
+					bool bIncludeCosts;
+					using (IScope _ = GlobalTracer.Instance.BuildSpan("AuthorizeViewCosts").StartActive())
+					{
+						bIncludeCosts = await JobService.AuthorizeAsync(Job, AclAction.ViewCosts, User, PermissionsCache);
+					}
+
 					using (IScope _ = GlobalTracer.Instance.BuildSpan("CreateResponse").StartActive())
 					{
-						Responses.Add(await CreateJobResponseAsync(Job, Graph, bIncludeAcl, Filter));
+						Responses.Add(await CreateJobResponseAsync(Job, Graph, bIncludeAcl, bIncludeCosts, Filter));
 					}
 				}
 			}
@@ -987,15 +1003,19 @@ namespace HordeServer.Controllers
 			{
 				return NotFound(JobId);
 			}
-			if (!await JobService.AuthorizeAsync(Job, AclAction.ViewJob, User, null))
+
+			StreamPermissionsCache Cache = new StreamPermissionsCache();
+			if (!await JobService.AuthorizeAsync(Job, AclAction.ViewJob, User, Cache))
 			{
 				return Forbid(AclAction.ViewJob, JobId);
 			}
 
+			bool bIncludeCosts = await JobService.AuthorizeAsync(Job, AclAction.ViewCosts, User, Cache);
+
 			List<object> Responses = new List<object>();
 			foreach (IJobStepBatch Batch in Job.Batches)
 			{
-				GetBatchResponse Response = await CreateBatchResponseAsync(Batch);
+				GetBatchResponse Response = await CreateBatchResponseAsync(Batch, bIncludeCosts);
 				Responses.Add(Response.ApplyFilter(Filter));
 			}
 			return Responses;
@@ -1052,17 +1072,21 @@ namespace HordeServer.Controllers
 			{
 				return NotFound(JobId);
 			}
-			if (!await JobService.AuthorizeAsync(Job, AclAction.ViewJob, User, null))
+
+			StreamPermissionsCache Cache = new StreamPermissionsCache();
+			if (!await JobService.AuthorizeAsync(Job, AclAction.ViewJob, User, Cache))
 			{
 				return Forbid(AclAction.ViewJob, JobId);
 			}
 
+			bool bIncludeCosts = await JobService.AuthorizeAsync(Job, AclAction.ViewCosts, User, Cache);
+
 			IGraph Graph = await JobService.GetGraphAsync(Job);
 			foreach (IJobStepBatch Batch in Job.Batches)
 			{
 				if (Batch.Id == BatchId)
 				{
-					GetBatchResponse Response = await CreateBatchResponseAsync(Batch);
+					GetBatchResponse Response = await CreateBatchResponseAsync(Batch, bIncludeCosts);
 					return Response.ApplyFilter(Filter);
 				}
 			}
