Horde: Fix hang running tests. If an exception occurs in an AsyncCachedValue delegate, it doesn't get propagated to the caller and the task stays in a waiting state forever.

diff --git a/Engine/Source/Programs/Horde/Horde.Build/Utilities/AsyncCachedValue.cs b/Engine/Source/Programs/Horde/Horde.Build/Utilities/AsyncCachedValue.cs
--- a/Engine/Source/Programs/Horde/Horde.Build/Utilities/AsyncCachedValue.cs
+++ b/Engine/Source/Programs/Horde/Horde.Build/Utilities/AsyncCachedValue.cs
@@ -85,7 +85,7 @@ namespace HordeServer.Utilities
 				TaskCompletionSource<State> NextTaskSource = new TaskCompletionSource<State>();
 				if (Interlocked.CompareExchange(ref UpdateTask, NextTaskSource.Task, null) == null)
 				{
-					Task.Run(async () => NextTaskSource.SetResult(new State(await Generator())));
+					Task.Run(() => GenerateNewValueAsync(NextTaskSource));
 					return NextTaskSource.Task;
 				}
 				CurrentTask = UpdateTask;
@@ -93,6 +93,18 @@ namespace HordeServer.Utilities
 			return CurrentTask;
 		}
 
+		async Task GenerateNewValueAsync(TaskCompletionSource<State> NextTaskSource)
+		{
+			try
+			{
+				NextTaskSource.SetResult(new State(await Generator()));
+			}
+			catch (Exception Ex)
+			{
+				NextTaskSource.SetException(Ex);
+			}
+		}
+
 		/// <summary>
 		/// Tries to get the current value
 		/// </summary>
