Horde: Add job/step pricing information to summaries
#rnx

diff --git a/Engine/Source/Programs/Horde/HordeDashboard/src/backend/Api.ts b/Engine/Source/Programs/Horde/HordeDashboard/src/backend/Api.ts
--- a/Engine/Source/Programs/Horde/HordeDashboard/src/backend/Api.ts
+++ b/Engine/Source/Programs/Horde/HordeDashboard/src/backend/Api.ts
@@ -1142,6 +1142,9 @@ export type GetBatchResponse = {
 	/**The agent assigned to execute this group */
 	agentId?: string;
 
+	/** The USD rate of an agent hour */
+	agentRate?: number;
+
 	/**The agent session holding this lease */
 	sessionId?: string;
 
diff --git a/Engine/Source/Programs/Horde/HordeDashboard/src/backend/JobDetails.ts b/Engine/Source/Programs/Horde/HordeDashboard/src/backend/JobDetails.ts
--- a/Engine/Source/Programs/Horde/HordeDashboard/src/backend/JobDetails.ts
+++ b/Engine/Source/Programs/Horde/HordeDashboard/src/backend/JobDetails.ts
@@ -5,8 +5,9 @@ import { getTheme, mergeStyles, mergeStyleSets } from '@fluentui/react/lib/Styli
 import backend from '.';
 import { getBatchInitElapsed, getNiceTime, getStepElapsed, getStepETA, getStepFinishTime, getStepTimingDelta } from '../base/utilities/timeUtils';
 import { getBatchText } from '../components/JobDetailCommon';
-import { AgentData, ArtifactData, BatchData, EventData, GetGroupResponse, GetJobStepRefResponse, GetJobTimingResponse, GetLabelResponse, GetLabelStateResponse, GetLabelTimingInfoResponse, GetTemplateResponse, GroupData, IssueData, JobData, JobStepBatchState, JobStepOutcome, JobStepState, LabelState, NodeData, ReportPlacement, StepData, StreamData, TestData } from './Api';
+import { AgentData, ArtifactData, BatchData, EventData, GetGroupResponse, GetJobStepRefResponse, GetJobTimingResponse, GetLabelResponse, GetLabelStateResponse, GetLabelTimingInfoResponse, GetTemplateResponse, GroupData, IssueData, JobData, JobState, JobStepBatchState, JobStepOutcome, JobStepState, LabelState, NodeData, ReportPlacement, StepData, StreamData, TestData } from './Api';
 import { projectStore } from './ProjectStore';
+import moment from 'moment';
 
 const theme = getTheme();
 
@@ -392,6 +393,55 @@ export class JobDetails {
         return this.labels.indexOf(label);
     }
 
+    stepPrice(stepId: string): number | undefined {
+
+        const step = this.stepById(stepId);
+        const batch = this.batchByStepId(stepId);
+
+        /*
+        if (batch) {
+            batch.agentRate = 5.0;
+        }
+        */
+                
+        if (!step || !batch || !batch.agentRate || !step.startTime || !step.finishTime) {
+            return undefined;
+        }
+
+        const start = moment(step.startTime);
+        const end = moment(step.finishTime);
+        const hours = moment.duration(end.diff(start)).asHours();
+        const price = hours * batch.agentRate;
+                                
+        return price ? price : undefined;
+    }
+
+    jobPrice(): number | undefined {
+
+        if (!this.batches || this.jobdata?.state !== JobState.Complete) {
+            return undefined;
+        }
+
+        let price = 0;
+
+        this.batches.forEach(b => {
+
+            /*
+            b.agentRate = 5.0;
+            */
+
+            if (b.agentRate && b.startTime && b.finishTime) {
+                const start = moment(b.startTime);
+                const end = moment(b.finishTime);
+                const hours = moment.duration(end.diff(start)).asHours();
+                price += hours * b.agentRate;
+                                
+            } 
+        });
+
+        return price ? price : undefined;
+    }
+
     labelByIndex(idxIn: number | string | undefined | null): JobLabel | undefined {
 
         const idx = parseInt((idxIn ?? "").toString());
