Horde: Disable cost functionality in Agent service.

diff --git a/Engine/Source/Programs/Horde/Horde.Build/Agents/Impl/AgentService.cs b/Engine/Source/Programs/Horde/Horde.Build/Agents/Impl/AgentService.cs
--- a/Engine/Source/Programs/Horde/Horde.Build/Agents/Impl/AgentService.cs
+++ b/Engine/Source/Programs/Horde/Horde.Build/Agents/Impl/AgentService.cs
@@ -802,9 +802,10 @@ namespace HordeServer.Services
 		/// </summary>
 		/// <param name="AgentId">Agent id to query</param>
 		/// <returns>Hourly rate of running the given agent</returns>
-		public async ValueTask<double?> GetRateAsync(AgentId AgentId)
+		public ValueTask<double?> GetRateAsync(AgentId AgentId)
 		{
-			RedisKey Key = $"agent-rate/{AgentId}";
+			return new ValueTask<double?>(Task.FromResult<double?>(null));
+/*			RedisKey Key = $"agent-rate/{AgentId}";
 
 			// Try to get the current value
 			RedisValue Value = await RedisService.Database.StringGetAsync(Key);
@@ -845,7 +846,7 @@ namespace HordeServer.Services
 				// Cache it for future reference
 				await RedisService.Database.StringSetAsync(Key, Rate, TimeSpan.FromMinutes(5.0), flags: CommandFlags.FireAndForget);
 				return Rate;
-			}
+			}*/
 		}
 
 		/// <summary>
