Horde: Remove BsonIgnoreIfDefault flag on PoolDocument as it prevented value from getting serialized

Possibly a bug with saving values with this flag turned on? Removing it allowed EnableAutoscaling field to be serialized, even when set to false.

diff --git a/Engine/Source/Programs/Horde/Horde.Build/Fleet/Impl/PoolCollection.cs b/Engine/Source/Programs/Horde/Horde.Build/Fleet/Impl/PoolCollection.cs
--- a/Engine/Source/Programs/Horde/Horde.Build/Fleet/Impl/PoolCollection.cs
+++ b/Engine/Source/Programs/Horde/Horde.Build/Fleet/Impl/PoolCollection.cs
@@ -35,7 +35,6 @@ namespace HordeServer.Collections.Impl
 			public List<AgentWorkspace> Workspaces { get; set; } = new List<AgentWorkspace>();
 			public bool UseAutoSdk { get; set; } = true;
 			public Dictionary<string, string> Properties { get; set; } = new Dictionary<string, string>();
-			[BsonIgnoreIfDefault(true)]
 			public bool EnableAutoscaling { get; set; } = true;
 			public int? MinAgents { get; set; }
 			public int? NumReserveAgents { get; set; }
