Horde: Allow jobs to be scheduled on auto-scaled pools even if no agents are online

Adds basic testing to JobTaskSource as well.

diff --git a/Engine/Source/Programs/Horde/Horde.Build.Tests/Fixture.cs b/Engine/Source/Programs/Horde/Horde.Build.Tests/Fixture.cs
--- a/Engine/Source/Programs/Horde/Horde.Build.Tests/Fixture.cs
+++ b/Engine/Source/Programs/Horde/Horde.Build.Tests/Fixture.cs
@@ -35,6 +35,7 @@ namespace HordeServerTests
 		public string Job1ArtifactData { get; private set; } = null!;
 		public IAgent Agent1 { get; private set; } = null!;
 		public string Agent1Name { get; private set; } = null!;
+		public string PoolName = "TestingPool";
 
 		public static async Task<Fixture> Create(IGraphCollection GraphCollection, ITemplateCollection TemplateCollection, JobService JobService, IArtifactCollection ArtifactCollection, StreamService StreamService, AgentService AgentService, IPerforceService PerforceService)
 		{
@@ -69,6 +70,11 @@ namespace HordeServerTests
 			List<CreateStreamTabRequest> Tabs = new List<CreateStreamTabRequest>();
 			Tabs.Add(new CreateJobsTabRequest { Title = "foo", Templates = new List<string> { TemplateRefId1.ToString(), TemplateRefId2.ToString() } });
 
+			Dictionary<string, CreateAgentTypeRequest> AgentTypes = new()
+			{
+				{ "Win64", new() { Pool = PoolName} }
+			};
+
 			Stream = await StreamService.StreamCollection.GetAsync(new StreamId("ue5-main"));
 			Stream = await StreamService.StreamCollection.TryCreateOrReplaceAsync(
 				new StreamId("ue5-main"),
@@ -76,7 +82,7 @@ namespace HordeServerTests
 				"",
 				"",
 				new ProjectId("does-not-exist"),
-				new StreamConfig { Name = "//UE5/Main", Tabs = Tabs, Templates = Templates }
+				new StreamConfig { Name = "//UE5/Main", Tabs = Tabs, Templates = Templates, AgentTypes = AgentTypes}
 			);
 			
 			Job1 = await JobService.CreateJobAsync(
diff --git a/Engine/Source/Programs/Horde/Horde.Build.Tests/JobTaskSourceTests.cs b/Engine/Source/Programs/Horde/Horde.Build.Tests/JobTaskSourceTests.cs
new file mode 100644
--- /dev/null
+++ b/Engine/Source/Programs/Horde/Horde.Build.Tests/JobTaskSourceTests.cs
@@ -0,0 +1,85 @@
+// Copyright Epic Games, Inc. All Rights Reserved.
+
+using HordeCommon;
+using HordeServer.Api;
+using HordeServer.Models;
+using Microsoft.VisualStudio.TestTools.UnitTesting;
+using System.Collections.Generic;
+using System.Threading;
+using System.Threading.Tasks;
+using HordeServer.Utilities;
+
+namespace HordeServerTests
+{
+	using JobId = ObjectId<IJob>;
+	using LogId = ObjectId<ILogFile>;
+
+	[TestClass]
+	public class JobTaskSourceTests : TestSetup
+	{
+		[TestMethod]
+		public async Task UpdateJobQueueNormal()
+		{
+			Fixture Fixture = await SetupPoolWithAgentAsync(IsPoolAutoScaled: true, ShouldCreateAgent: true, IsAgentEnabled: true);
+
+			Assert.AreEqual(0, JobTaskSource.GetQueueForTesting().Count);
+			await JobTaskSource.TickAsync(CancellationToken.None);
+			Assert.AreEqual(1, JobTaskSource.GetQueueForTesting().Count);
+			Assert.AreEqual(Fixture.Job1.Id, JobTaskSource.GetQueueForTesting().Min!.Id.Item1);
+			Assert.AreEqual(JobStepBatchState.Ready, JobTaskSource.GetQueueForTesting().Min!.Batch.State);
+		}
+		
+		[TestMethod]
+		public async Task UpdateJobQueueWithNoAgentsInPool()
+		{
+			Fixture Fixture = await SetupPoolWithAgentAsync(IsPoolAutoScaled: true, ShouldCreateAgent: false, IsAgentEnabled: false);
+			
+			Assert.AreEqual(0, JobTaskSource.GetQueueForTesting().Count);
+			await JobTaskSource.TickAsync(CancellationToken.None);
+			Assert.AreEqual(0, JobTaskSource.GetQueueForTesting().Count);
+
+			IJob Job = (await JobService.GetJobAsync(Fixture.Job1.Id))!;
+			Assert.AreEqual(JobStepBatchError.NoAgentsInPool, Job.Batches[0].Error);
+		}
+		
+		[TestMethod]
+		public async Task UpdateJobQueueWithNoAgentsOnlineInPool()
+		{
+			Fixture Fixture = await SetupPoolWithAgentAsync(IsPoolAutoScaled: false, ShouldCreateAgent: true, IsAgentEnabled: false);
+			
+			Assert.AreEqual(0, JobTaskSource.GetQueueForTesting().Count);
+			await JobTaskSource.TickAsync(CancellationToken.None);
+			Assert.AreEqual(0, JobTaskSource.GetQueueForTesting().Count);
+
+			IJob Job = (await JobService.GetJobAsync(Fixture.Job1.Id))!;
+			Assert.AreEqual(JobStepBatchError.NoAgentsOnline, Job.Batches[0].Error);
+		}
+		
+		[TestMethod]
+		public async Task UpdateJobQueueWithNoAgentsOnlineInAutoScaledPool()
+		{
+			Fixture Fixture = await SetupPoolWithAgentAsync(IsPoolAutoScaled: true, ShouldCreateAgent: true, IsAgentEnabled: false);
+			
+			Assert.AreEqual(0, JobTaskSource.GetQueueForTesting().Count);
+			await JobTaskSource.TickAsync(CancellationToken.None);
+			Assert.AreEqual(1, JobTaskSource.GetQueueForTesting().Count);
+
+			Assert.AreEqual(Fixture.Job1.Id, JobTaskSource.GetQueueForTesting().Min!.Id.Item1);
+			Assert.AreEqual(JobStepBatchState.Ready, JobTaskSource.GetQueueForTesting().Min!.Batch.State);
+		}
+
+		private async Task<Fixture> SetupPoolWithAgentAsync(bool IsPoolAutoScaled, bool ShouldCreateAgent, bool IsAgentEnabled)
+		{
+			Fixture Fixture = await CreateFixtureAsync();
+			IPool Pool = await PoolService.CreatePoolAsync(Fixture.PoolName, null, IsPoolAutoScaled, 0, 0);
+
+			if (ShouldCreateAgent)
+			{
+				IAgent? Agent = await AgentService.CreateAgentAsync("TestAgent", IsAgentEnabled, null, new List<StringId<IPool>> { Pool.Id });
+				await AgentService.CreateSessionAsync(Agent, AgentStatus.Ok, new List<string>(), new Dictionary<string, int>(), null);
+			}
+
+			return Fixture;
+		}
+	}
+}
diff --git a/Engine/Source/Programs/Horde/Horde.Build/Jobs/Impl/JobTaskSource.cs b/Engine/Source/Programs/Horde/Horde.Build/Jobs/Impl/JobTaskSource.cs
--- a/Engine/Source/Programs/Horde/Horde.Build/Jobs/Impl/JobTaskSource.cs
+++ b/Engine/Source/Programs/Horde/Horde.Build/Jobs/Impl/JobTaskSource.cs
@@ -50,7 +50,7 @@ namespace HordeServer.Tasks.Impl
 		/// <summary>
 		/// An item in the queue to be executed
 		/// </summary>
-		class QueueItem
+		internal class QueueItem
 		{
 			/// <summary>
 			/// The stream for this job
@@ -305,7 +305,7 @@ namespace HordeServer.Tasks.Impl
 		/// </summary>
 		/// <param name="StoppingToken">Token that indicates that the service should shut down</param>
 		/// <returns>Async task</returns>
-		async ValueTask TickAsync(CancellationToken StoppingToken)
+		internal async ValueTask TickAsync(CancellationToken StoppingToken)
 		{
 			// Set the NewBatchIdToQueueItem member, so we capture any updated jobs during the DB query.
 			lock (LockObject)
@@ -348,6 +348,15 @@ namespace HordeServer.Tasks.Impl
 			// New list of queue items
 			SortedSet<QueueItem> NewQueue = new SortedSet<QueueItem>(Queue.Comparer);
 			Dictionary<(JobId, SubResourceId), QueueItem> NewBatchIdToQueueItem = new Dictionary<(JobId, SubResourceId), QueueItem>();
+			
+			bool HasAgentsOnline(PoolId PoolId)
+			{
+				IPool? Pool = Pools.Find(p => p.Id == PoolId);
+				bool IsPoolOnline = OnlinePools.Contains(PoolId);
+				bool IsPoolAutoScaled = ValidPools.Contains(PoolId) && Pool != null && Pool.EnableAutoscaling;
+				// If pool is auto-scaled, it will be considered online even if it has no agents online
+				return IsPoolOnline || IsPoolAutoScaled;
+			}
 
 			// Query for a new list of jobs for the queue
 			List<IJob> NewJobs = await Jobs.GetDispatchQueueAsync();
@@ -404,7 +413,7 @@ namespace HordeServer.Tasks.Impl
 					{
 						NewJob = await SkipBatchAsync(NewJob, Batch.Id, Graph, JobStepBatchError.NoAgentsInPool);
 					}
-					else if (!OnlinePools.Contains(AgentType.Pool))
+					else if (!HasAgentsOnline(AgentType.Pool))
 					{
 						NewJob = await SkipBatchAsync(NewJob, Batch.Id, Graph, JobStepBatchError.NoAgentsOnline);
 					}
@@ -475,6 +484,18 @@ namespace HordeServer.Tasks.Impl
 			return NewJob;
 		}
 
+		/// <summary>
+		/// Get the queue, for internal testing only
+		/// </summary>
+		/// <returns>A copy of the queue</returns>
+		internal SortedSet<QueueItem> GetQueueForTesting()
+		{
+			lock (LockObject)
+			{
+				return new SortedSet<QueueItem>(Queue);
+			}
+		}
+
 		/// <summary>
 		/// Updates the current state of a job
 		/// </summary>
