Horde: Fix infinite loop when enumerating agent properties.

diff --git a/Engine/Source/Programs/Horde/Horde.Build/Agents/IAgent.cs b/Engine/Source/Programs/Horde/Horde.Build/Agents/IAgent.cs
--- a/Engine/Source/Programs/Horde/Horde.Build/Agents/IAgent.cs
+++ b/Engine/Source/Programs/Horde/Horde.Build/Agents/IAgent.cs
@@ -610,7 +610,7 @@ namespace HordeServer.Models
 				if (Index < 0)
 				{
 					Index = ~Index;
-					while (Index < Agent.Properties.Count)
+					for (; Index < Agent.Properties.Count; Index++)
 					{
 						string Property = Agent.Properties[Index];
 						if (Property.Length <= Name.Length || !Property.StartsWith(Name, StringComparison.OrdinalIgnoreCase) || Property[Name.Length] != '=')
diff --git a/Engine/Source/Programs/Horde/Horde.Build/Agents/Impl/AgentService.cs b/Engine/Source/Programs/Horde/Horde.Build/Agents/Impl/AgentService.cs
--- a/Engine/Source/Programs/Horde/Horde.Build/Agents/Impl/AgentService.cs
+++ b/Engine/Source/Programs/Horde/Horde.Build/Agents/Impl/AgentService.cs
@@ -802,10 +802,9 @@ namespace HordeServer.Services
 		/// </summary>
 		/// <param name="AgentId">Agent id to query</param>
 		/// <returns>Hourly rate of running the given agent</returns>
-		public ValueTask<double?> GetRateAsync(AgentId AgentId)
+		public async ValueTask<double?> GetRateAsync(AgentId AgentId)
 		{
-			return new ValueTask<double?>(Task.FromResult<double?>(null));
-/*			RedisKey Key = $"agent-rate/{AgentId}";
+			RedisKey Key = $"agent-rate/{AgentId}";
 
 			// Try to get the current value
 			RedisValue Value = await RedisService.Database.StringGetAsync(Key);
@@ -846,7 +845,7 @@ namespace HordeServer.Services
 				// Cache it for future reference
 				await RedisService.Database.StringSetAsync(Key, Rate, TimeSpan.FromMinutes(5.0), flags: CommandFlags.FireAndForget);
 				return Rate;
-			}*/
+			}
 		}
 
 		/// <summary>
diff --git a/Engine/Source/Programs/Horde/Horde.Build/Utilities/AsyncCachedValue.cs b/Engine/Source/Programs/Horde/Horde.Build/Utilities/AsyncCachedValue.cs
--- a/Engine/Source/Programs/Horde/Horde.Build/Utilities/AsyncCachedValue.cs
+++ b/Engine/Source/Programs/Horde/Horde.Build/Utilities/AsyncCachedValue.cs
@@ -77,32 +77,27 @@ namespace HordeServer.Utilities
 			return GetAsync(MaxRefreshTime);
 		}
 
-		Task<State> CreateUpdateTask(ref Task<State>? UpdateTask)
+		Task<State> CreateOrGetStateTask(ref Task<State>? StateTask)
 		{
-			Task<State>? CurrentTask = UpdateTask;
-			while (CurrentTask == null)
+			for(; ;)
 			{
-				TaskCompletionSource<State> NextTaskSource = new TaskCompletionSource<State>();
-				if (Interlocked.CompareExchange(ref UpdateTask, NextTaskSource.Task, null) == null)
+				Task<State>? CurrentStateTask = StateTask;
+				if (CurrentStateTask != null)
 				{
-					Task.Run(() => GenerateNewValueAsync(NextTaskSource));
-					return NextTaskSource.Task;
+					return CurrentStateTask;
+				}
+
+				Task<Task<State>> InnerNewStateTask = new Task<Task<State>>(() => CreateState());
+				if (Interlocked.CompareExchange(ref StateTask, InnerNewStateTask.Unwrap(), null) == null)
+				{
+					InnerNewStateTask.Start();
 				}
-				CurrentTask = UpdateTask;
 			}
-			return CurrentTask;
 		}
 
-		async Task GenerateNewValueAsync(TaskCompletionSource<State> NextTaskSource)
+		async Task<State> CreateState()
 		{
-			try
-			{
-				NextTaskSource.SetResult(new State(await Generator()));
-			}
-			catch (Exception Ex)
-			{
-				NextTaskSource.SetException(Ex);
-			}
+			return new State(await Generator());
 		}
 
 		/// <summary>
@@ -111,25 +106,24 @@ namespace HordeServer.Utilities
 		/// <returns>The cached value, if valid</returns>
 		public async Task<T> GetAsync(TimeSpan MaxAge)
 		{
-			Task<State> CurrentCopy = CreateUpdateTask(ref Current);
+			Task<State> StateTask = CreateOrGetStateTask(ref Current);
 
-			State CurrentState = await CurrentCopy;
-			if (CurrentState.Elapsed > MaxAge)
+			State State = await StateTask;
+			if (State.Next != null && State.Next.IsCompleted)
 			{
-				CurrentState = await CreateUpdateTask(ref CurrentState.Next);
-				return CurrentState.Value;
+				_ = Interlocked.CompareExchange(ref Current, State.Next, StateTask);
+				State = await State.Next;
 			}
-			if (CurrentState.Elapsed > MinRefreshTime)
+			if (State.Elapsed > MaxAge)
 			{
-				_ = CreateUpdateTask(ref CurrentState.Next);
+				State = await CreateOrGetStateTask(ref State.Next);
 			}
-			if (CurrentState.Next != null && CurrentState.Next.IsCompleted)
+			if (State.Elapsed > MinRefreshTime)
 			{
-				_ = Interlocked.CompareExchange(ref Current, CurrentState.Next, CurrentCopy);
-				return (await CurrentState.Next).Value;
+				_ = CreateOrGetStateTask(ref State.Next);
 			}
 
-			return CurrentState.Value;
+			return State.Value;
 		}
 	}
 }
