Horde Storage

Fixes for compact binary serialization to json for child objects and arrays, as well as positive integers and object ids.

This should resolve a bunch of cases were we are unable to convert a compact binary to json.

diff --git a/Engine/Source/Programs/Horde/HordeStorage/Horde.Storage/Tests/UnitTests/CompactBinaryTests.cs b/Engine/Source/Programs/Horde/HordeStorage/Horde.Storage/Tests/UnitTests/CompactBinaryTests.cs
--- a/Engine/Source/Programs/Horde/HordeStorage/Horde.Storage/Tests/UnitTests/CompactBinaryTests.cs
+++ b/Engine/Source/Programs/Horde/HordeStorage/Horde.Storage/Tests/UnitTests/CompactBinaryTests.cs
@@ -64,6 +64,26 @@ namespace EuropaUnit
             Assert.AreEqual(new BlobIdentifier("c7a03f83c08cdca882110ecf2b5654ee3b09b11e"), payloadFields[2]["RawHash"]!.AsHash());
         }
 
+        [TestMethod]
+        public void compact_binary()
+        {
+            byte[] bytes = File.ReadAllBytes("CompactBinaryObjects/compact_binary");
+
+            BlobIdentifier bi = BlobIdentifier.FromBlob(bytes);
+            ReadOnlyMemory<byte> memory = new ReadOnlyMemory<byte>(bytes);
+            CompactBinaryObject o = CompactBinaryObject.Load(ref memory);
+
+            Assert.AreEqual("", o.Name);
+            List<CompactBinaryField> buildActionFields = o.GetFields().ToList();
+            Assert.AreEqual(3, buildActionFields.Count);
+            CompactBinaryField payloads = buildActionFields[0];
+            List<CompactBinaryField> payloadFields = payloads.GetFields().ToList();
+            Assert.AreEqual(2, payloadFields.Count);
+
+            Assert.AreEqual("{\"Key\":{\"Bucket\":\"EditorDomainPackage\"},\"Meta\":{\"FileSize\":24789},\"Attachments\":[{\"Id\":\"000000000000000000000001\"}]}", o.ToJson());
+        }
+
+
 
         [TestMethod]
         public void WriteArray()
diff --git a/Engine/Source/Programs/Horde/HordeStorage/Horde.Storage/Tests/UnitTests/Horde.Storage.Tests.Unit.csproj b/Engine/Source/Programs/Horde/HordeStorage/Horde.Storage/Tests/UnitTests/Horde.Storage.Tests.Unit.csproj
--- a/Engine/Source/Programs/Horde/HordeStorage/Horde.Storage/Tests/UnitTests/Horde.Storage.Tests.Unit.csproj
+++ b/Engine/Source/Programs/Horde/HordeStorage/Horde.Storage/Tests/UnitTests/Horde.Storage.Tests.Unit.csproj
@@ -30,6 +30,9 @@
     <None Update="CompactBinaryObjects\build">
       <CopyToOutputDirectory>PreserveNewest</CopyToOutputDirectory>
     </None>
+    <None Update="CompactBinaryObjects\compact_binary">
+      <CopyToOutputDirectory>PreserveNewest</CopyToOutputDirectory>
+    </None>
     <None Update="CompactBinaryObjects\ReferenceOutput">
       <CopyToOutputDirectory>PreserveNewest</CopyToOutputDirectory>
     </None>
diff --git a/Engine/Source/Programs/Horde/HordeStorage/Jupiter.Common/Implementation/CompactBinaryObject.cs b/Engine/Source/Programs/Horde/HordeStorage/Jupiter.Common/Implementation/CompactBinaryObject.cs
--- a/Engine/Source/Programs/Horde/HordeStorage/Jupiter.Common/Implementation/CompactBinaryObject.cs
+++ b/Engine/Source/Programs/Horde/HordeStorage/Jupiter.Common/Implementation/CompactBinaryObject.cs
@@ -20,7 +20,7 @@ namespace Jupiter.Implementation
         private const CompactBinaryFieldType TypeMask              = (CompactBinaryFieldType)0b_0001_1111;
 
         private const CompactBinaryFieldType ObjectMask             = (CompactBinaryFieldType)0b_0001_1110;
-        private const CompactBinaryFieldType ObjectBase             = (CompactBinaryFieldType)0b_0001_0010;
+        private const CompactBinaryFieldType ObjectBase             = (CompactBinaryFieldType)0b_0000_0010;
         
         private const CompactBinaryFieldType ArrayMask             = (CompactBinaryFieldType)0b_0001_1110;
         private const CompactBinaryFieldType ArrayBase             = (CompactBinaryFieldType)0b_0000_0100;
@@ -58,6 +58,11 @@ namespace Jupiter.Implementation
             return (type & IntegerMask) == IntegerBase;
         }
 
+        public static bool IsPositiveInteger(CompactBinaryFieldType type)
+        {
+            return IsInteger(type) && RemoveFlags(type) == CompactBinaryFieldType.IntegerPositive;
+        }
+
         public static bool IsFloat(CompactBinaryFieldType type)
         {
             return (type & FloatMask) == FloatBase;
@@ -108,6 +113,11 @@ namespace Jupiter.Implementation
             return RemoveFlags(type) == CompactBinaryFieldType.Binary;
         }
 
+        public static bool IsObjectId(CompactBinaryFieldType type)
+        {
+            return RemoveFlags(type) == CompactBinaryFieldType.ObjectId;
+        }
+
         public static bool HasUniformFields(CompactBinaryFieldType type)
         {
             CompactBinaryFieldType localType = RemoveFlags(type);
@@ -470,6 +480,18 @@ namespace Jupiter.Implementation
             return null;
         }
 
+        public uint? AsPositiveInteger()
+        {
+            if (CompactBinaryFieldUtils.IsPositiveInteger(_type))
+            {                
+                ReadOnlyMemory<byte> localPayload = _payload.Slice(0);
+                ulong integer = BitUtils.ReadVarUInt(ref localPayload, out int magnitudeByteCount);
+                return (uint)integer;
+            }
+
+            return null;
+        }
+
         public bool? AsBool()
         {
             if (CompactBinaryFieldUtils.IsBool(_type))
@@ -507,8 +529,12 @@ namespace Jupiter.Implementation
         
         public CompactBinaryObject AsObject()
         {
-            ReadOnlyMemory<byte> localMemory = _payload.Slice(0);
-            return CompactBinaryObject.Load(ref localMemory);
+            ReadOnlyMemory<byte> localPayload = _payload.Slice(0);
+
+            ulong payloadSize = BitUtils.ReadVarUInt(ref localPayload, out int payloadBytesRead);
+            localPayload = localPayload.Slice(payloadBytesRead);
+
+            return CompactBinaryObject.Load(ref localPayload);
         }
 
         public Guid AsUuid(Guid @default = default(Guid))
@@ -521,6 +547,11 @@ namespace Jupiter.Implementation
             return @default;
         }
 
+        public byte[]? AsObjectId()
+        {
+            return _payload.ToArray();
+        }
+
         public bool IsObject()
         {
             return CompactBinaryFieldUtils.IsObject(_type);
@@ -535,7 +566,10 @@ namespace Jupiter.Implementation
         {
             return CompactBinaryFieldUtils.IsInteger(_type);
         }
-
+        public bool IsPositiveInteger()
+        {
+            return CompactBinaryFieldUtils.IsPositiveInteger(_type);
+        }
         public bool IsDateTime()
         {
             return CompactBinaryFieldUtils.IsDateTime(_type);
@@ -581,6 +615,11 @@ namespace Jupiter.Implementation
             return CompactBinaryFieldUtils.IsAttachment(_type);
         }
 
+        public bool IsObjectId()
+        {
+            return CompactBinaryFieldUtils.IsObjectId(_type);
+        }
+
         public int GetObjectLength()
         {
             return _objectSize;
@@ -832,23 +871,40 @@ namespace Jupiter.Implementation
             return Encoding.UTF8.GetString(buffer.WrittenMemory.Span);
         }
 
-        private static void WriteField(CompactBinaryField field, Utf8JsonWriter jsonWriter)
+        private static void WriteField(CompactBinaryField field, Utf8JsonWriter jsonWriter, bool addPropertyName = true)
         {
             if (field.IsObject())
             {
+                if (addPropertyName)
+                {
+                    jsonWriter.WritePropertyName(field.Name);
+                }
                 jsonWriter.WriteStartObject();
                 CompactBinaryObject o = field.AsObject();
+                WriteField(o, jsonWriter);
+                jsonWriter.WriteEndObject();
+                /*jsonWriter.WritePropertyName(o.Name);
+                jsonWriter.WriteStartObject();
                 foreach (CompactBinaryField f in o.GetFields())
                 {
                     WriteField(f, jsonWriter);
                 }
-                jsonWriter.WriteEndObject();
+                jsonWriter.WriteEndObject();*/
+
             }
             else if (field.IsArray())
             {
-                jsonWriter.WriteStartArray();
+                jsonWriter.WriteStartArray(field.Name);
+                foreach (CompactBinaryField f in field.AsArray())
+                {
+                    WriteField(f, jsonWriter, addPropertyName: false);
+                }
                 jsonWriter.WriteEndArray();
             }
+            else if (field.IsPositiveInteger())
+            {
+                jsonWriter.WriteNumber(field.Name, field.AsPositiveInteger()!.Value);
+            }
             else if (field.IsInteger())
             {
                 jsonWriter.WriteNumber(field.Name, field.AsInteger()!.Value);
@@ -873,6 +929,10 @@ namespace Jupiter.Implementation
             {
                 jsonWriter.WriteString(field.Name, field.AsString());
             }
+            else if (field.IsObjectId())
+            {
+                jsonWriter.WriteString(field.Name, StringUtils.FormatAsHexString(field.AsObjectId()!));
+            }
             else
             {
                 throw new NotImplementedException($"Not handled type {field.Type} when attempting to convert to json");
