Horde Storage

Converted some scylla ad hoc statements to prepared statements, should perform slightly better.
Kept some none critical statements as ad hoc as its easier to read the code that way.

diff --git a/Engine/Source/Programs/Horde/HordeStorage/Horde.Storage/Implementation/Objects/ScyllaReferencesStore.cs b/Engine/Source/Programs/Horde/HordeStorage/Horde.Storage/Implementation/Objects/ScyllaReferencesStore.cs
--- a/Engine/Source/Programs/Horde/HordeStorage/Horde.Storage/Implementation/Objects/ScyllaReferencesStore.cs
+++ b/Engine/Source/Programs/Horde/HordeStorage/Horde.Storage/Implementation/Objects/ScyllaReferencesStore.cs
@@ -20,6 +20,8 @@ namespace Horde.Storage.Implementation
         private readonly IMapper _mapper;
         private readonly IOptionsMonitor<ScyllaSettings> _settings;
         private readonly ILogger _logger = Log.ForContext<ScyllaReferencesStore>();
+        private readonly PreparedStatement _getObjectsStatement;
+        private readonly PreparedStatement _getNamespacesStatement;
 
         public ScyllaReferencesStore(IScyllaSessionManager scyllaSessionManager, IOptionsMonitor<ScyllaSettings> settings)
         {
@@ -57,6 +59,8 @@ namespace Horde.Storage.Implementation
             ));
 
 
+            _getObjectsStatement = _session.Prepare("SELECT bucket, name, last_access_time FROM objects WHERE namespace = ? ALLOW FILTERING");
+            _getNamespacesStatement = _session.Prepare("SELECT DISTINCT namespace FROM buckets");
         }
 
         public async Task<ObjectRecord> Get(NamespaceId ns, BucketId bucket, IoHashKey name, IReferencesStore.FieldFlags flags)
@@ -218,7 +222,7 @@ namespace Horde.Storage.Implementation
         {
             using Scope _ = Tracer.Instance.StartActive("scylla.get_records");
 
-            RowSet rowSet = await _session.ExecuteAsync(new SimpleStatement("SELECT bucket, name, last_access_time FROM objects WHERE namespace = ? ALLOW FILTERING", ns.ToString()));
+            RowSet rowSet = await _session.ExecuteAsync(_getObjectsStatement.Bind(ns.ToString()));
             foreach (Row row in rowSet)
             {
                 if (rowSet.GetAvailableWithoutFetching() == 0)
@@ -238,7 +242,7 @@ namespace Horde.Storage.Implementation
         public async IAsyncEnumerable<NamespaceId> GetNamespaces()
         {
             using Scope _ = Tracer.Instance.StartActive("scylla.get_namespaces");
-            RowSet rowSet = await _session.ExecuteAsync(new SimpleStatement("SELECT DISTINCT namespace FROM buckets"));
+            RowSet rowSet = await _session.ExecuteAsync(_getNamespacesStatement.Bind());
 
             foreach (Row row in rowSet)
             {
diff --git a/Engine/Source/Programs/Horde/HordeStorage/Horde.Storage/Implementation/TransactionLog/ScyllaReplicationLog.cs b/Engine/Source/Programs/Horde/HordeStorage/Horde.Storage/Implementation/TransactionLog/ScyllaReplicationLog.cs
--- a/Engine/Source/Programs/Horde/HordeStorage/Horde.Storage/Implementation/TransactionLog/ScyllaReplicationLog.cs
+++ b/Engine/Source/Programs/Horde/HordeStorage/Horde.Storage/Implementation/TransactionLog/ScyllaReplicationLog.cs
@@ -18,6 +18,7 @@ namespace Horde.Storage.Implementation
         private readonly IOptionsMonitor<ScyllaSettings> _settings;
         private readonly ISession _session;
         private readonly Mapper _mapper;
+        private readonly PreparedStatement _selectReplicationBucketStatement;
 
         public ScyllaReplicationLog(IScyllaSessionManager scyllaSessionManager, IOptionsMonitor<ScyllaSettings> settings)
         {
@@ -51,6 +52,9 @@ namespace Horde.Storage.Implementation
                 PRIMARY KEY ((namespace))
             );"  
             ));
+
+            _selectReplicationBucketStatement = _session.Prepare(
+                "SELECT replication_bucket FROM replication_log WHERE namespace = ? PER PARTITION LIMIT 1 ALLOW FILTERING");
         }
 
         public async IAsyncEnumerable<NamespaceId> GetNamespaces()
@@ -173,7 +177,7 @@ namespace Horde.Storage.Implementation
                 SortedSet<long> buckets = new SortedSet<long>();
 
                 // fetch all the buckets that exists and sort them based on time
-                RowSet replicationBuckets = await _session.ExecuteAsync(new SimpleStatement("SELECT replication_bucket FROM replication_log WHERE namespace = ? PER PARTITION LIMIT 1 ALLOW FILTERING", ns.ToString()));
+                RowSet replicationBuckets = await _session.ExecuteAsync(_selectReplicationBucketStatement.Bind(ns.ToString()));
                 foreach (Row replicationBucket in replicationBuckets)
                 {
                     long bucketField = (long)replicationBucket["replication_bucket"];
