Horde Storage

Another fix for json serialization of compact binaries

diff --git a/Engine/Source/Programs/Horde/HordeStorage/Horde.Storage/Tests/UnitTests/CompactBinaryTests.cs b/Engine/Source/Programs/Horde/HordeStorage/Horde.Storage/Tests/UnitTests/CompactBinaryTests.cs
--- a/Engine/Source/Programs/Horde/HordeStorage/Horde.Storage/Tests/UnitTests/CompactBinaryTests.cs
+++ b/Engine/Source/Programs/Horde/HordeStorage/Horde.Storage/Tests/UnitTests/CompactBinaryTests.cs
@@ -80,7 +80,7 @@ namespace EuropaUnit
             List<CompactBinaryField> payloadFields = payloads.GetFields().ToList();
             Assert.AreEqual(2, payloadFields.Count);
 
-            Assert.AreEqual("{\"Key\":{\"Bucket\":\"EditorDomainPackage\"},\"Meta\":{\"FileSize\":24789},\"Attachments\":[{\"Id\":\"000000000000000000000001\"}]}", o.ToJson());
+            Assert.AreEqual("{\"Key\":{\"Bucket\":\"EditorDomainPackage\",\"Hash\":\"37DBAA409EF30BA67F18C8FC2FAAF606636CB915\"},\"Meta\":{\"FileSize\":24789},\"Attachments\":[{\"Id\":\"000000000000000000000001\",\"RawHash\":\"DA6FC57E4B9F91377C9509EA0AD567BACB3796C5\",\"RawSize\":24789}]}", o.ToJson());
         }
 
 
diff --git a/Engine/Source/Programs/Horde/HordeStorage/Jupiter.Common/Implementation/CompactBinaryObject.cs b/Engine/Source/Programs/Horde/HordeStorage/Jupiter.Common/Implementation/CompactBinaryObject.cs
--- a/Engine/Source/Programs/Horde/HordeStorage/Jupiter.Common/Implementation/CompactBinaryObject.cs
+++ b/Engine/Source/Programs/Horde/HordeStorage/Jupiter.Common/Implementation/CompactBinaryObject.cs
@@ -876,21 +876,15 @@ namespace Jupiter.Implementation
             if (field.IsObject())
             {
                 if (addPropertyName)
-                {
                     jsonWriter.WritePropertyName(field.Name);
-                }
-                jsonWriter.WriteStartObject();
-                CompactBinaryObject o = field.AsObject();
-                WriteField(o, jsonWriter);
-                jsonWriter.WriteEndObject();
-                /*jsonWriter.WritePropertyName(o.Name);
                 jsonWriter.WriteStartObject();
-                foreach (CompactBinaryField f in o.GetFields())
+                
+                foreach (CompactBinaryField f in field.GetFields())
                 {
                     WriteField(f, jsonWriter);
                 }
-                jsonWriter.WriteEndObject();*/
 
+                jsonWriter.WriteEndObject();
             }
             else if (field.IsArray())
             {
