Horde: Use the number of stopped instances in an ASG to populate the hibernating time.

diff --git a/Engine/Source/Programs/Horde/Horde.Build/Fleet/AutoscaleService.cs b/Engine/Source/Programs/Horde/Horde.Build/Fleet/AutoscaleService.cs
--- a/Engine/Source/Programs/Horde/Horde.Build/Fleet/AutoscaleService.cs
+++ b/Engine/Source/Programs/Horde/Horde.Build/Fleet/AutoscaleService.cs
@@ -233,7 +233,7 @@ namespace HordeServer.Services
 
 							if (Delta > 0)
 							{
-								await FleetManager.ExpandPool(Pool, PoolData.Agents, Delta);
+								await FleetManager.ExpandPoolAsync(Pool, PoolData.Agents, Delta);
 								await PoolCollection.TryUpdateAsync(Pool, LastScaleUpTime: DateTime.UtcNow);
 							}
 
@@ -242,7 +242,7 @@ namespace HordeServer.Services
 								bool bShrinkIsOnCoolDown = Pool.LastScaleDownTime != null && Pool.LastScaleDownTime + ShrinkPoolCoolDown > DateTime.UtcNow;
 								if (!bShrinkIsOnCoolDown)
 								{
-									await FleetManager.ShrinkPool(Pool, PoolData.Agents, -Delta);
+									await FleetManager.ShrinkPoolAsync(Pool, PoolData.Agents, -Delta);
 									await PoolCollection.TryUpdateAsync(Pool, LastScaleDownTime: DateTime.UtcNow);
 								}
 								else
diff --git a/Engine/Source/Programs/Horde/Horde.Build/Fleet/IFleetManager.cs b/Engine/Source/Programs/Horde/Horde.Build/Fleet/IFleetManager.cs
--- a/Engine/Source/Programs/Horde/Horde.Build/Fleet/IFleetManager.cs
+++ b/Engine/Source/Programs/Horde/Horde.Build/Fleet/IFleetManager.cs
@@ -20,7 +20,7 @@ namespace HordeServer.Services
 		/// <param name="Agents">Current list of agents in the pool</param>
 		/// <param name="Count">Number of agents to add</param>
 		/// <returns>Async task</returns>
-		Task ExpandPool(IPool Pool, IReadOnlyList<IAgent> Agents, int Count);
+		Task ExpandPoolAsync(IPool Pool, IReadOnlyList<IAgent> Agents, int Count);
 
 		/// <summary>
 		/// Shrink the given pool
@@ -29,7 +29,14 @@ namespace HordeServer.Services
 		/// <param name="Agents">Current list of agents in the pool</param>
 		/// <param name="Count">Number of agents to remove</param>
 		/// <returns>Async task</returns>
-		Task ShrinkPool(IPool Pool, IReadOnlyList<IAgent> Agents, int Count);
+		Task ShrinkPoolAsync(IPool Pool, IReadOnlyList<IAgent> Agents, int Count);
+
+		/// <summary>
+		/// Returns the number of stopped instances in the given pool
+		/// </summary>
+		/// <param name="Pool">Pool to resize</param>
+		/// <returns>Async task</returns>
+		Task<int> GetNumStoppedInstancesAsync(IPool Pool);
 	}
 }
 
diff --git a/Engine/Source/Programs/Horde/Horde.Build/Fleet/Impl/Managers/Fleet/AwsFleetManager.cs b/Engine/Source/Programs/Horde/Horde.Build/Fleet/Impl/Managers/Fleet/AwsFleetManager.cs
--- a/Engine/Source/Programs/Horde/Horde.Build/Fleet/Impl/Managers/Fleet/AwsFleetManager.cs
+++ b/Engine/Source/Programs/Horde/Horde.Build/Fleet/Impl/Managers/Fleet/AwsFleetManager.cs
@@ -48,7 +48,7 @@ namespace HordeServer.Services.Impl
 		}
 
 		/// <inheritdoc/>
-		public async Task ExpandPool(IPool Pool, IReadOnlyList<IAgent> Agents, int Count)
+		public async Task ExpandPoolAsync(IPool Pool, IReadOnlyList<IAgent> Agents, int Count)
 		{
 			// Find stopped instances in the correct pool
 			DescribeInstancesRequest DescribeRequest = new DescribeInstancesRequest();
@@ -75,7 +75,7 @@ namespace HordeServer.Services.Impl
 		}
 
 		/// <inheritdoc/>
-		public async Task ShrinkPool(IPool Pool, IReadOnlyList<IAgent> Agents, int Count)
+		public async Task ShrinkPoolAsync(IPool Pool, IReadOnlyList<IAgent> Agents, int Count)
 		{
 			string AwsTagProperty = $"{AwsTagPropertyName}={PoolTagName}:{Pool.Name}";
 
@@ -100,5 +100,18 @@ namespace HordeServer.Services.Impl
 				}
 			}
 		}
+
+		/// <inheritdoc/>
+		public async Task<int> GetNumStoppedInstancesAsync(IPool Pool)
+		{
+			// Find all instances in the pool
+			DescribeInstancesRequest DescribeRequest = new DescribeInstancesRequest();
+			DescribeRequest.Filters = new List<Filter>();
+			DescribeRequest.Filters.Add(new Filter("instance-state-name", new List<string> { InstanceStateName.Stopped.Value }));
+			DescribeRequest.Filters.Add(new Filter("tag:" + PoolTagName, new List<string> { Pool.Name }));
+
+			DescribeInstancesResponse DescribeResponse = await Client.DescribeInstancesAsync(DescribeRequest);
+			return DescribeResponse.Reservations.SelectMany(x => x.Instances).Select(x => x.InstanceId).Distinct().Count();
+		}
 	}
 }
diff --git a/Engine/Source/Programs/Horde/Horde.Build/Fleet/Impl/Managers/Fleet/DefaultFleetManager.cs b/Engine/Source/Programs/Horde/Horde.Build/Fleet/Impl/Managers/Fleet/DefaultFleetManager.cs
--- a/Engine/Source/Programs/Horde/Horde.Build/Fleet/Impl/Managers/Fleet/DefaultFleetManager.cs
+++ b/Engine/Source/Programs/Horde/Horde.Build/Fleet/Impl/Managers/Fleet/DefaultFleetManager.cs
@@ -26,17 +26,20 @@ namespace HordeServer.Services.Impl
 		}
 
 		/// <inheritdoc/>
-		public Task ExpandPool(IPool Pool, IReadOnlyList<IAgent> Agents, int Count)
+		public Task ExpandPoolAsync(IPool Pool, IReadOnlyList<IAgent> Agents, int Count)
 		{
 			Logger.LogInformation("Expand pool {PoolId} by {Count} agents", Pool.Id, Count);
 			return Task.CompletedTask;
 		}
 
 		/// <inheritdoc/>
-		public Task ShrinkPool(IPool Pool, IReadOnlyList<IAgent> Agents, int Count)
+		public Task ShrinkPoolAsync(IPool Pool, IReadOnlyList<IAgent> Agents, int Count)
 		{
 			Logger.LogInformation("Shrink pool {PoolId} by {Count} agents", Pool.Id, Count);
 			return Task.CompletedTask;
 		}
+
+		/// <inheritdoc/>
+		public Task<int> GetNumStoppedInstancesAsync(IPool Pool) => Task.FromResult(0);
 	}
 }
diff --git a/Engine/Source/Programs/Horde/Horde.Build/Fleet/TelemetryService.cs b/Engine/Source/Programs/Horde/Horde.Build/Fleet/TelemetryService.cs
--- a/Engine/Source/Programs/Horde/Horde.Build/Fleet/TelemetryService.cs
+++ b/Engine/Source/Programs/Horde/Horde.Build/Fleet/TelemetryService.cs
@@ -28,6 +28,7 @@ namespace HordeServer.Services
 		IAgentCollection AgentCollection;
 		ILeaseCollection LeaseCollection;
 		IPoolCollection PoolCollection;
+		IFleetManager FleetManager;
 		IClock Clock;
 		ITicker Tick;
 		ILogger Logger;
@@ -35,12 +36,13 @@ namespace HordeServer.Services
 		/// <summary>
 		/// Constructor
 		/// </summary>
-		public TelemetryService(ITelemetryCollection TelemetryCollection, IAgentCollection AgentCollection, ILeaseCollection LeaseCollection, IPoolCollection PoolCollection, IClock Clock, ILogger<TelemetryService> Logger)
+		public TelemetryService(ITelemetryCollection TelemetryCollection, IAgentCollection AgentCollection, ILeaseCollection LeaseCollection, IPoolCollection PoolCollection, IFleetManager FleetManager, IClock Clock, ILogger<TelemetryService> Logger)
 		{
 			this.TelemetryCollection = TelemetryCollection;
 			this.AgentCollection = AgentCollection;
 			this.LeaseCollection = LeaseCollection;
 			this.PoolCollection = PoolCollection;
+			this.FleetManager = FleetManager;
 			this.Clock = Clock;
 			this.Tick = Clock.AddSharedTicker<TelemetryService>(TimeSpan.FromMinutes(10.0), TickLeaderAsync, Logger);
 			this.Logger = Logger;
@@ -89,6 +91,16 @@ namespace HordeServer.Services
 
 				NewUtilizationTelemetry Telemetry = new NewUtilizationTelemetry(BucketMinTime, BucketMaxTime);
 				Telemetry.NumAgents = Agents.Count;
+				foreach (IPool Pool in Pools)
+				{
+					if (Pool.EnableAutoscaling)
+					{
+						int NumStoppedInstances = await FleetManager.GetNumStoppedInstancesAsync(Pool);
+
+						NewPoolUtilizationTelemetry PoolTelemetry = Telemetry.FindOrAddPool(Pool.Id);
+						PoolTelemetry.HibernatingTime += Interval.TotalHours * NumStoppedInstances;
+					}
+				}
 				foreach (PoolId PoolId in AgentToPoolIds.Values.SelectMany(x => x))
 				{
 					Telemetry.FindOrAddPool(PoolId).NumAgents++;
