Horde: Styling consistency pass and many other tweaks to layouts, consolidated some dark/light theme settings, improved handling of theme switching
#rnx

diff --git a/Engine/Source/Programs/Horde/HordeDashboard/src/backend/Dashboard.ts b/Engine/Source/Programs/Horde/HordeDashboard/src/backend/Dashboard.ts
--- a/Engine/Source/Programs/Horde/HordeDashboard/src/backend/Dashboard.ts
+++ b/Engine/Source/Programs/Horde/HordeDashboard/src/backend/Dashboard.ts
@@ -206,7 +206,7 @@ export class Dashboard {
         if (value) {
             localStorage?.setItem("horde_darktheme", "true");
         } else {
-            localStorage?.removeItem("horde_darktheme");
+            localStorage?.setItem("horde_darktheme", "false");
         }
 
         if (resetColors) {
@@ -436,7 +436,7 @@ export class Dashboard {
             this.preferences.set(pref, value);
         }
 
-        this.postPreferences();
+        this.postPreferences(pref === DashboardPreference.Darktheme);
 
     }
 
@@ -458,7 +458,7 @@ export class Dashboard {
     }
 
 
-    private async postPreferences(): Promise<boolean> {
+    private async postPreferences(reload?:boolean): Promise<boolean> {
 
         // cancel any pending        
         for (let i = 0; i < this.cancelId; i++) {
@@ -477,6 +477,11 @@ export class Dashboard {
         } catch (reason) {
             success = false;
             console.error("Error posting user preferences", reason)
+        } finally {
+            if (reload) {
+                window.location.reload();
+            }
+                
         }
 
         return success;
