Horde: Yet more style fixes
#rnx

diff --git a/Engine/Source/Programs/Horde/HordeDashboard/src/index.css b/Engine/Source/Programs/Horde/HordeDashboard/src/index.css
--- a/Engine/Source/Programs/Horde/HordeDashboard/src/index.css
+++ b/Engine/Source/Programs/Horde/HordeDashboard/src/index.css
@@ -5,7 +5,6 @@ body {
     sans-serif;
   -webkit-font-smoothing: antialiased;
   -moz-osx-font-smoothing: grayscale; 
-  color:#0F0F0FFF;
   overflow: hidden;
 }
 
@@ -44,7 +43,6 @@ code {
 
 .view-log-link, .view-log-link:visited, .view-log-link:hover {
   text-decoration: none;
-  color: rgb(48, 66, 97);
   user-select: none;
   cursor: pointer;
 }
