Horde: Add pool operations to agent registration token.

diff --git a/Engine/Source/Programs/Horde/Horde.Build/Acls/AclService.cs b/Engine/Source/Programs/Horde/Horde.Build/Acls/AclService.cs
--- a/Engine/Source/Programs/Horde/Horde.Build/Acls/AclService.cs
+++ b/Engine/Source/Programs/Horde/Horde.Build/Acls/AclService.cs
@@ -109,7 +109,7 @@ namespace HordeServer.Services
 
 			DefaultAcl = new Acl();
 			DefaultAcl.Entries.Add(new AclEntry(new AclClaim(ClaimTypes.Role, "internal:AgentRegistration"), new[] { AclAction.CreateAgent, AclAction.CreateSession }));
-			DefaultAcl.Entries.Add(new AclEntry(AgentRegistrationClaim, new[] { AclAction.CreateAgent, AclAction.CreateSession, AclAction.UpdateAgent, AclAction.DownloadSoftware }));
+			DefaultAcl.Entries.Add(new AclEntry(AgentRegistrationClaim, new[] { AclAction.CreateAgent, AclAction.CreateSession, AclAction.UpdateAgent, AclAction.DownloadSoftware, AclAction.CreatePool, AclAction.UpdatePool, AclAction.ViewPool, AclAction.DeletePool, AclAction.ListPools }));
 			DefaultAcl.Entries.Add(new AclEntry(AgentRoleClaim, new[] { AclAction.ViewProject, AclAction.ViewStream, AclAction.CreateEvent, AclAction.DownloadSoftware }));
 			DefaultAcl.Entries.Add(new AclEntry(DownloadSoftwareClaim, new[] { AclAction.DownloadSoftware }));
 			DefaultAcl.Entries.Add(new AclEntry(UploadSoftwareClaim, new[] { AclAction.UploadSoftware }));
