Horde: Add missing hibernating values to telemetry data.

diff --git a/Engine/Source/Programs/Horde/Horde.Build/Api/Reports.cs b/Engine/Source/Programs/Horde/Horde.Build/Api/Reports.cs
--- a/Engine/Source/Programs/Horde/Horde.Build/Api/Reports.cs
+++ b/Engine/Source/Programs/Horde/Horde.Build/Api/Reports.cs
@@ -56,6 +56,11 @@ namespace HordeServer.Api
 		/// </summary>
 		public double AdminTime { get; set; }
 
+		/// <summary>
+		/// Time spent hibernating
+		/// </summary>
+		public double HibernatingTime { get; }
+
 		/// <summary>
 		/// Total time agents in this pool were doing work for other pools
 		/// </summary>
@@ -75,6 +80,7 @@ namespace HordeServer.Api
 			this.PoolId = Pool.PoolId.ToString();
 			this.NumAgents = Pool.NumAgents;
 			this.AdminTime = Pool.AdminTime;
+			this.HibernatingTime = Pool.HibernatingTime;
 			this.OtherTime = Pool.OtherTime;
 
 			this.Streams = Pool.Streams.ConvertAll(Stream => new UtilizationTelemetryStream(Stream));
@@ -106,6 +112,11 @@ namespace HordeServer.Api
 		/// </summary>
 		public double AdminTime { get; set; }
 
+		/// <summary>
+		/// Total hibernating time
+		/// </summary>
+		public double HibernatingTime { get; set; }
+
 		/// <summary>
 		/// Total agents
 		/// </summary>
@@ -119,6 +130,7 @@ namespace HordeServer.Api
 			this.StartTime = Telemetry.StartTime;
 			this.FinishTime = Telemetry.FinishTime;
 			this.AdminTime = Telemetry.AdminTime;
+			this.HibernatingTime = Telemetry.HibernatingTime;
 			this.NumAgents = Telemetry.NumAgents;
 
 			this.Pools = Telemetry.Pools.ConvertAll(Pool => new UtilizationTelemetryPool(Pool));
