Horde: Include number of stopped agents in telemetry agent counts.

diff --git a/Engine/Source/Programs/Horde/Horde.Build/Fleet/TelemetryService.cs b/Engine/Source/Programs/Horde/Horde.Build/Fleet/TelemetryService.cs
--- a/Engine/Source/Programs/Horde/Horde.Build/Fleet/TelemetryService.cs
+++ b/Engine/Source/Programs/Horde/Horde.Build/Fleet/TelemetryService.cs
@@ -96,8 +96,10 @@ namespace HordeServer.Services
 					if (Pool.EnableAutoscaling)
 					{
 						int NumStoppedInstances = await FleetManager.GetNumStoppedInstancesAsync(Pool);
+						Telemetry.NumAgents += NumStoppedInstances;
 
 						NewPoolUtilizationTelemetry PoolTelemetry = Telemetry.FindOrAddPool(Pool.Id);
+						PoolTelemetry.NumAgents += NumStoppedInstances;
 						PoolTelemetry.HibernatingTime += Interval.TotalHours * NumStoppedInstances;
 					}
 				}
