Horde: Better theme defaulting, fixing agent view having wrong text
#rnx

diff --git a/Engine/Source/Programs/Horde/HordeDashboard/src/backend/Dashboard.ts b/Engine/Source/Programs/Horde/HordeDashboard/src/backend/Dashboard.ts
--- a/Engine/Source/Programs/Horde/HordeDashboard/src/backend/Dashboard.ts
+++ b/Engine/Source/Programs/Horde/HordeDashboard/src/backend/Dashboard.ts
@@ -191,8 +191,7 @@ export class Dashboard {
         const pref = this.preferences.get(DashboardPreference.Darktheme);
 
         if (pref !== "true" && pref !== "false") {
-            console.log("setting dark theme");            
-            this.setDarkTheme(true, false, true);            
+            console.error("No theme preference set");
         }
 
         return this.preferences.get(DashboardPreference.Darktheme) === 'true';
