Horde: Simplify some theme logic
#rnx

diff --git a/Engine/Source/Programs/Horde/HordeDashboard/src/backend/Backend.ts b/Engine/Source/Programs/Horde/HordeDashboard/src/backend/Backend.ts
--- a/Engine/Source/Programs/Horde/HordeDashboard/src/backend/Backend.ts
+++ b/Engine/Source/Programs/Horde/HordeDashboard/src/backend/Backend.ts
@@ -887,7 +887,15 @@ export class Backend {
 
                     data.dashboardSettings.preferences = preferences;
                 }
+
+                // apply defaults
+                const current = data.dashboardSettings.preferences.get(DashboardPreference.Darktheme);
+                if (current !== "true" && current !== "false") {
+                    data.dashboardSettings.preferences.set(DashboardPreference.Darktheme, "true");
+                }
+
                 resolve(data);
+
             }).catch(reason => reject(reason));
 
         });
diff --git a/Engine/Source/Programs/Horde/HordeDashboard/src/backend/Dashboard.ts b/Engine/Source/Programs/Horde/HordeDashboard/src/backend/Dashboard.ts
--- a/Engine/Source/Programs/Horde/HordeDashboard/src/backend/Dashboard.ts
+++ b/Engine/Source/Programs/Horde/HordeDashboard/src/backend/Dashboard.ts
@@ -185,20 +185,21 @@ export class Dashboard {
 
         if (!this.available) {
             // avoid intial flash when loading into site before backend is initialized
-            return localStorage?.getItem("horde_darktheme") === "true";
+            return localStorage?.getItem("horde_darktheme") !== "false";
         }
 
         const pref = this.preferences.get(DashboardPreference.Darktheme);
 
         if (pref !== "true" && pref !== "false") {
-            console.error("No theme preference set");
+            console.error("No theme preference set, should be defaulted in getCurrentUser");
+            return localStorage?.getItem("horde_darktheme") !== "false";
         }
 
-        return this.preferences.get(DashboardPreference.Darktheme) === 'true';
+        return this.preferences.get(DashboardPreference.Darktheme) !== 'false';
 
     }
 
-    setDarkTheme(value: boolean | undefined, update: boolean = true, resetColors: boolean = false) {
+    setDarkTheme(value: boolean | undefined) {
 
         this.setPreference(DashboardPreference.Darktheme, value ? "true" : "false");
 
@@ -207,15 +208,6 @@ export class Dashboard {
         } else {
             localStorage?.setItem("horde_darktheme", "false");
         }
-
-        if (resetColors) {
-            this.resetStatusColors();
-        }
-
-        if (update) {
-            this.setUpdated();
-        }
-
     }
 
     setDisplayUTC(value: boolean | undefined) {
