Horde Storage
Fix for handling replication events with empty blob references (delete records)

diff --git a/Engine/Source/Programs/Horde/HordeStorage/Horde.Storage/Implementation/TransactionLog/ScyllaReplicationLog.cs b/Engine/Source/Programs/Horde/HordeStorage/Horde.Storage/Implementation/TransactionLog/ScyllaReplicationLog.cs
--- a/Engine/Source/Programs/Horde/HordeStorage/Horde.Storage/Implementation/TransactionLog/ScyllaReplicationLog.cs
+++ b/Engine/Source/Programs/Horde/HordeStorage/Horde.Storage/Implementation/TransactionLog/ScyllaReplicationLog.cs
@@ -135,9 +135,9 @@ namespace Horde.Storage.Implementation
 
                     yield return new ReplicationLogEvent(
                         new NamespaceId(scyllaReplicationLog.Namespace),
-                        new BucketId(scyllaReplicationLog.Bucket!),
-                        new IoHashKey(scyllaReplicationLog.Key!),
-                        scyllaReplicationLog.ObjectIdentifier!.AsBlobIdentifier(),
+                        new BucketId(scyllaReplicationLog.Bucket),
+                        new IoHashKey(scyllaReplicationLog.Key),
+                        scyllaReplicationLog.ObjectIdentifier?.AsBlobIdentifier(),
                         scyllaReplicationLog.ReplicationId,
                         scyllaReplicationLog.GetReplicationBucketIdentifier(),
                         scyllaReplicationLog.GetReplicationBucketTimestamp(),
