Horde: Add persistent login redirect
#rnx

diff --git a/Engine/Source/Programs/Horde/HordeDashboard/src/backend/Backend.ts b/Engine/Source/Programs/Horde/HordeDashboard/src/backend/Backend.ts
--- a/Engine/Source/Programs/Horde/HordeDashboard/src/backend/Backend.ts
+++ b/Engine/Source/Programs/Horde/HordeDashboard/src/backend/Backend.ts
@@ -1416,8 +1416,54 @@ export class Backend {
         return token;
 
     }
+    
+    // When authenticating, we may get a top level redirected for login
+    // we need to store the initiating URL so can replace history with it upon return
+    // local storage is used to store the URL which is valid for one minute
+    authRedirect = 'horde_auth_redirect';
+    authRedirectTime = 'horde_auth_redirect_time';
+    authRedirectExpiration = 60 * 1000;    
+
+    // accessing sessionStorage can raise an access exception, so fallback to localStorage
+    storage: Storage | undefined;
+    
+    clearRedirect = () => {
+        this.storage?.setItem(this.authRedirect, "");
+        this.storage?.setItem(this.authRedirectTime, "");
+    };
+
+    getRedirect = (): string | undefined => {
+
+        const redirect = this.storage?.getItem(this.authRedirect);
+        const redirectTime: number = parseInt(this.storage?.getItem(this.authRedirectTime) ?? "");
+
+        this.clearRedirect();
+
+        if (!redirect || isNaN(redirectTime)) {
+            return undefined;
+        }
+
+        if ((Date.now() - redirectTime) > this.authRedirectExpiration) {
+            return undefined;
+        }
+
+        if (redirect.trim() === "/") {
+            return undefined;
+        }
+
+        return redirect;
+
+    };
+
+    redirect: string | undefined;
+
 
     init() {
+
+        try { this.storage = sessionStorage } catch { try { this.storage = localStorage } catch { console.warn("Unable to use session/local storage") } }
+
+        this.redirect = this.getRedirect();
+    
         return new Promise<boolean>(async (resolve, reject) => {
 
             this.backend.setBaseUrl(this.serverUrl);
@@ -1426,6 +1472,14 @@ export class Backend {
             const challenge = await this.backend.challenge();
 
             if (challenge === ChallengeStatus.Unauthorized) {
+                // store return url
+                if (!this.redirect) {
+                    const href = window.location.pathname + (window.location.search ?? "") + (window.location.hash ?? "").trim();
+                    if (href !== "/" && href !== "/index" && href !== "/index/") {
+                        this.storage?.setItem(this.authRedirect, href);
+                        this.storage?.setItem(this.authRedirectTime, Date.now().toString());	
+                    } 
+                }                
                 this.backend.login(window.location.toString());
                 return;
             }
