Horde: Send notification when jobs are scheduled in empty but auto-scaled pool

diff --git a/Engine/Source/Programs/Horde/Horde.Build.Tests/JobTaskSourceTests.cs b/Engine/Source/Programs/Horde/Horde.Build.Tests/JobTaskSourceTests.cs
--- a/Engine/Source/Programs/Horde/Horde.Build.Tests/JobTaskSourceTests.cs
+++ b/Engine/Source/Programs/Horde/Horde.Build.Tests/JobTaskSourceTests.cs
@@ -17,6 +17,10 @@ namespace HordeServerTests
 	[TestClass]
 	public class JobTaskSourceTests : TestSetup
 	{
+		private bool EventReceived;
+		private IPool? EventPool;
+		private bool? EventPoolHasAgentsOnline;
+		
 		[TestMethod]
 		public async Task UpdateJobQueueNormal()
 		{
@@ -27,6 +31,9 @@ namespace HordeServerTests
 			Assert.AreEqual(1, JobTaskSource.GetQueueForTesting().Count);
 			Assert.AreEqual(Fixture.Job1.Id, JobTaskSource.GetQueueForTesting().Min!.Id.Item1);
 			Assert.AreEqual(JobStepBatchState.Ready, JobTaskSource.GetQueueForTesting().Min!.Batch.State);
+			
+			Assert.IsTrue(EventReceived);
+			Assert.IsTrue(EventPoolHasAgentsOnline!.Value);
 		}
 		
 		[TestMethod]
@@ -40,6 +47,8 @@ namespace HordeServerTests
 
 			IJob Job = (await JobService.GetJobAsync(Fixture.Job1.Id))!;
 			Assert.AreEqual(JobStepBatchError.NoAgentsInPool, Job.Batches[0].Error);
+			
+			Assert.IsFalse(EventReceived);
 		}
 		
 		[TestMethod]
@@ -53,6 +62,8 @@ namespace HordeServerTests
 
 			IJob Job = (await JobService.GetJobAsync(Fixture.Job1.Id))!;
 			Assert.AreEqual(JobStepBatchError.NoAgentsOnline, Job.Batches[0].Error);
+			
+			Assert.IsFalse(EventReceived);
 		}
 		
 		[TestMethod]
@@ -66,6 +77,9 @@ namespace HordeServerTests
 
 			Assert.AreEqual(Fixture.Job1.Id, JobTaskSource.GetQueueForTesting().Min!.Id.Item1);
 			Assert.AreEqual(JobStepBatchState.Ready, JobTaskSource.GetQueueForTesting().Min!.Batch.State);
+
+			Assert.IsTrue(EventReceived);
+			Assert.IsFalse(EventPoolHasAgentsOnline!.Value);
 		}
 
 		private async Task<Fixture> SetupPoolWithAgentAsync(bool IsPoolAutoScaled, bool ShouldCreateAgent, bool IsAgentEnabled)
@@ -78,6 +92,13 @@ namespace HordeServerTests
 				IAgent? Agent = await AgentService.CreateAgentAsync("TestAgent", IsAgentEnabled, null, new List<StringId<IPool>> { Pool.Id });
 				await AgentService.CreateSessionAsync(Agent, AgentStatus.Ok, new List<string>(), new Dictionary<string, int>(), null);
 			}
+			
+			JobTaskSource.OnJobScheduled += (Pool, PoolHasAgentsOnline, Job, Graph, BatchId) =>
+			{
+				EventReceived = true;
+				EventPool = Pool;
+				EventPoolHasAgentsOnline = PoolHasAgentsOnline;
+			};
 
 			return Fixture;
 		}
diff --git a/Engine/Source/Programs/Horde/Horde.Build/Jobs/Impl/JobService.cs b/Engine/Source/Programs/Horde/Horde.Build/Jobs/Impl/JobService.cs
--- a/Engine/Source/Programs/Horde/Horde.Build/Jobs/Impl/JobService.cs
+++ b/Engine/Source/Programs/Horde/Horde.Build/Jobs/Impl/JobService.cs
@@ -165,6 +165,13 @@ namespace HordeServer.Services
 		/// Event triggered when a job step completes
 		/// </summary>
 		public event JobStepCompleteEvent? OnJobStepComplete;
+		
+		/// <summary>
+		/// Event triggered when a job is scheduled on the underlying JobTaskSource
+		///
+		/// Exposed and duplicated here to avoid dependency on the exact JobTaskSource
+		/// </summary>
+		public event JobTaskSource.JobScheduleEvent? OnJobScheduled;
 
 		/// <summary>
 		/// Constructor
@@ -199,6 +206,11 @@ namespace HordeServer.Services
 			this.PerforceService = PerforceService;
 			this.Settings = Settings;
 			this.Logger = Logger;
+
+			this.JobTaskSource.OnJobScheduled += (Pool, NumAgentsOnline, Job, Graph, BatchId) =>
+			{
+				OnJobScheduled?.Invoke(Pool, NumAgentsOnline, Job, Graph, BatchId);
+			};
 		}
 
 #pragma warning disable CA1801
diff --git a/Engine/Source/Programs/Horde/Horde.Build/Jobs/Impl/JobTaskSource.cs b/Engine/Source/Programs/Horde/Horde.Build/Jobs/Impl/JobTaskSource.cs
--- a/Engine/Source/Programs/Horde/Horde.Build/Jobs/Impl/JobTaskSource.cs
+++ b/Engine/Source/Programs/Horde/Horde.Build/Jobs/Impl/JobTaskSource.cs
@@ -224,6 +224,16 @@ namespace HordeServer.Tasks.Impl
 		// Cache of stream objects. Used to resolve agent types.
 		private Dictionary<StreamId, IStream> Streams = new Dictionary<StreamId, IStream>();
 
+		/// <summary>
+		/// Delegate for job schedule events
+		/// </summary>
+		public delegate void JobScheduleEvent(IPool Pool, bool HasAgentsOnline, IJob Job, IGraph Graph, SubResourceId BatchId);
+		
+		/// <summary>
+		/// Event triggered when a job is scheduled
+		/// </summary>
+		public event JobScheduleEvent? OnJobScheduled;
+
 		// Interval between querying the database for jobs to execute
 		static readonly TimeSpan RefreshInterval = TimeSpan.FromSeconds(5.0);
 
@@ -349,13 +359,23 @@ namespace HordeServer.Tasks.Impl
 			SortedSet<QueueItem> NewQueue = new SortedSet<QueueItem>(Queue.Comparer);
 			Dictionary<(JobId, SubResourceId), QueueItem> NewBatchIdToQueueItem = new Dictionary<(JobId, SubResourceId), QueueItem>();
 			
-			bool HasAgentsOnline(PoolId PoolId)
+			// Returns true if agents are online and available for scheduling for a pool
+			bool IsPoolOnline(PoolId PoolId)
+			{
+				return OnlinePools.Contains(PoolId);
+			}
+			
+			// Returns true if a pool can be auto-scaled
+			bool IsPoolAutoScaled(PoolId PoolId)
 			{
 				IPool? Pool = Pools.Find(p => p.Id == PoolId);
-				bool IsPoolOnline = OnlinePools.Contains(PoolId);
-				bool IsPoolAutoScaled = ValidPools.Contains(PoolId) && Pool != null && Pool.EnableAutoscaling;
+				return ValidPools.Contains(PoolId) && Pool != null && Pool.EnableAutoscaling;
+			}
+			
+			bool HasAgentsOnlineOrIsAutoScaled(PoolId PoolId)
+			{
 				// If pool is auto-scaled, it will be considered online even if it has no agents online
-				return IsPoolOnline || IsPoolAutoScaled;
+				return IsPoolOnline(PoolId) || IsPoolAutoScaled(PoolId);
 			}
 
 			// Query for a new list of jobs for the queue
@@ -413,7 +433,7 @@ namespace HordeServer.Tasks.Impl
 					{
 						NewJob = await SkipBatchAsync(NewJob, Batch.Id, Graph, JobStepBatchError.NoAgentsInPool);
 					}
-					else if (!HasAgentsOnline(AgentType.Pool))
+					else if (!HasAgentsOnlineOrIsAutoScaled(AgentType.Pool))
 					{
 						NewJob = await SkipBatchAsync(NewJob, Batch.Id, Graph, JobStepBatchError.NoAgentsOnline);
 					}
@@ -427,6 +447,12 @@ namespace HordeServer.Tasks.Impl
 						QueueItem NewQueueItem = new QueueItem(Stream, NewJob, BatchIdx, AgentType.Pool, Workspace, UseAutoSDK);
 						NewQueue.Add(NewQueueItem);
 						NewBatchIdToQueueItem[(NewJob.Id, Batch.Id)] = NewQueueItem;
+						
+						IPool? NewJobPool = Pools.Find(p => p.Id == AgentType.Pool);
+						if (NewJobPool != null)
+						{
+							OnJobScheduled?.Invoke(NewJobPool, IsPoolOnline(AgentType.Pool), NewJob, Graph, Batch.Id);
+						}
 					}
 				}
 
diff --git a/Engine/Source/Programs/Horde/Horde.Build/Notifications/INotificationService.cs b/Engine/Source/Programs/Horde/Horde.Build/Notifications/INotificationService.cs
--- a/Engine/Source/Programs/Horde/Horde.Build/Notifications/INotificationService.cs
+++ b/Engine/Source/Programs/Horde/Horde.Build/Notifications/INotificationService.cs
@@ -43,6 +43,16 @@ namespace HordeServer.Notifications
 		/// <param name="StepId">The step id</param>
 		/// <returns>Async task</returns>
 		void NotifyJobStepComplete(IJob Job, IGraph Graph, SubResourceId BatchId, SubResourceId StepId);
+		
+		/// <summary>
+		/// Notify all subscribers that a job batch has been scheduled
+		/// </summary>
+		/// <param name="Pool">Pool for which job scheduled in</param>
+		/// <param name="PoolHasAgentsOnline">True if there are agents online in the pool</param>
+		/// <param name="Job">The job containing the batch that has been scheduled</param>
+		/// <param name="Graph">Graph used for the job</param>
+		/// <param name="BatchId">The batch id</param>
+		void NotifyJobScheduled(IPool Pool, bool PoolHasAgentsOnline, IJob Job, IGraph Graph, SubResourceId BatchId);
 
 		/// <summary>
 		/// Notify all subscribers that a job step's outcome has changed
diff --git a/Engine/Source/Programs/Horde/Horde.Build/Notifications/INotificationSink.cs b/Engine/Source/Programs/Horde/Horde.Build/Notifications/INotificationSink.cs
--- a/Engine/Source/Programs/Horde/Horde.Build/Notifications/INotificationSink.cs
+++ b/Engine/Source/Programs/Horde/Horde.Build/Notifications/INotificationSink.cs
@@ -16,6 +16,17 @@ namespace HordeServer.Notifications
 	/// </summary>
 	public interface INotificationSink
 	{
+		/// <summary>
+		/// Send notifications that a job has been scheduled
+		/// </summary>
+		/// <param name="Pool">Pool for which job scheduled in</param>
+		/// <param name="PoolHasAgentsOnline">True if there are agents online in the pool</param>
+		/// <param name="Job">The job containing the batch that has been scheduled</param>
+		/// <param name="Graph">Graph used for the job</param>
+		/// <param name="BatchId">The batch id</param>
+		/// <returns>Async task</returns>
+		Task NotifyJobScheduledAsync(IPool Pool, bool PoolHasAgentsOnline, IJob Job, IGraph Graph, SubResourceId BatchId);
+		
 		/// <summary>
 		/// Send notifications that a job has completed
 		/// </summary>
diff --git a/Engine/Source/Programs/Horde/Horde.Build/Notifications/Impl/NotificationService.cs b/Engine/Source/Programs/Horde/Horde.Build/Notifications/Impl/NotificationService.cs
--- a/Engine/Source/Programs/Horde/Horde.Build/Notifications/Impl/NotificationService.cs
+++ b/Engine/Source/Programs/Horde/Horde.Build/Notifications/Impl/NotificationService.cs
@@ -131,6 +131,7 @@ namespace HordeServer.Notifications.Impl
 
 			IssueService.OnIssueUpdated += NotifyIssueUpdated;
 			JobService.OnJobStepComplete += NotifyJobStepComplete;
+			JobService.OnJobScheduled += NotifyJobScheduled;
 			JobService.OnLabelUpdate += NotifyLabelUpdate;
 		}
 
@@ -141,6 +142,7 @@ namespace HordeServer.Notifications.Impl
 
 			IssueService.OnIssueUpdated -= NotifyIssueUpdated;
 			JobService.OnJobStepComplete -= NotifyJobStepComplete;
+			JobService.OnJobScheduled += NotifyJobScheduled;
 			JobService.OnLabelUpdate -= NotifyLabelUpdate;
 
 			GC.SuppressFinalize(this);
@@ -197,6 +199,15 @@ namespace HordeServer.Notifications.Impl
 				EnqueueTask(() => RecordJobCompleteMetrics(Job));
 			}
 		}
+		
+		/// <inheritdoc/>
+		public void NotifyJobScheduled(IPool Pool, bool PoolHasAgentsOnline, IJob Job, IGraph Graph, SubResourceId BatchId)
+		{
+			if (Pool.EnableAutoscaling && !PoolHasAgentsOnline)
+			{
+				EnqueueTasks(Sink => Sink.NotifyJobScheduledAsync(Pool, PoolHasAgentsOnline, Job, Graph, BatchId));
+			}
+		}
 
 		/// <inheritdoc/>
 		public void NotifyLabelUpdate(IJob Job, IReadOnlyList<(LabelState, LabelOutcome)> OldLabelStates, IReadOnlyList<(LabelState, LabelOutcome)> NewLabelStates)
diff --git a/Engine/Source/Programs/Horde/Horde.Build/Notifications/Impl/SlackNotificationSink.cs b/Engine/Source/Programs/Horde/Horde.Build/Notifications/Impl/SlackNotificationSink.cs
--- a/Engine/Source/Programs/Horde/Horde.Build/Notifications/Impl/SlackNotificationSink.cs
+++ b/Engine/Source/Programs/Horde/Horde.Build/Notifications/Impl/SlackNotificationSink.cs
@@ -278,6 +278,33 @@ namespace HordeServer.Notifications.Impl
 		}
 
 		#endregion
+		
+		/// <inheritdoc/>
+		public async Task NotifyJobScheduledAsync(IPool Pool, bool PoolHasAgentsOnline, IJob Job, IGraph Graph, SubResourceId BatchId)
+		{
+			if (Settings.JobNotificationChannel != null)
+			{
+				Logger.LogInformation("Sending Slack notification for scheduled job {JobId}, batch {BatchId} to channel {SlackChannel}", Job.Id, BatchId, Settings.JobNotificationChannel );
+				await SendJobScheduledOnEmptyAutoScaledPoolMessageAsync($"#{Settings.JobNotificationChannel}", Pool, PoolHasAgentsOnline, Job, Graph, BatchId);
+			}
+		}
+		
+		private Task SendJobScheduledOnEmptyAutoScaledPoolMessageAsync(string Recipient, IPool Pool, bool PoolHasAgentsOnline, IJob Job, IGraph Graph, SubResourceId BatchId)
+		{
+			string JobUrl = Settings.DashboardUrl + "/job/" + Job.Id;
+			
+			Color OutcomeColor = BlockKitAttachmentColors.Warning;
+			BlockKitAttachment Attachment = new BlockKitAttachment();
+			Attachment.Color = OutcomeColor;
+			Attachment.FallbackText = $"Job scheduled in an auto-scaled pool with no agents online. Job ID {Job.Id}";
+
+			Attachment.Blocks.Add(new HeaderBlock($"Job scheduled in empty pool", false, true));
+
+			Attachment.Blocks.Add(new SectionBlock($"A <{JobUrl}|job> was scheduled in an auto-scaled pool but with no current agents online."));
+			Attachment.Blocks.Add(new SectionBlock($"Job {Job.Name}\nJob ID <{JobUrl}|{Job.Id}>\nPool {Pool.Name}"));
+
+			return SendMessageAsync(Recipient, Attachments: new[] { Attachment });
+		}
 
 		#region Job Complete
 
diff --git a/Engine/Source/Programs/Horde/Horde.Build/ServerSettings.cs b/Engine/Source/Programs/Horde/Horde.Build/ServerSettings.cs
--- a/Engine/Source/Programs/Horde/Horde.Build/ServerSettings.cs
+++ b/Engine/Source/Programs/Horde/Horde.Build/ServerSettings.cs
@@ -439,6 +439,11 @@ namespace HordeServer
 		/// Channel to send device notifications to
 		/// </summary>
 		public string? DeviceServiceNotificationChannel { get; set; }
+		
+		/// <summary>
+		/// Slack channel to send job related notifications to
+		/// </summary>
+		public string? JobNotificationChannel { get; set; }
 
 		/// <summary>
 		/// URI to the SmtpServer to use for sending email notifications
