Horde: Simplified auth redirect login logic (fixes issue escaping Okta uri)
#rnx

diff --git a/Engine/Source/Programs/Horde/Horde.Build/Controllers/DashboardController.cs b/Engine/Source/Programs/Horde/Horde.Build/Controllers/DashboardController.cs
--- a/Engine/Source/Programs/Horde/Horde.Build/Controllers/DashboardController.cs
+++ b/Engine/Source/Programs/Horde/Horde.Build/Controllers/DashboardController.cs
@@ -11,6 +11,8 @@ using System.Threading.Tasks;
 using HordeServer;
 using HordeServer.Controllers;
 using Microsoft.Extensions.Options;
+using System;
+using System.Text;
 
 namespace Horde.Build
 {
@@ -59,6 +61,27 @@ namespace Horde.Build
 			return new ChallengeResult(AuthenticationScheme, new AuthenticationProperties { RedirectUri = Redirect ?? "/" });
 		}
 
+		/// <summary>
+		/// Login to server, redirecting to the specified Base64 encoded URL, which fixes some escaping issues on some auth providers, on success
+		/// </summary>
+		/// <param name="Redirect"></param>
+		/// <returns></returns>
+		[HttpGet]
+		[Route("/api/v2/dashboard/login")]
+		public IActionResult LoginV2([FromQuery] string? Redirect)
+		{
+			string? RedirectUri = null;
+
+			if (Redirect != null)
+			{
+				byte[] Data = Convert.FromBase64String(Redirect);
+				RedirectUri = Encoding.UTF8.GetString(Data);
+			}
+
+			return new ChallengeResult(AuthenticationScheme, new AuthenticationProperties { RedirectUri = RedirectUri ?? "/index" });
+		}
+
+
 		/// <summary>
 		/// Logout of the current account
 		/// </summary>
diff --git a/Engine/Source/Programs/Horde/HordeDashboard/src/backend/Backend.ts b/Engine/Source/Programs/Horde/HordeDashboard/src/backend/Backend.ts
--- a/Engine/Source/Programs/Horde/HordeDashboard/src/backend/Backend.ts
+++ b/Engine/Source/Programs/Horde/HordeDashboard/src/backend/Backend.ts
@@ -1416,53 +1416,10 @@ export class Backend {
         return token;
 
     }
-    
-    // When authenticating, we may get a top level redirected for login
-    // we need to store the initiating URL so can replace history with it upon return
-    // local storage is used to store the URL which is valid for one minute
-    authRedirect = 'horde_auth_redirect';
-    authRedirectTime = 'horde_auth_redirect_time';
-    authRedirectExpiration = 60 * 1000;    
-
-    // accessing sessionStorage can raise an access exception, so fallback to localStorage
-    storage: Storage | undefined;
-    
-    clearRedirect = () => {
-        this.storage?.setItem(this.authRedirect, "");
-        this.storage?.setItem(this.authRedirectTime, "");
-    };
-
-    getRedirect = (): string | undefined => {
-
-        const redirect = this.storage?.getItem(this.authRedirect);
-        const redirectTime: number = parseInt(this.storage?.getItem(this.authRedirectTime) ?? "");
-
-        this.clearRedirect();
-
-        if (!redirect || isNaN(redirectTime)) {
-            return undefined;
-        }
-
-        if ((Date.now() - redirectTime) > this.authRedirectExpiration) {
-            return undefined;
-        }
-
-        if (redirect.trim() === "/") {
-            return undefined;
-        }
-
-        return redirect;
-
-    };
-
-    redirect: string | undefined;
-
+        
 
     init() {
 
-        try { this.storage = sessionStorage } catch { try { this.storage = localStorage } catch { console.warn("Unable to use session/local storage") } }
-
-        this.redirect = this.getRedirect();
     
         return new Promise<boolean>(async (resolve, reject) => {
 
@@ -1472,14 +1429,6 @@ export class Backend {
             const challenge = await this.backend.challenge();
 
             if (challenge === ChallengeStatus.Unauthorized) {
-                // store return url
-                if (!this.redirect) {
-                    const href = window.location.pathname + (window.location.search ?? "") + (window.location.hash ?? "").trim();
-                    if (href !== "/" && href !== "/index" && href !== "/index/") {
-                        this.storage?.setItem(this.authRedirect, href);
-                        this.storage?.setItem(this.authRedirectTime, Date.now().toString());	
-                    } 
-                }                
                 this.backend.login(window.location.toString());
                 return;
             }
diff --git a/Engine/Source/Programs/Horde/HordeDashboard/src/backend/Fetch.ts b/Engine/Source/Programs/Horde/HordeDashboard/src/backend/Fetch.ts
--- a/Engine/Source/Programs/Horde/HordeDashboard/src/backend/Fetch.ts
+++ b/Engine/Source/Programs/Horde/HordeDashboard/src/backend/Fetch.ts
@@ -166,11 +166,9 @@ export class Fetch {
             return;
         }
 
-        window.location.assign("/api/v1/dashboard/login?redirect=" + encodeURI(redirect));
+        window.location.assign("/api/v2/dashboard/login?redirect=" + btoa(redirect ?? "/index"));
     }
 
-
-
     private async handleResponse(response: Response, url: string, mode: string, resolve: (value: FetchResponse | PromiseLike<FetchResponse>) => void, reject: (reason?: any) => void, config?: FetchRequestConfig) {
 
 
