Horde: Add JiraService for issue integration
#rnx

diff --git a/Engine/Source/Programs/Horde/Horde.Build/Api/Jira.cs b/Engine/Source/Programs/Horde/Horde.Build/Api/Jira.cs
new file mode 100644
--- /dev/null
+++ b/Engine/Source/Programs/Horde/Horde.Build/Api/Jira.cs
@@ -0,0 +1,78 @@
+// Copyright Epic Games, Inc. All Rights Reserved.
+
+using System;
+using System.Collections.Generic;
+using System.ComponentModel.DataAnnotations;
+using System.Diagnostics.CodeAnalysis;
+using System.Linq;
+using HordeServer.Models;
+using HordeServer.Services;
+
+namespace HordeServer.Api
+{
+
+	/// <summary>
+	/// Jira issue response object
+	/// </summary>
+	public class GetJiraIssueResponse
+	{
+		/// <summary>
+		/// The Jira issue key
+		/// </summary>
+		public string Key { get; set; }
+
+		/// <summary>
+		/// The Jira issue link
+		/// </summary>
+		public string JiraLink { get; set; }
+
+		/// <summary>
+		/// The issue status name, "To Do", "In Progress", etc
+		/// </summary>
+		public string? StatusName { get; set; }
+
+		/// <summary>
+		/// The issue resolution name, "Fixed", "Closed", etc
+		/// </summary>
+		public string? ResolutionName { get; set; }
+
+		/// <summary>
+		/// The issue priority name, "1 - Critical", "2 - Major", etc
+		/// </summary>
+		public string? PriorityName { get; set; }
+
+		/// <summary>
+		/// The current assignee's user name
+		/// </summary>
+		public string? AssigneeName { get; set; }
+
+		/// <summary>
+		/// The current assignee's display name
+		/// </summary>
+		public string? AssigneeDisplayName { get; set; }
+
+		/// <summary>
+		/// The current assignee's email address
+		/// </summary>
+		public string? AssigneeEmailAddress { get; set; }
+
+		/// <summary>
+		/// Response constructor
+		/// </summary>
+		public GetJiraIssueResponse(JiraIssue Issue)
+		{
+			this.Key = Issue.Key;
+			this.JiraLink = Issue.JiraLink;
+			this.StatusName = Issue.StatusName;
+			this.ResolutionName = Issue.ResolutionName;
+			this.PriorityName = Issue.PriorityName;
+			this.AssigneeName = Issue.AssigneeName;
+			this.AssigneeDisplayName = Issue.AssigneeDisplayName;
+			this.AssigneeEmailAddress = Issue.AssigneeEmailAddress;
+		}
+
+	}
+
+}
+
+ 
\ No newline at end of file
diff --git a/Engine/Source/Programs/Horde/Horde.Build/Controllers/JiraController.cs b/Engine/Source/Programs/Horde/Horde.Build/Controllers/JiraController.cs
new file mode 100644
--- /dev/null
+++ b/Engine/Source/Programs/Horde/Horde.Build/Controllers/JiraController.cs
@@ -0,0 +1,84 @@
+// Copyright Epic Games, Inc. All Rights Reserved.
+
+using Microsoft.AspNetCore.Authorization;
+using Microsoft.AspNetCore.Mvc;
+using System.Threading.Tasks;
+using HordeServer.Services;
+using Microsoft.Extensions.Logging;
+using HordeServer.Api;
+using System.Collections.Generic;
+using HordeServer.Models;
+using HordeServer.Utilities;
+
+namespace Horde.Build
+{
+	using StreamId = StringId<IStream>;
+
+	/// <summary>	
+	/// Jira service controller
+	/// </summary>	
+	[ApiController]
+	[Route("[controller]")]
+	public class JiraController : Controller
+	{
+		/// <summary>
+		/// Singleton instance of the jira service
+		/// </summary>
+		JiraService JiraService;
+
+		/// <summary>
+		/// Singleton instance of the stream service
+		/// </summary>
+		StreamService StreamService;
+
+		/// <summary>
+		///  Logger for controller
+		/// </summary>
+		private readonly ILogger<JiraController> Logger;
+
+		/// <summary>
+		/// Constructor
+		/// </summary>
+		public JiraController(JiraService JiraService, StreamService StreamService, ILogger<JiraController> Logger)
+		{
+			this.JiraService = JiraService;
+			this.Logger = Logger;
+			this.StreamService = StreamService;
+		}
+
+		/// <summary>
+		/// Get jira issue information for provided keys
+		/// </summary>
+		[HttpGet]
+		[Authorize]
+		[Route("/api/v1/jira")]
+		public async Task<ActionResult<List<GetJiraIssueResponse>>> GetJiraIssuesAsync([FromQuery] string StreamId, [FromQuery] string[] JiraKeys)
+		{
+
+			StreamId StreamIdValue = new StreamId(StreamId);
+			StreamPermissionsCache Cache = new StreamPermissionsCache();
+
+			if (!await StreamService.AuthorizeAsync(StreamIdValue, AclAction.ViewStream, User, Cache))
+			{
+				return Forbid();
+			}
+
+			List<GetJiraIssueResponse> Response = new List<GetJiraIssueResponse>();
+
+			if (JiraKeys.Length != 0)
+			{
+				List<JiraIssue> Issues = await this.JiraService.GetJiraIssuesAsync(JiraKeys);
+
+				for (int i = 0; i < Issues.Count; i++)
+				{
+					Response.Add(new GetJiraIssueResponse(Issues[i]));
+				}
+
+			}			
+
+			return Response;
+		}
+
+
+	}
+}
diff --git a/Engine/Source/Programs/Horde/Horde.Build/Jira/IJiraService.cs b/Engine/Source/Programs/Horde/Horde.Build/Jira/IJiraService.cs
new file mode 100644
--- /dev/null
+++ b/Engine/Source/Programs/Horde/Horde.Build/Jira/IJiraService.cs
@@ -0,0 +1,68 @@
+// Copyright Epic Games, Inc. All Rights Reserved.	
+
+using System.Collections.Generic;
+using System.Threading.Tasks;
+
+namespace HordeServer.Services
+{
+
+	/// <summary>
+	/// Jira issue information
+	/// </summary>
+	public class JiraIssue
+	{
+		/// <summary>
+		/// The Jira issue key
+		/// </summary>
+		public string Key { get; set; } = string.Empty;
+
+		/// <summary>
+		/// The Jira issue link
+		/// </summary>
+		public string JiraLink { get; set; } = string.Empty;
+
+		/// <summary>
+		/// The issue status name, "To Do", "In Progress", etc
+		/// </summary>
+		public string? StatusName { get; set; }
+
+		/// <summary>
+		/// The issue resolution name, "Fixed", "Closed", etc
+		/// </summary>
+		public string? ResolutionName { get; set; }
+
+		/// <summary>
+		/// The issue priority name, "1 - Critical", "2 - Major", etc
+		/// </summary>
+		public string? PriorityName { get; set; }
+
+		/// <summary>
+		/// The current assignee's user name
+		/// </summary>
+		public string? AssigneeName { get; set; }
+
+		/// <summary>
+		/// The current assignee's display name
+		/// </summary>
+		public string? AssigneeDisplayName { get; set; }
+
+		/// <summary>
+		/// The current assignee's email address
+		/// </summary>
+		public string? AssigneeEmailAddress { get; set; }
+	}
+
+	/// <summary>
+	/// Wrapper around Jira functionality.
+	/// </summary>
+	public interface IJiraService
+	{
+		/// <summary>
+		/// Get Jira issues associated with provided keys
+		/// </summary>
+		/// <param name="JiraKeys"></param>
+		/// <returns></returns>
+		Task<List<JiraIssue>> GetJiraIssuesAsync(string[] JiraKeys);
+	}
+
+}
\ No newline at end of file
diff --git a/Engine/Source/Programs/Horde/Horde.Build/Jira/Impl/JiraService.cs b/Engine/Source/Programs/Horde/Horde.Build/Jira/Impl/JiraService.cs
new file mode 100644
--- /dev/null
+++ b/Engine/Source/Programs/Horde/Horde.Build/Jira/Impl/JiraService.cs
@@ -0,0 +1,193 @@
+// Copyright Epic Games, Inc. All Rights Reserved.	
+
+using Microsoft.Extensions.Logging;
+using Microsoft.Extensions.Options;
+using System;
+using System.Collections.Generic;
+using System.Text;
+using System.Threading.Tasks;
+using System.Net.Http;
+using Polly;
+using Polly.Extensions.Http;
+using System.Text.Json;
+using System.Net.Http.Headers;
+using System.Collections.Concurrent;
+
+namespace HordeServer.Services
+{
+	/// <summary>
+	/// Jira service functinality
+	/// </summary>
+	public sealed class JiraService : IJiraService, IDisposable
+	{
+		ILogger Logger;
+
+		/// <summary>
+		/// The server settings
+		/// </summary>
+		ServerSettings Settings;
+
+		HttpClient Client;
+		AsyncPolicy<HttpResponseMessage> RetryPolicy;
+
+		ConcurrentDictionary<string, JiraCacheValue> IssueCache = new ConcurrentDictionary<string, JiraCacheValue>();
+
+		/// <summary>
+		/// Jira service constructor
+		/// </summary>
+		/// <param name="Settings"></param>
+		/// <param name="Logger"></param>
+		public JiraService(IOptions<ServerSettings> Settings, ILogger<JiraService> Logger)
+		{
+
+			this.Settings = Settings.Value;
+			this.Logger = Logger;
+
+			// setup http client for Jira rest api queries
+			Client = new HttpClient();
+			byte[] AuthBytes = Encoding.ASCII.GetBytes($"{this.Settings.JiraUsername}:{this.Settings.JiraApiToken}");
+			Client.DefaultRequestHeaders.Authorization = new AuthenticationHeaderValue("Basic", Convert.ToBase64String(AuthBytes));
+			Client.Timeout = TimeSpan.FromSeconds(15.0);			
+			RetryPolicy = HttpPolicyExtensions.HandleTransientHttpError().WaitAndRetryAsync(3, Attempt => TimeSpan.FromSeconds(Math.Pow(2.0, Attempt)));
+		}
+
+		/// <inheritdoc/>
+		public void Dispose()
+		{
+			Client.Dispose();
+		}
+
+		/// <inheritdoc/>
+		public async Task<List<JiraIssue>> GetJiraIssuesAsync(string[] JiraKeys)
+		{
+			List<JiraIssue> Result = new List<JiraIssue>();
+			
+			if (Settings.JiraUrl == null || JiraKeys.Length == 0)
+			{
+				return Result;
+			}
+
+			List<string> QueryJiras = new List<string>();
+			for (int i = 0; i < JiraKeys.Length; i++)
+			{
+				JiraCacheValue Value;
+				if (IssueCache.TryGetValue(JiraKeys[i], out Value))
+				{
+					if (DateTime.UtcNow.Subtract(Value.CacheTime).TotalMinutes >= 2)
+					{
+						QueryJiras.Add(JiraKeys[i]);				
+					}
+				}
+				else
+				{
+					QueryJiras.Add(JiraKeys[i]);
+				}
+			}
+
+			if (QueryJiras.Count > 0)
+			{
+				Uri Uri = new Uri(Settings.JiraUrl, $"/rest/api/2/search?jql=issueKey%20in%20({string.Join(",", QueryJiras)})&fields=assignee,status,resolution,priority&maxResults={QueryJiras.Count}");
+
+				HttpResponseMessage Response = await RetryPolicy.ExecuteAsync(() => Client.GetAsync(Uri));
+				if (!Response.IsSuccessStatusCode)
+				{
+					Logger.LogError("GET to {Uri} returned {Code} ({Response})", Uri, Response.StatusCode, await Response.Content.ReadAsStringAsync());
+					throw new Exception("GetJiraIssuesAsync call failed");
+				}
+
+				byte[] Data = await Response.Content.ReadAsByteArrayAsync();
+				IssueQueryResponse? Jiras = JsonSerializer.Deserialize<IssueQueryResponse>(Data.AsSpan(), new JsonSerializerOptions() { PropertyNameCaseInsensitive = true });
+
+				if (Jiras != null)
+				{
+					for (int i = 0; i < Jiras.Issues.Count; i++)
+					{
+						IssueResponse Issue = Jiras.Issues[i];
+
+						JiraIssue JiraIssue = new JiraIssue();
+						JiraIssue.Key = Issue.Key;
+						JiraIssue.JiraLink = $"{Settings.JiraUrl}browse/{Issue.Key}";
+						JiraIssue.AssigneeName = Issue.Fields?.Assignee?.Name;
+						JiraIssue.AssigneeDisplayName = Issue.Fields?.Assignee?.DisplayName;
+						JiraIssue.AssigneeEmailAddress = Issue.Fields?.Assignee?.EmailAddress;
+						JiraIssue.PriorityName = Issue.Fields?.Priority?.Name;
+						JiraIssue.ResolutionName = Issue.Fields?.Resolution?.Name;
+						JiraIssue.StatusName = Issue.Fields?.Status?.Name;
+
+						IssueCache[Issue.Key] = new JiraCacheValue() { Issue = JiraIssue, CacheTime = DateTime.UtcNow };
+					}
+				}
+			}
+
+			for (int i = 0; i < JiraKeys.Length; i++)
+			{
+				JiraCacheValue Value;
+				if (IssueCache.TryGetValue(JiraKeys[i], out Value))
+				{
+					Result.Add(Value.Issue);
+				}
+			}
+
+			return Result;
+
+		}
+
+		// Jira rest api mapping and caching
+
+		struct JiraCacheValue
+		{
+			public JiraIssue Issue;
+			public DateTime CacheTime;
+		};
+
+		class IssueQueryResponse
+		{
+			public int Total { get; set; } = 0;
+
+			public List<IssueResponse> Issues { get; set; } = new List<IssueResponse>();
+		}
+
+		class IssueResponse
+		{
+			public string Key { get; set; } = String.Empty;
+
+			public IssueFields? Fields { get; set; }
+		}
+
+		class IssueFields
+		{
+			public AssigneeField? Assignee { get; set; }
+
+			public StatusField? Status { get; set; }
+
+			public ResolutionField? Resolution { get; set; }
+
+			public PriorityField? Priority { get; set; }
+		}
+
+		class AssigneeField
+		{
+			public string? Name { get; set; }
+
+			public string? DisplayName { get; set; }
+
+			public string? EmailAddress { get; set; }
+		}
+
+		class StatusField
+		{
+			public string? Name { get; set; }
+		}
+
+		class PriorityField
+		{
+			public string? Name { get; set; }
+		}
+
+		class ResolutionField
+		{
+			public string? Name { get; set; }
+		}
+
+	}
+}
diff --git a/Engine/Source/Programs/Horde/Horde.Build/ServerSettings.cs b/Engine/Source/Programs/Horde/Horde.Build/ServerSettings.cs
--- a/Engine/Source/Programs/Horde/Horde.Build/ServerSettings.cs
+++ b/Engine/Source/Programs/Horde/Horde.Build/ServerSettings.cs
@@ -485,6 +485,21 @@ namespace HordeServer
 		/// </summary>
 		public bool P4BridgeCanImpersonate { get; set; } = false;
 
+		/// <summary>
+		/// The Jira service account user name
+		/// </summary>
+		public string? JiraUsername { get; set; }
+
+		/// <summary>
+		/// The Jira service account API token
+		/// </summary>
+		public string? JiraApiToken { get; set; }
+
+		/// <summary>
+		/// The Uri for the Jira installation
+		/// </summary>
+		public Uri? JiraUrl { get; set; }
+
 		/// <summary>
 		/// Set the minimum size of the global thread pool
 		/// This value has been found in need of tweaking to avoid timeouts with the Redis client during bursts
diff --git a/Engine/Source/Programs/Horde/Horde.Build/Startup.cs b/Engine/Source/Programs/Horde/Horde.Build/Startup.cs
--- a/Engine/Source/Programs/Horde/Horde.Build/Startup.cs
+++ b/Engine/Source/Programs/Horde/Horde.Build/Startup.cs
@@ -386,8 +386,8 @@ namespace HordeServer
 			Services.AddSingleton<INotificationSink, SlackNotificationSink>(SP => SP.GetRequiredService<SlackNotificationSink>());
 			Services.AddSingleton<StreamService>();
 			Services.AddSingleton<UpgradeService>();
-
-			Services.AddSingleton<DeviceService>();
+			Services.AddSingleton<DeviceService>();			
+			Services.AddSingleton<JiraService>();
 
 			AWSOptions AwsOptions = Configuration.GetAWSOptions();
 			if (Settings.S3CredentialType == "AssumeRole" && Settings.S3AssumeArn != null)
diff --git a/Engine/Source/Programs/Horde/HordeDashboard/src/backend/Api.ts b/Engine/Source/Programs/Horde/HordeDashboard/src/backend/Api.ts
--- a/Engine/Source/Programs/Horde/HordeDashboard/src/backend/Api.ts
+++ b/Engine/Source/Programs/Horde/HordeDashboard/src/backend/Api.ts
@@ -3825,3 +3825,18 @@ export type FindIssueResponse = {
 	/// Spans for this issue
 	spans: FindIssueSpanResponse[];
 }
+
+
+/// Jira issue information
+export type GetJiraIssueResponse = {
+
+	key: string;
+	jiraLink: string;
+	statusName?: string;
+	priorityName?: string;
+	resolutionName?: string;
+	assigneeName?: string;
+	assigneeDisplayName?: string;
+	assigneeEmailAddress?: string;
+
+}
\ No newline at end of file
diff --git a/Engine/Source/Programs/Horde/HordeDashboard/src/backend/Backend.ts b/Engine/Source/Programs/Horde/HordeDashboard/src/backend/Backend.ts
--- a/Engine/Source/Programs/Horde/HordeDashboard/src/backend/Backend.ts
+++ b/Engine/Source/Programs/Horde/HordeDashboard/src/backend/Backend.ts
@@ -3,7 +3,7 @@
 import { configure } from "mobx";
 import { isNumber } from 'util';
 import templateCache from '../backend/TemplateCache';
-import { AgentData, ArtifactData, AuditLogEntry, AuditLogQuery, BatchUpdatePoolRequest, ChangeSummaryData, CreateDeviceRequest, CreateDeviceResponse, CreateJobRequest, CreateJobResponse, CreatePoolRequest, CreateSoftwareResponse, CreateSubscriptionRequest, CreateSubscriptionResponse, DashboardPreference, EventData, FindIssueResponse, GetArtifactZipRequest, GetDevicePlatformResponse, GetDevicePoolResponse, GetDeviceReservationResponse, GetDeviceResponse, GetGraphResponse, GetIssueStreamResponse, GetJobsTabResponse, GetJobStepRefResponse, GetJobStepTraceResponse, GetJobTimingResponse, GetLogEventResponse, GetNotificationResponse, GetPerforceServerStatusResponse, GetServerInfoResponse, GetServerSettingsResponse, GetSoftwareResponse, GetSubscriptionResponse, GetUserResponse, GetUtilizationTelemetryResponse, GlobalConfig, IssueData, IssueQuery, IssueQueryV2, JobData, JobQuery, JobsTabColumnType, JobStreamQuery, LeaseData, LogData, LogLineData, PoolData, ProjectData, ScheduleData, ScheduleQuery, SearchLogFileResponse, ServerUpdateResponse, SessionData, StreamData, TabType, TemplateData, TestData, UpdateAgentRequest, UpdateDeviceRequest, UpdateGlobalConfigRequest, UpdateIssueRequest, UpdateJobRequest, UpdateLeaseRequest, UpdateNotificationsRequest, UpdatePoolRequest, UpdateServerSettingsRequest, UpdateStepRequest, UpdateStepResponse, UpdateUserRequest, UsersQuery } from './Api';
+import { AgentData, ArtifactData, AuditLogEntry, AuditLogQuery, BatchUpdatePoolRequest, ChangeSummaryData, CreateDeviceRequest, CreateDeviceResponse, CreateJobRequest, CreateJobResponse, CreatePoolRequest, CreateSoftwareResponse, CreateSubscriptionRequest, CreateSubscriptionResponse, DashboardPreference, EventData, FindIssueResponse, GetArtifactZipRequest, GetDevicePlatformResponse, GetDevicePoolResponse, GetDeviceReservationResponse, GetDeviceResponse, GetGraphResponse, GetIssueStreamResponse, GetJiraIssueResponse, GetJobsTabResponse, GetJobStepRefResponse, GetJobStepTraceResponse, GetJobTimingResponse, GetLogEventResponse, GetNotificationResponse, GetPerforceServerStatusResponse, GetServerInfoResponse, GetServerSettingsResponse, GetSoftwareResponse, GetSubscriptionResponse, GetUserResponse, GetUtilizationTelemetryResponse, GlobalConfig, IssueData, IssueQuery, IssueQueryV2, JobData, JobQuery, JobsTabColumnType, JobStreamQuery, LeaseData, LogData, LogLineData, PoolData, ProjectData, ScheduleData, ScheduleQuery, SearchLogFileResponse, ServerUpdateResponse, SessionData, StreamData, TabType, TemplateData, TestData, UpdateAgentRequest, UpdateDeviceRequest, UpdateGlobalConfigRequest, UpdateIssueRequest, UpdateJobRequest, UpdateLeaseRequest, UpdateNotificationsRequest, UpdatePoolRequest, UpdateServerSettingsRequest, UpdateStepRequest, UpdateStepResponse, UpdateUserRequest, UsersQuery } from './Api';
 import dashboard from './Dashboard';
 import { ChallengeStatus, Fetch } from './Fetch';
 import graphCache, { GraphQuery } from './GraphCache';
@@ -1013,6 +1013,23 @@ export class Backend {
         });
     }
 
+    getJiraIssues(streamId: string, jiraKeys: string[]): Promise<GetJiraIssueResponse[]> {
+    
+        const params = {
+            streamId: streamId,
+            jiraKeys: jiraKeys
+        };
+        
+        return new Promise<GetJiraIssueResponse[]>((resolve, reject) => {
+            this.backend.get(`/api/v1/jira`, { params: params }).then((value) => {
+                resolve(value.data as GetJiraIssueResponse[]);
+            }).catch((reason) => {
+                reject(reason);
+            });
+        });
+
+    }
+
     getJobStepTestDataByKey(jobId: string, stepId: string, key: string): Promise<TestData[]> {
         return new Promise<TestData[]>((resolve, reject) => {
             this.backend.get(`/api/v1/testdata?JobId=${jobId}&JobStepId=${stepId}&Key=${key}`).then((value) => {
