Horde: Fix issues with HistoryModal refresh
#rnx

diff --git a/Engine/Source/Programs/Horde/HordeDashboard/src/backend/Backend.ts b/Engine/Source/Programs/Horde/HordeDashboard/src/backend/Backend.ts
--- a/Engine/Source/Programs/Horde/HordeDashboard/src/backend/Backend.ts
+++ b/Engine/Source/Programs/Horde/HordeDashboard/src/backend/Backend.ts
@@ -1014,7 +1014,13 @@ export class Backend {
     }
 
     getJiraIssues(streamId: string, jiraKeys: string[]): Promise<GetJiraIssueResponse[]> {
+
+        // disabled until hooked up for prod
+        return new Promise<GetJiraIssueResponse[]>((resolve, reject) => {
+            resolve([]);
+        });
     
+        /*
         const params = {
             streamId: streamId,
             jiraKeys: jiraKeys
@@ -1027,6 +1033,7 @@ export class Backend {
                 reject(reason);
             });
         });
+        */
 
     }
 
