Horde: Remove issue preference
#rnx

diff --git a/Engine/Source/Programs/Horde/Horde.Build/Api/Users.cs b/Engine/Source/Programs/Horde/Horde.Build/Api/Users.cs
--- a/Engine/Source/Programs/Horde/Horde.Build/Api/Users.cs
+++ b/Engine/Source/Programs/Horde/Horde.Build/Api/Users.cs
@@ -62,11 +62,6 @@ namespace HordeServer.Api
 		/// </summary>
 		public bool? EnableExperimentalFeatures { get; set; }
 
-		/// <summary>
-		/// Whether to enable slack notifications for this user
-		/// </summary>
-		public bool? EnableIssueNotifications { get; set; }
-
 		/// <summary>
 		/// Settings for the dashboard
 		/// </summary>
@@ -96,7 +91,7 @@ namespace HordeServer.Api
 			if (Settings != null)
 			{
 				this.EnableExperimentalFeatures = Settings.EnableExperimentalFeatures;
-				this.EnableIssueNotifications = Settings.EnableIssueNotifications;
+				
 				this.DashboardSettings = BsonTypeMapper.MapToDotNetValue(Settings.DashboardSettings);
 				this.PinnedJobIds = Settings.PinnedJobIds.ConvertAll(x => x.ToString());
 			}
@@ -161,11 +156,6 @@ namespace HordeServer.Api
 		/// </summary>
 		public bool? EnableExperimentalFeatures { get; set; }
 
-		/// <summary>
-		/// Whether to enable Slack notifications for this user
-		/// </summary>
-		public bool? EnableIssueNotifications { get; set; }
-
 		/// <summary>
 		/// New dashboard settings
 		/// </summary>
diff --git a/Engine/Source/Programs/Horde/Horde.Build/Controllers/UserController.cs b/Engine/Source/Programs/Horde/Horde.Build/Controllers/UserController.cs
--- a/Engine/Source/Programs/Horde/Horde.Build/Controllers/UserController.cs
+++ b/Engine/Source/Programs/Horde/Horde.Build/Controllers/UserController.cs
@@ -81,7 +81,7 @@ namespace HordeServer.Controllers
 			{
 				return BadRequest("Current user does not have a registered profile");
 			}
-			await UserCollection.UpdateSettingsAsync(UserId.Value, Request.EnableExperimentalFeatures, Request.EnableIssueNotifications, Request.DashboardSettings?.ToBsonValue(), Request.AddPinnedJobIds?.Select(x => new JobId(x)), Request.RemovePinnedJobIds?.Select(x => new JobId(x)));
+			await UserCollection.UpdateSettingsAsync(UserId.Value, Request.EnableExperimentalFeatures, Request.DashboardSettings?.ToBsonValue(), Request.AddPinnedJobIds?.Select(x => new JobId(x)), Request.RemovePinnedJobIds?.Select(x => new JobId(x)));
 			return Ok();
 		}
 	}
diff --git a/Engine/Source/Programs/Horde/Horde.Build/Notifications/Impl/SlackNotificationSink.cs b/Engine/Source/Programs/Horde/Horde.Build/Notifications/Impl/SlackNotificationSink.cs
--- a/Engine/Source/Programs/Horde/Horde.Build/Notifications/Impl/SlackNotificationSink.cs
+++ b/Engine/Source/Programs/Horde/Horde.Build/Notifications/Impl/SlackNotificationSink.cs
@@ -641,12 +641,6 @@ namespace HordeServer.Notifications.Impl
 
 		async Task NotifyIssueUpdatedAsync(IUser User, IIssue Issue, IIssueDetails Details)
 		{
-			IUserSettings UserSettings = await UserCollection.GetSettingsAsync(User.Id);
-			if (!UserSettings.EnableIssueNotifications)
-			{
-				Logger.LogInformation("Issue notifications are disabled for user {UserId} ({UserName})", User.Id, User.Name);
-				return;
-			}
 
 			string? SlackUserId = await GetSlackUserId(User);
 			if (SlackUserId == null)
diff --git a/Engine/Source/Programs/Horde/Horde.Build/Users/IUser.cs b/Engine/Source/Programs/Horde/Horde.Build/Users/IUser.cs
--- a/Engine/Source/Programs/Horde/Horde.Build/Users/IUser.cs
+++ b/Engine/Source/Programs/Horde/Horde.Build/Users/IUser.cs
@@ -81,11 +81,6 @@ namespace HordeServer.Models
 		/// </summary>
 		public bool EnableExperimentalFeatures { get; }
 
-		/// <summary>
-		/// Whether to enable Slack notifications
-		/// </summary>
-		public bool EnableIssueNotifications { get; }
-
 		/// <summary>
 		/// Opaque settings dictionary for the dashboard
 		/// </summary>
diff --git a/Engine/Source/Programs/Horde/Horde.Build/Users/IUserCollection.cs b/Engine/Source/Programs/Horde/Horde.Build/Users/IUserCollection.cs
--- a/Engine/Source/Programs/Horde/Horde.Build/Users/IUserCollection.cs
+++ b/Engine/Source/Programs/Horde/Horde.Build/Users/IUserCollection.cs
@@ -87,12 +87,11 @@ namespace HordeServer.Collections
 		/// </summary>
 		/// <param name="UserId">The user to update</param>
 		/// <param name="EnableExperimentalFeatures"></param>
-		/// <param name="EnableIssueNotifications"></param>
 		/// <param name="DashboardSettings">Opaque settings object for the dashboard</param>
 		/// <param name="AddPinnedJobIds"></param>
 		/// <param name="RemovePinnedJobIds"></param>
 		/// <returns>Updated user object</returns>
-		Task UpdateSettingsAsync(UserId UserId, bool? EnableExperimentalFeatures = null, bool? EnableIssueNotifications = null, BsonValue? DashboardSettings = null, IEnumerable<JobId>? AddPinnedJobIds = null, IEnumerable<JobId>? RemovePinnedJobIds = null);
+		Task UpdateSettingsAsync(UserId UserId, bool? EnableExperimentalFeatures = null, BsonValue? DashboardSettings = null, IEnumerable<JobId>? AddPinnedJobIds = null, IEnumerable<JobId>? RemovePinnedJobIds = null);
 	}
 
 	/// <summary>
diff --git a/Engine/Source/Programs/Horde/Horde.Build/Users/Impl/UserCollection.cs b/Engine/Source/Programs/Horde/Horde.Build/Users/Impl/UserCollection.cs
--- a/Engine/Source/Programs/Horde/Horde.Build/Users/Impl/UserCollection.cs
+++ b/Engine/Source/Programs/Horde/Horde.Build/Users/Impl/UserCollection.cs
@@ -33,9 +33,6 @@ namespace HordeServer.Collections.Impl
 			[BsonDefaultValue(false), BsonIgnoreIfDefault]
 			public bool EnableExperimentalFeatures { get; set; }
 
-			[BsonDefaultValue(false), BsonIgnoreIfDefault]
-			public bool EnableIssueNotifications { get; set; }
-
 			public BsonValue DashboardSettings { get; set; } = BsonNull.Value;
 			public List<JobId> PinnedJobIds { get; set; } = new List<JobId>();
 
@@ -160,7 +157,7 @@ namespace HordeServer.Collections.Impl
 		}
 
 		/// <inheritdoc/>
-		public async Task UpdateSettingsAsync(UserId UserId, bool? EnableExperimentalFeatures, bool? EnableIssueNotifications, BsonValue? DashboardSettings = null, IEnumerable<JobId>? AddPinnedJobIds = null, IEnumerable<JobId>? RemovePinnedJobIds = null)
+		public async Task UpdateSettingsAsync(UserId UserId, bool? EnableExperimentalFeatures, BsonValue? DashboardSettings = null, IEnumerable<JobId>? AddPinnedJobIds = null, IEnumerable<JobId>? RemovePinnedJobIds = null)
 		{
 			if (AddPinnedJobIds != null)
 			{
@@ -177,10 +174,6 @@ namespace HordeServer.Collections.Impl
 			{
 				Updates.Add(Builders<UserDocument>.Update.SetOrUnsetNull(x => x.EnableExperimentalFeatures, EnableExperimentalFeatures));
 			}
-			if (EnableIssueNotifications != null)
-			{
-				Updates.Add(Builders<UserDocument>.Update.SetOrUnsetNull(x => x.EnableIssueNotifications, EnableIssueNotifications));
-			}
 			if (DashboardSettings != null)
 			{
 				Updates.Add(Builders<UserDocument>.Update.Set(x => x.DashboardSettings, DashboardSettings));
diff --git a/Engine/Source/Programs/Horde/Horde.Build/Users/Impl/UserCollectionV2.cs b/Engine/Source/Programs/Horde/Horde.Build/Users/Impl/UserCollectionV2.cs
--- a/Engine/Source/Programs/Horde/Horde.Build/Users/Impl/UserCollectionV2.cs
+++ b/Engine/Source/Programs/Horde/Horde.Build/Users/Impl/UserCollectionV2.cs
@@ -94,9 +94,6 @@ namespace HordeServer.Collections.Impl
 			[BsonDefaultValue(false), BsonIgnoreIfDefault]
 			public bool EnableExperimentalFeatures { get; set; }
 
-			[BsonDefaultValue(true), BsonIgnoreIfDefault]
-			public bool EnableIssueNotifications { get; set; } = true;
-
 			public BsonValue DashboardSettings { get; set; } = BsonNull.Value;
 			public List<JobId> PinnedJobIds { get; set; } = new List<JobId>();
 
@@ -117,7 +114,6 @@ namespace HordeServer.Collections.Impl
 				: this(Other.UserId)
 			{
 				this.EnableExperimentalFeatures = Other.EnableExperimentalFeatures;
-				this.EnableIssueNotifications = Other.EnableIssueNotifications;
 				this.DashboardSettings = Other.DashboardSettings;
 				this.PinnedJobIds = new List<JobId>(Other.PinnedJobIds);
 			}
@@ -307,24 +303,13 @@ namespace HordeServer.Collections.Impl
 		}
 
 		/// <inheritdoc/>
-		public async Task UpdateSettingsAsync(UserId UserId, bool? EnableExperimentalFeatures = null, bool? EnableIssueNotifications = null, BsonValue? DashboardSettings = null, IEnumerable<JobId>? AddPinnedJobIds = null, IEnumerable<JobId>? RemovePinnedJobIds = null)
+		public async Task UpdateSettingsAsync(UserId UserId, bool? EnableExperimentalFeatures = null, BsonValue? DashboardSettings = null, IEnumerable<JobId>? AddPinnedJobIds = null, IEnumerable<JobId>? RemovePinnedJobIds = null)
 		{
 			List<UpdateDefinition<UserSettingsDocument>> Updates = new List<UpdateDefinition<UserSettingsDocument>>();
 			if (EnableExperimentalFeatures != null)
 			{
 				Updates.Add(Builders<UserSettingsDocument>.Update.SetOrUnsetNull(x => x.EnableExperimentalFeatures, EnableExperimentalFeatures));
 			}
-			if (EnableIssueNotifications != null)
-			{
-				if (EnableIssueNotifications.Value)
-				{
-					Updates.Add(Builders<UserSettingsDocument>.Update.Unset(x => x.EnableIssueNotifications));
-				}
-				else
-				{
-					Updates.Add(Builders<UserSettingsDocument>.Update.Set(x => x.EnableIssueNotifications, false));
-				}
-			}
 			if (DashboardSettings != null)
 			{
 				Updates.Add(Builders<UserSettingsDocument>.Update.Set(x => x.DashboardSettings, DashboardSettings));
diff --git a/Engine/Source/Programs/Horde/HordeDashboard/src/backend/Api.ts b/Engine/Source/Programs/Horde/HordeDashboard/src/backend/Api.ts
--- a/Engine/Source/Programs/Horde/HordeDashboard/src/backend/Api.ts
+++ b/Engine/Source/Programs/Horde/HordeDashboard/src/backend/Api.ts
@@ -2861,9 +2861,6 @@ export type GetUserResponse = {
 	/** Claims for the user */
 	claims?: UserClaim[];
 
-	/** Whether to enable slack notifications for this user */
-	enableIssueNotifications?: boolean;
-
 	/** Whether to enable experimental features for this user */
 	enableExperimentalFeatures?: boolean;
 
@@ -2904,9 +2901,6 @@ export type UpdateUserRequest = {
 	/** Job ids to remove from the pinned list */
 	removePinnedJobIds?: string[];
 
-	/** Whether to enable slack notifications for this user */
-	enableIssueNotifications?: boolean;
-
 	/** Whether to enable experimental features for this user */
 	enableExperimentalFeatures?: boolean;
 
diff --git a/Engine/Source/Programs/Horde/HordeDashboard/src/backend/Dashboard.ts b/Engine/Source/Programs/Horde/HordeDashboard/src/backend/Dashboard.ts
--- a/Engine/Source/Programs/Horde/HordeDashboard/src/backend/Dashboard.ts
+++ b/Engine/Source/Programs/Horde/HordeDashboard/src/backend/Dashboard.ts
@@ -133,22 +133,6 @@ export class Dashboard {
         return this.data.pinnedJobIds ?? [];
     }
 
-    get issueNotifications(): boolean {
-
-        return this.data.enableIssueNotifications!;
-    }
-
-    set issueNotifications(value: boolean) {
-
-        backend.updateUser({ enableIssueNotifications: value }).then(() => {
-
-            this.data.enableIssueNotifications = value;
-            this.setUpdated();
-
-        });
-
-    }
-
     get experimentalFeatures(): boolean {
 
         return this.data.enableExperimentalFeatures!;
@@ -484,7 +468,7 @@ export class Dashboard {
 
     serverSettingsChanged: boolean = false;
 
-    private data: GetUserResponse = { id: "", name: "", enableIssueNotifications: false, enableExperimentalFeatures: false, claims: [], pinnedJobIds: [], dashboardSettings: { preferences: new Map() } };
+    private data: GetUserResponse = { id: "", name: "", enableExperimentalFeatures: false, claims: [], pinnedJobIds: [], dashboardSettings: { preferences: new Map() } };
 
     private updateTimeoutId: any = undefined;
 
