Trying to fix it so author's only get notified of a step with errors or warnings when the job first transitions into that state.

#rnx

diff --git a/Engine/Source/Programs/Horde/Horde.Build/Notifications/Impl/NotificationService.cs b/Engine/Source/Programs/Horde/Horde.Build/Notifications/Impl/NotificationService.cs
--- a/Engine/Source/Programs/Horde/Horde.Build/Notifications/Impl/NotificationService.cs
+++ b/Engine/Source/Programs/Horde/Horde.Build/Notifications/Impl/NotificationService.cs
@@ -487,8 +487,11 @@ namespace HordeServer.Notifications.Impl
 
 			List<IUser> UsersToNotify = await GetUsersToNotify(EventRecord, Step.NotificationTriggerId, true);
 
-			// If this is not a success notification and the author isn't in the list to notify, add them manually.
-			if (Job.StartedByUserId.HasValue && !UsersToNotify.Any(x => x.Id == Job.StartedByUserId) && Step.Outcome != JobStepOutcome.Success)
+			// If this is not a success notification and the author isn't in the list to notify, add them manually if this is the outcome has gotten worse.
+			int Failures = Job.Batches.Sum(x => x.Steps.Count(y => y.Outcome == JobStepOutcome.Failure));
+			bool FirstFailure = Step.Outcome == JobStepOutcome.Failure && Failures == 1;
+			bool FirstWarning = Step.Outcome == JobStepOutcome.Warnings && Failures == 0 && Job.Batches.Sum(x => x.Steps.Count(y => y.Outcome == JobStepOutcome.Warnings)) == 1;
+			if (Job.StartedByUserId.HasValue && !UsersToNotify.Any(x => x.Id == Job.StartedByUserId) && (FirstFailure || FirstWarning))
 			{
 				Logger.LogInformation("Author {AuthorUserId} is not in notify list but step outcome is {JobStepOutcome}, adding them to the list...", Job.StartedByUserId, Step.Outcome);
 				IUser? AuthorUser = await UserCollection.GetUserAsync(Job.StartedByUserId.Value);
