Horde: Fix PoolId/ObjectId mixed usage
#rnx

diff --git a/Engine/Source/Programs/Horde/Horde.Build/Agents/IAgentCollection.cs b/Engine/Source/Programs/Horde/Horde.Build/Agents/IAgentCollection.cs
--- a/Engine/Source/Programs/Horde/Horde.Build/Agents/IAgentCollection.cs
+++ b/Engine/Source/Programs/Horde/Horde.Build/Agents/IAgentCollection.cs
@@ -60,14 +60,13 @@ namespace HordeServer.Collections
 		/// <summary>
 		/// Finds all agents matching certain criteria
 		/// </summary>
-		/// <param name="Pool">The pool containing the agent</param>
 		/// <param name="PoolId">The pool ID in string form containing the agent</param>
 		/// <param name="ModifiedAfter">If set, only returns agents modified after this time</param>
 		/// <param name="Status">Status to look for</param>
 		/// <param name="Index">Index of the first result</param>
 		/// <param name="Count">Number of results to return</param>
 		/// <returns>List of agents matching the given criteria</returns>
-		Task<List<IAgent>> FindAsync(ObjectId? Pool = null, string? PoolId = null, DateTime? ModifiedAfter = null, AgentStatus? Status = null, int? Index = null, int? Count = null);
+		Task<List<IAgent>> FindAsync(PoolId? PoolId = null, DateTime? ModifiedAfter = null, AgentStatus? Status = null, int? Index = null, int? Count = null);
 
 		/// <summary>
 		/// Finds all the expried agents
diff --git a/Engine/Source/Programs/Horde/Horde.Build/Agents/Impl/AgentCollection.cs b/Engine/Source/Programs/Horde/Horde.Build/Agents/Impl/AgentCollection.cs
--- a/Engine/Source/Programs/Horde/Horde.Build/Agents/Impl/AgentCollection.cs
+++ b/Engine/Source/Programs/Horde/Horde.Build/Agents/Impl/AgentCollection.cs
@@ -193,15 +193,11 @@ namespace HordeServer.Collections.Impl
 		}
 
 		/// <inheritdoc/>
-		public async Task<List<IAgent>> FindAsync(ObjectId? Pool, string? PoolId, DateTime? ModifiedAfter, AgentStatus? Status, int? Index, int? Count)
+		public async Task<List<IAgent>> FindAsync(PoolId? PoolId, DateTime? ModifiedAfter, AgentStatus? Status, int? Index, int? Count)
 		{
 			FilterDefinitionBuilder<AgentDocument> FilterBuilder = new FilterDefinitionBuilder<AgentDocument>();
 
 			FilterDefinition<AgentDocument> Filter = FilterBuilder.Ne(x => x.Deleted, true);
-			if (Pool != null)
-			{
-				Filter &= FilterBuilder.Eq(nameof(AgentDocument.Pools), Pool);
-			}
 			
 			if (PoolId != null)
 			{
diff --git a/Engine/Source/Programs/Horde/Horde.Build/Agents/Impl/AgentService.cs b/Engine/Source/Programs/Horde/Horde.Build/Agents/Impl/AgentService.cs
--- a/Engine/Source/Programs/Horde/Horde.Build/Agents/Impl/AgentService.cs
+++ b/Engine/Source/Programs/Horde/Horde.Build/Agents/Impl/AgentService.cs
@@ -173,14 +173,14 @@ namespace HordeServer.Services
 		/// <summary>
 		/// Finds all agents matching certain criteria
 		/// </summary>
-		/// <param name="Pool">The pool containing the agent</param>
+		/// <param name="PoolId">The pool containing the agent</param>
 		/// <param name="ModifiedAfter">If set, only returns agents modified after this time</param>
 		/// <param name="Index">Index within the list of results</param>
 		/// <param name="Count">Number of results to return</param>
 		/// <returns>List of agents matching the given criteria</returns>
-		public Task<List<IAgent>> FindAgentsAsync(ObjectId? Pool, DateTime? ModifiedAfter, int? Index, int? Count)
+		public Task<List<IAgent>> FindAgentsAsync(PoolId? PoolId, DateTime? ModifiedAfter, int? Index, int? Count)
 		{
-			return Agents.FindAsync(Pool, null, ModifiedAfter, null, Index, Count);
+			return Agents.FindAsync(PoolId, ModifiedAfter, null, Index, Count);
 		}
 
 		/// <summary>
diff --git a/Engine/Source/Programs/Horde/Horde.Build/Controllers/AgentsController.cs b/Engine/Source/Programs/Horde/Horde.Build/Controllers/AgentsController.cs
--- a/Engine/Source/Programs/Horde/Horde.Build/Controllers/AgentsController.cs
+++ b/Engine/Source/Programs/Horde/Horde.Build/Controllers/AgentsController.cs
@@ -81,7 +81,7 @@ namespace HordeServer.Controllers
 		/// <summary>
 		/// Finds the agents matching specified criteria.
 		/// </summary>
-		/// <param name="Pool">The pool containing the agent</param>
+		/// <param name="PoolId">The pool containing the agent</param>
 		/// <param name="Index">First result to return</param>
 		/// <param name="Count">Number of results to return</param>
 		/// <param name="ModifiedAfter">If set, only returns agents modified after this time</param>
@@ -90,10 +90,10 @@ namespace HordeServer.Controllers
 		[HttpGet]
 		[Route("/api/v1/agents")]
 		[ProducesResponseType(typeof(List<GetAgentResponse>), 200)]
-		public async Task<ActionResult<List<object>>> FindAgentsAsync([FromQuery] string? Pool = null, [FromQuery] int? Index = null, [FromQuery] int? Count = null, [FromQuery] DateTimeOffset? ModifiedAfter = null, [FromQuery] PropertyFilter? Filter = null)
+		public async Task<ActionResult<List<object>>> FindAgentsAsync([FromQuery] PoolId? PoolId = null, [FromQuery] int? Index = null, [FromQuery] int? Count = null, [FromQuery] DateTimeOffset? ModifiedAfter = null, [FromQuery] PropertyFilter? Filter = null)
 		{
 			GlobalPermissionsCache PermissionsCache = new GlobalPermissionsCache();
-			List<IAgent> Agents = await AgentService.FindAgentsAsync(Pool?.ToObjectId(), ModifiedAfter?.UtcDateTime, Index, Count);
+			List<IAgent> Agents = await AgentService.FindAgentsAsync(PoolId, ModifiedAfter?.UtcDateTime, Index, Count);
 
 			List<object> Responses = new List<object>();
 			foreach (IAgent Agent in Agents)
