Horde Storage - If no start time is specified for the replication to read from we just assume to start from time to live, so 2 weeks ago by default. This shouldnt be used to often as we should be starting from were we were last, or from a snapshot, but when this is used it does put a lot of load on the db.

diff --git a/Engine/Source/Programs/Horde/HordeStorage/Horde.Storage/Implementation/TransactionLog/ScyllaReplicationLog.cs b/Engine/Source/Programs/Horde/HordeStorage/Horde.Storage/Implementation/TransactionLog/ScyllaReplicationLog.cs
--- a/Engine/Source/Programs/Horde/HordeStorage/Horde.Storage/Implementation/TransactionLog/ScyllaReplicationLog.cs
+++ b/Engine/Source/Programs/Horde/HordeStorage/Horde.Storage/Implementation/TransactionLog/ScyllaReplicationLog.cs
@@ -18,7 +18,6 @@ namespace Horde.Storage.Implementation
         private readonly IOptionsMonitor<ScyllaSettings> _settings;
         private readonly ISession _session;
         private readonly Mapper _mapper;
-        private readonly PreparedStatement _selectReplicationBucketStatement;
 
         public ScyllaReplicationLog(IScyllaSessionManager scyllaSessionManager, IOptionsMonitor<ScyllaSettings> settings)
         {
@@ -52,9 +51,6 @@ namespace Horde.Storage.Implementation
                 PRIMARY KEY ((namespace))
             );"  
             ));
-
-            _selectReplicationBucketStatement = _session.Prepare(
-                "SELECT replication_bucket FROM replication_log WHERE namespace = ? PER PARTITION LIMIT 1 ALLOW FILTERING");
         }
 
         public async IAsyncEnumerable<NamespaceId> GetNamespaces()
@@ -174,24 +170,10 @@ namespace Horde.Storage.Implementation
             else
             {
                 using Scope _ = Tracer.Instance.StartActive("scylla.determine_first_replication_bucket");
-                SortedSet<long> buckets = new SortedSet<long>();
-
-                // fetch all the buckets that exists and sort them based on time
-                RowSet replicationBuckets = await _session.ExecuteAsync(_selectReplicationBucketStatement.Bind(ns.ToString()));
-                foreach (Row replicationBucket in replicationBuckets)
-                {
-                    long bucketField = (long)replicationBucket["replication_bucket"];
-                    buckets.Add(bucketField);
-                }
-
-                // if there are no buckets we stop here
-                if (buckets.Count == 0)
-                    yield break;
 
-                // pick the first bucket
-                long bucket = buckets.First();
-                startBucketTime = DateTime.FromFileTimeUtc(bucket);
-                yield return bucket;
+                // we should have no data older then the ttl to lets just assume that the bucket to start searching from is now - time to live
+                DateTime oldestTimestamp = DateTime.Now.AddSeconds(-1 * _settings.CurrentValue.ReplicationLogTimeToLive.TotalSeconds);
+                startBucketTime = oldestTimestamp;
             }
 
             // ignore any bucket that is older then a cutoff, as that can cause us to end up scanning thru a lot of hours that will never exist (incremental logs are deleted after 7 days)
