Horde Storage - Reduced scopes which have access to the transaction log - will be used to verify that we can actually remove these scopes from Okta.

diff --git a/Engine/Source/Programs/Horde/HordeStorage/Callisto/CallistoStartup.cs b/Engine/Source/Programs/Horde/HordeStorage/Callisto/CallistoStartup.cs
--- a/Engine/Source/Programs/Horde/HordeStorage/Callisto/CallistoStartup.cs
+++ b/Engine/Source/Programs/Horde/HordeStorage/Callisto/CallistoStartup.cs
@@ -5,12 +5,9 @@ using System.Collections.Generic;
 using System.ComponentModel.DataAnnotations;
 using System.IO;
 using Callisto.Implementation;
-using JsonSubTypes;
 using Jupiter;
-using Jupiter.Implementation;
 using Microsoft.AspNetCore.Authorization;
 using Microsoft.AspNetCore.Hosting;
-using Microsoft.AspNetCore.Mvc;
 using Microsoft.Extensions.Configuration;
 using Microsoft.Extensions.DependencyInjection;
 using Microsoft.Extensions.Options;
@@ -45,13 +42,13 @@ namespace Callisto
             authorizationOptions.AddPolicy("TLog.read", policy =>
             {
                 policy.AuthenticationSchemes = defaultSchemes;
-                policy.RequireClaim("transactionLog", "read", "readwrite", "full");
+                policy.RequireClaim("transactionLog", "full");
             });
 
             authorizationOptions.AddPolicy("TLog.write", policy =>
             {
                 policy.AuthenticationSchemes = defaultSchemes;
-                policy.RequireClaim("transactionLog", "write", "readwrite", "full");
+                policy.RequireClaim("transactionLog", "full");
             });
 
             authorizationOptions.AddPolicy("TLog.delete", policy =>
