Horde: Allow dylib files when building Horde server.

diff --git a/Engine/Source/Programs/Horde/BuildHorde.xml b/Engine/Source/Programs/Horde/BuildHorde.xml
--- a/Engine/Source/Programs/Horde/BuildHorde.xml
+++ b/Engine/Source/Programs/Horde/BuildHorde.xml
@@ -22,7 +22,7 @@
 				Engine/Source/ThirdParty/Perforce/...
 			</Property>
 			<Tag Files="$(StagePaths)" Except=".../.vs/...;.../.git/...;.../bin/...;.../obj/..." With="#InputFiles"/>
-			<Tag Files="Engine/Source/Programs/Shared/EpicGames.Perforce.Native/bin/..." Filter="*.pdb;*.dll;*.so" With="#InputFiles"/>
+			<Tag Files="Engine/Source/Programs/Shared/EpicGames.Perforce.Native/bin/..." Filter="*.pdb;*.dll;*.so;*.dylib" With="#InputFiles"/>
 
 			<!-- Build the Docker image and publish it -->
 			<Docker-Build BaseDir="Engine/Source" Files="#InputFiles" Tag="hordeserver-public" DockerFile="Engine/Source/Programs/Horde/Horde.Build/Dockerfile" Arguments="--build-arg msbuild_args=&quot;$(VersionArguments)&quot;"/>
