Horde: Output a warning with callstack when attempting to add a user with an invalid login name.

diff --git a/Engine/Source/Programs/Horde/Horde.Build/Users/Impl/UserCollectionV2.cs b/Engine/Source/Programs/Horde/Horde.Build/Users/Impl/UserCollectionV2.cs
--- a/Engine/Source/Programs/Horde/Horde.Build/Users/Impl/UserCollectionV2.cs
+++ b/Engine/Source/Programs/Horde/Horde.Build/Users/Impl/UserCollectionV2.cs
@@ -252,6 +252,11 @@ namespace HordeServer.Collections.Impl
 		/// <inheritdoc/>
 		public async Task<IUser> FindOrAddUserByLoginAsync(string Login, string? Name, string? Email)
 		{
+			if (Login.Contains(' ', StringComparison.Ordinal))
+			{
+				Logger.LogWarning("Potentially invalid login name: {Login} (from {Trace})", Login, Environment.StackTrace);
+			}
+
 			UpdateDefinition<UserDocument> Update = Builders<UserDocument>.Update.SetOnInsert(x => x.Id, UserId.GenerateNewId()).SetOnInsert(x => x.Login, Login).Unset(x => x.Hidden);
 
 			if (Name == null)
