Horde: Additional logging to track down hang reading plain text output from logs.

diff --git a/Engine/Source/Programs/Horde/Horde.Build/Logs/Impl/LogFileService.cs b/Engine/Source/Programs/Horde/Horde.Build/Logs/Impl/LogFileService.cs
--- a/Engine/Source/Programs/Horde/Horde.Build/Logs/Impl/LogFileService.cs
+++ b/Engine/Source/Programs/Horde/Horde.Build/Logs/Impl/LogFileService.cs
@@ -231,13 +231,17 @@ namespace HordeServer.Services
 					int ReadBytes = await Stream.ReadAsync(ReadBuffer.AsMemory(ReadBufferLength, ReadBuffer.Length - ReadBufferLength));
 					ReadBufferLength += ReadBytes;
 
+					Serilog.Log.Debug("LOG: Read {NumBytes}", ReadBytes);
+
 					// Copy as many lines as possible to the output
 					int ConvertedBytes = 0;
 					for(int EndIdx = 1; EndIdx < ReadBufferLength; EndIdx++)
 					{
 						if (ReadBuffer[EndIdx] == '\n')
 						{
+							Serilog.Log.Debug("LOG: Parsing {Start} -> {End}", ConvertedBytes, EndIdx);
 							WriteBufferLength = LogText.ConvertToPlainText(ReadBuffer.AsSpan(ConvertedBytes, EndIdx - ConvertedBytes), WriteBuffer, WriteBufferLength);
+							Serilog.Log.Debug("LOG: Done");
 							ConvertedBytes = EndIdx + 1;
 						}
 					}
@@ -248,7 +252,9 @@ namespace HordeServer.Services
 						if (Offset < WriteBufferLength)
 						{
 							int WriteLength = (int)Math.Min((long)WriteBufferLength - Offset, Length);
+							Serilog.Log.Debug("LOG: Writing {Start} -> {End}", Offset, WriteLength);
 							await OutputStream.WriteAsync(WriteBuffer.AsMemory((int)Offset, WriteLength));
+							Serilog.Log.Debug("LOG: Done");
 							Length -= WriteLength;
 						}
 						Offset = Math.Max(Offset - WriteBufferLength, 0);
@@ -260,16 +266,19 @@ namespace HordeServer.Services
 					{
 						Buffer.BlockCopy(ReadBuffer, ConvertedBytes, ReadBuffer, 0, ReadBufferLength - ConvertedBytes);
 						ReadBufferLength -= ConvertedBytes;
+						Serilog.Log.Debug("LOG: Removing {NumBytes}", ConvertedBytes);
 					}
 					else if(ReadBufferLength > 0)
 					{
 						Array.Resize(ref ReadBuffer, ReadBuffer.Length + 128);
 						WriteBuffer = new byte[ReadBuffer.Length];
+						Serilog.Log.Debug("LOG: Expanding to {Length}", ReadBuffer.Length);
 					}
 
 					// Exit if we didn't read anything in this iteration
-					if(ReadBytes == 0)
+					if (ReadBytes == 0)
 					{
+						Serilog.Log.Debug("LOG: COMPLETE");
 						break;
 					}
 				}
