Horde: Increase logging level for issue service.

diff --git a/Engine/Source/Programs/Horde/Horde.Build/Issues/Impl/IssueService.cs b/Engine/Source/Programs/Horde/Horde.Build/Issues/Impl/IssueService.cs
--- a/Engine/Source/Programs/Horde/Horde.Build/Issues/Impl/IssueService.cs
+++ b/Engine/Source/Programs/Horde/Horde.Build/Issues/Impl/IssueService.cs
@@ -543,6 +543,8 @@ namespace HordeServer.Services.Impl
 			Scope.Span.SetTag("BatchId", BatchId.ToString());
 			Scope.Span.SetTag("StepId", StepId.ToString());
 
+			Logger.LogInformation("Updating issues for {JobId}:{BatchId}:{StepId}", Job.Id, BatchId, StepId);
+
 			IStream? Stream = await Streams.GetStreamAsync(Job.StreamId);
 			if (Stream == null)
 			{
diff --git a/Engine/Source/Programs/Horde/Horde.Build/appsettings.json b/Engine/Source/Programs/Horde/Horde.Build/appsettings.json
--- a/Engine/Source/Programs/Horde/Horde.Build/appsettings.json
+++ b/Engine/Source/Programs/Horde/Horde.Build/appsettings.json
@@ -1,6 +1,5 @@
 {
-	"Horde": 
-	{
+	"Horde": {
 		"DashboardUrl": "https://localhost:3000",
 		"DisableAuth": true,
 		"CorsEnabled": true,
@@ -12,16 +11,16 @@
 		"MinimumLevel": 
 		{
 			"Default": "Information",
-			"Override": 
-			{
+			"Override": {
 				"MongoDB": "Warning", // For bundled MongoDB output
-				"Redis" : "Warning", // For bundled Redis output
+				"Redis": "Warning", // For bundled Redis output
 				"Microsoft": "Warning",
 				"Microsoft.AspNetCore": "Warning",
 				"Microsoft.Hosting": "Information",
 				"HordeServer.Authentication": "Warning",
 				"System.Net.Http.HttpClient": "Warning",
-				"Grpc": "Warning"
+				"Grpc": "Warning",
+				"HordeServer.Services.Impl.IssueService": "Debug"
 			}
 		}
 	},
