Horde: Fix logging for log service.

diff --git a/Engine/Source/Programs/Horde/Horde.Build/Controllers/LogsController.cs b/Engine/Source/Programs/Horde/Horde.Build/Controllers/LogsController.cs
--- a/Engine/Source/Programs/Horde/Horde.Build/Controllers/LogsController.cs
+++ b/Engine/Source/Programs/Horde/Horde.Build/Controllers/LogsController.cs
@@ -20,6 +20,7 @@ using HordeServer.Utilities;
 using Microsoft.AspNetCore.Authorization;
 using Microsoft.AspNetCore.Http;
 using Microsoft.AspNetCore.Mvc;
+using Microsoft.Extensions.Logging;
 using MongoDB.Bson;
 
 namespace HordeServer.Controllers
@@ -51,39 +52,22 @@ namespace HordeServer.Controllers
 	[Route("[controller]")]
 	public class LogsController : ControllerBase
 	{
-		/// <summary>
-		/// Instance of the LogFile service
-		/// </summary>
 		private readonly ILogFileService LogFileService;
-
-		/// <summary>
-		/// Instance of the issue collection
-		/// </summary>
 		private readonly IIssueCollection IssueCollection;
-
-		/// <summary>
-		/// Instance of the ACL service
-		/// </summary>
 		private readonly AclService AclService;
-
-		/// <summary>
-		/// Instance of the Job service
-		/// </summary>
 		private readonly JobService JobService;
+		private readonly ILogger Logger;
 
 		/// <summary>
 		/// Constructor
 		/// </summary>
-		/// <param name="LogFileService">The Logfile service</param>
-		/// <param name="IssueCollection">The issue collection</param>
-		/// <param name="AclService">The ACL service</param>
-		/// <param name="JobService">The Job service</param>
-		public LogsController(ILogFileService LogFileService, IIssueCollection IssueCollection, AclService AclService, JobService JobService)
+		public LogsController(ILogFileService LogFileService, IIssueCollection IssueCollection, AclService AclService, JobService JobService, ILogger<LogsController> Logger)
 		{
 			this.LogFileService = LogFileService;
 			this.IssueCollection = IssueCollection;
 			this.AclService = AclService;
 			this.JobService = JobService;
+			this.Logger = Logger;
 		}
 
 		/// <summary>
@@ -138,7 +122,7 @@ namespace HordeServer.Controllers
 			Func<Stream, ActionContext, Task> CopyTask;
 			if (Format == LogOutputFormat.Text && LogFile.Type == LogType.Json)
 			{
-				CopyTask = (OutputStream, Context) => LogFileService.CopyPlainTextStreamAsync(LogFile, Offset, Length, OutputStream);
+				CopyTask = (OutputStream, Context) => LogFileService.CopyPlainTextStreamAsync(LogFile, Offset, Length, OutputStream, Logger);
 			}
 			else
 			{
diff --git a/Engine/Source/Programs/Horde/Horde.Build/Logs/Impl/LogFileService.cs b/Engine/Source/Programs/Horde/Horde.Build/Logs/Impl/LogFileService.cs
--- a/Engine/Source/Programs/Horde/Horde.Build/Logs/Impl/LogFileService.cs
+++ b/Engine/Source/Programs/Horde/Horde.Build/Logs/Impl/LogFileService.cs
@@ -214,8 +214,9 @@ namespace HordeServer.Services
 		/// <param name="Offset">Offset within the data to copy from</param>
 		/// <param name="Length">Length of the data to copy</param>
 		/// <param name="OutputStream">Output stream to receive the text data</param>
+		/// <param name="Logger">Logger for output messages</param>
 		/// <returns>Async text</returns>
-		public static async Task CopyPlainTextStreamAsync(this ILogFileService LogFileService, ILogFile LogFile, long Offset, long Length, Stream OutputStream)
+		public static async Task CopyPlainTextStreamAsync(this ILogFileService LogFileService, ILogFile LogFile, long Offset, long Length, Stream OutputStream, ILogger Logger)
 		{
 			using (Stream Stream = await LogFileService.OpenRawStreamAsync(LogFile, 0, long.MaxValue))
 			{
@@ -231,7 +232,7 @@ namespace HordeServer.Services
 					int ReadBytes = await Stream.ReadAsync(ReadBuffer.AsMemory(ReadBufferLength, ReadBuffer.Length - ReadBufferLength));
 					ReadBufferLength += ReadBytes;
 
-					Serilog.Log.Information("LOG: Read {NumBytes}", ReadBytes);
+					Logger.LogInformation("LOG: Read {NumBytes}", ReadBytes);
 
 					// Copy as many lines as possible to the output
 					int ConvertedBytes = 0;
@@ -239,9 +240,9 @@ namespace HordeServer.Services
 					{
 						if (ReadBuffer[EndIdx] == '\n')
 						{
-							Serilog.Log.Information("LOG: Parsing {Start} -> {End}", ConvertedBytes, EndIdx);
+							Logger.LogInformation("LOG: Parsing {Start} -> {End}", ConvertedBytes, EndIdx);
 							WriteBufferLength = LogText.ConvertToPlainText(ReadBuffer.AsSpan(ConvertedBytes, EndIdx - ConvertedBytes), WriteBuffer, WriteBufferLength);
-							Serilog.Log.Information("LOG: Done");
+							Logger.LogInformation("LOG: Done");
 							ConvertedBytes = EndIdx + 1;
 						}
 					}
@@ -252,9 +253,9 @@ namespace HordeServer.Services
 						if (Offset < WriteBufferLength)
 						{
 							int WriteLength = (int)Math.Min((long)WriteBufferLength - Offset, Length);
-							Serilog.Log.Information("LOG: Writing {Start} -> {End}", Offset, WriteLength);
+							Logger.LogInformation("LOG: Writing {Start} -> {End}", Offset, WriteLength);
 							await OutputStream.WriteAsync(WriteBuffer.AsMemory((int)Offset, WriteLength));
-							Serilog.Log.Information("LOG: Done");
+							Logger.LogInformation("LOG: Done");
 							Length -= WriteLength;
 						}
 						Offset = Math.Max(Offset - WriteBufferLength, 0);
@@ -266,19 +267,19 @@ namespace HordeServer.Services
 					{
 						Buffer.BlockCopy(ReadBuffer, ConvertedBytes, ReadBuffer, 0, ReadBufferLength - ConvertedBytes);
 						ReadBufferLength -= ConvertedBytes;
-						Serilog.Log.Information("LOG: Removing {NumBytes}", ConvertedBytes);
+						Logger.LogInformation("LOG: Removing {NumBytes}", ConvertedBytes);
 					}
 					else if(ReadBufferLength > 0)
 					{
 						Array.Resize(ref ReadBuffer, ReadBuffer.Length + 128);
 						WriteBuffer = new byte[ReadBuffer.Length];
-						Serilog.Log.Information("LOG: Expanding to {Length}", ReadBuffer.Length);
+						Logger.LogInformation("LOG: Expanding to {Length}", ReadBuffer.Length);
 					}
 
 					// Exit if we didn't read anything in this iteration
 					if (ReadBytes == 0)
 					{
-						Serilog.Log.Information("LOG: COMPLETE");
+						Logger.LogInformation("LOG: COMPLETE");
 						break;
 					}
 				}
