Horde: More logging for log download endpoint.

diff --git a/Engine/Source/Programs/Horde/Horde.Build/Controllers/LogsController.cs b/Engine/Source/Programs/Horde/Horde.Build/Controllers/LogsController.cs
--- a/Engine/Source/Programs/Horde/Horde.Build/Controllers/LogsController.cs
+++ b/Engine/Source/Programs/Horde/Horde.Build/Controllers/LogsController.cs
@@ -109,6 +109,8 @@ namespace HordeServer.Controllers
 		[Route("/api/v1/logs/{LogFileId}/data")]
 		public async Task<ActionResult> GetLogData(LogId LogFileId, [FromQuery] LogOutputFormat Format = LogOutputFormat.Raw, [FromQuery] long Offset = 0, [FromQuery] long Length = long.MaxValue, [FromQuery] string? FileName = null, [FromQuery] bool Download = false)
 		{
+			Logger.LogInformation("Requesting log {LogId}", LogFileId);
+
 			ILogFile? LogFile = await LogFileService.GetLogFileAsync(LogFileId);
 			if (LogFile == null)
 			{
@@ -122,10 +124,12 @@ namespace HordeServer.Controllers
 			Func<Stream, ActionContext, Task> CopyTask;
 			if (Format == LogOutputFormat.Text && LogFile.Type == LogType.Json)
 			{
+				Logger.LogInformation("Creating plain text stream for {LogId}", LogFileId);
 				CopyTask = (OutputStream, Context) => LogFileService.CopyPlainTextStreamAsync(LogFile, Offset, Length, OutputStream, Logger);
 			}
 			else
 			{
+				Logger.LogInformation("Creating raw stream for {LogId}", LogFileId);
 				CopyTask = (OutputStream, Context) => LogFileService.CopyRawStreamAsync(LogFile, Offset, Length, OutputStream);
 			}
 
