Horde: More logging for log streaming.

diff --git a/Engine/Source/Programs/Horde/Horde.Build/Logs/Impl/LogFileService.cs b/Engine/Source/Programs/Horde/Horde.Build/Logs/Impl/LogFileService.cs
--- a/Engine/Source/Programs/Horde/Horde.Build/Logs/Impl/LogFileService.cs
+++ b/Engine/Source/Programs/Horde/Horde.Build/Logs/Impl/LogFileService.cs
@@ -218,6 +218,7 @@ namespace HordeServer.Services
 		/// <returns>Async text</returns>
 		public static async Task CopyPlainTextStreamAsync(this ILogFileService LogFileService, ILogFile LogFile, long Offset, long Length, Stream OutputStream, ILogger Logger)
 		{
+			Logger.LogInformation("LOG: Opening raw stream for log {LogId}", LogFile.Id);
 			using (Stream Stream = await LogFileService.OpenRawStreamAsync(LogFile, 0, long.MaxValue))
 			{
 				byte[] ReadBuffer = new byte[4096];
@@ -228,6 +229,8 @@ namespace HordeServer.Services
 
 				while(Length > 0)
 				{
+					Logger.LogInformation("LOG: Reading up to {MaxRead} bytes", ReadBuffer.Length - ReadBufferLength);
+
 					// Add more data to the buffer
 					int ReadBytes = await Stream.ReadAsync(ReadBuffer.AsMemory(ReadBufferLength, ReadBuffer.Length - ReadBufferLength));
 					ReadBufferLength += ReadBytes;
@@ -240,13 +243,13 @@ namespace HordeServer.Services
 					{
 						if (ReadBuffer[EndIdx] == '\n')
 						{
-							Logger.LogInformation("LOG: Parsing {Start} -> {End}", ConvertedBytes, EndIdx);
 							WriteBufferLength = LogText.ConvertToPlainText(ReadBuffer.AsSpan(ConvertedBytes, EndIdx - ConvertedBytes), WriteBuffer, WriteBufferLength);
-							Logger.LogInformation("LOG: Done");
 							ConvertedBytes = EndIdx + 1;
 						}
 					}
 
+					Logger.LogInformation("LOG: Converted {NumBytes}", ConvertedBytes);
+
 					// If there's anything in the write buffer, write it out
 					if (WriteBufferLength > 0)
 					{
