Horde: Additional logging for log parsing.

diff --git a/Engine/Source/Programs/Horde/Horde.Build/Controllers/LogsController.cs b/Engine/Source/Programs/Horde/Horde.Build/Controllers/LogsController.cs
--- a/Engine/Source/Programs/Horde/Horde.Build/Controllers/LogsController.cs
+++ b/Engine/Source/Programs/Horde/Horde.Build/Controllers/LogsController.cs
@@ -21,6 +21,7 @@ using Microsoft.AspNetCore.Authorization;
 using Microsoft.AspNetCore.Http;
 using Microsoft.AspNetCore.Mvc;
 using Microsoft.Extensions.Logging;
+using Microsoft.Extensions.Logging.Abstractions;
 using MongoDB.Bson;
 
 namespace HordeServer.Controllers
@@ -165,7 +166,7 @@ namespace HordeServer.Controllers
 			Count = MaxIndex - MinIndex;
 
 			byte[] Result;
-			using (System.IO.Stream Stream = await LogFileService.OpenRawStreamAsync(LogFile, MinOffset, MaxOffset - MinOffset))
+			using (System.IO.Stream Stream = await LogFileService.OpenRawStreamAsync(LogFile, MinOffset, MaxOffset - MinOffset, NullLogger.Instance))
 			{
 				Result = new byte[Stream.Length];
 				await Stream.ReadFixedSizeDataAsync(Result, 0, Result.Length);
diff --git a/Engine/Source/Programs/Horde/Horde.Build/Logs/Impl/LogFileService.cs b/Engine/Source/Programs/Horde/Horde.Build/Logs/Impl/LogFileService.cs
--- a/Engine/Source/Programs/Horde/Horde.Build/Logs/Impl/LogFileService.cs
+++ b/Engine/Source/Programs/Horde/Horde.Build/Logs/Impl/LogFileService.cs
@@ -33,6 +33,7 @@ using Stream = System.IO.Stream;
 using OpenTracing.Util;
 using OpenTracing;
 using Microsoft.Extensions.Hosting;
+using Microsoft.Extensions.Logging.Abstractions;
 
 namespace HordeServer.Services
 {
@@ -161,8 +162,9 @@ namespace HordeServer.Services
 		/// <param name="LogFile">The log file</param>
 		/// <param name="Offset">Offset of the data to return</param>
 		/// <param name="Length">Length of the data to return</param>
+		/// <param name="Logger"></param>
 		/// <returns>Data for the requested range</returns>
-		Task<Stream> OpenRawStreamAsync(ILogFile LogFile, long Offset, long Length);
+		Task<Stream> OpenRawStreamAsync(ILogFile LogFile, long Offset, long Length, ILogger Logger);
 
 		/// <summary>
 		/// Gets the offset of the given line number
@@ -200,7 +202,7 @@ namespace HordeServer.Services
 		/// <returns>Async text</returns>
 		public static async Task CopyRawStreamAsync(this ILogFileService LogFileService, ILogFile LogFile, long Offset, long Length, Stream OutputStream)
 		{
-			using (Stream Stream = await LogFileService.OpenRawStreamAsync(LogFile, Offset, Length))
+			using (Stream Stream = await LogFileService.OpenRawStreamAsync(LogFile, Offset, Length, NullLogger.Instance))
 			{
 				await Stream.CopyToAsync(OutputStream);
 			}
@@ -218,75 +220,82 @@ namespace HordeServer.Services
 		/// <returns>Async text</returns>
 		public static async Task CopyPlainTextStreamAsync(this ILogFileService LogFileService, ILogFile LogFile, long Offset, long Length, Stream OutputStream, ILogger Logger)
 		{
-			Logger.LogInformation("LOG: Opening raw stream for log {LogId}", LogFile.Id);
-			using (Stream Stream = await LogFileService.OpenRawStreamAsync(LogFile, 0, long.MaxValue))
+			try
 			{
-				byte[] ReadBuffer = new byte[4096];
-				int ReadBufferLength = 0;
+				Logger.LogInformation("LOG: Opening raw stream for log {LogId}", LogFile.Id);
+				using (Stream Stream = await LogFileService.OpenRawStreamAsync(LogFile, 0, long.MaxValue, Logger))
+				{
+					byte[] ReadBuffer = new byte[4096];
+					int ReadBufferLength = 0;
 
-				byte[] WriteBuffer = new byte[4096];
-				int WriteBufferLength = 0;
+					byte[] WriteBuffer = new byte[4096];
+					int WriteBufferLength = 0;
 
-				while(Length > 0)
-				{
-					Logger.LogInformation("LOG: Reading up to {MaxRead} bytes", ReadBuffer.Length - ReadBufferLength);
+					while (Length > 0)
+					{
+						Logger.LogInformation("LOG: Reading up to {MaxRead} bytes", ReadBuffer.Length - ReadBufferLength);
 
-					// Add more data to the buffer
-					int ReadBytes = await Stream.ReadAsync(ReadBuffer.AsMemory(ReadBufferLength, ReadBuffer.Length - ReadBufferLength));
-					ReadBufferLength += ReadBytes;
+						// Add more data to the buffer
+						int ReadBytes = await Stream.ReadAsync(ReadBuffer.AsMemory(ReadBufferLength, ReadBuffer.Length - ReadBufferLength));
+						ReadBufferLength += ReadBytes;
 
-					Logger.LogInformation("LOG: Read {NumBytes}", ReadBytes);
+						Logger.LogInformation("LOG: Read {NumBytes}", ReadBytes);
 
-					// Copy as many lines as possible to the output
-					int ConvertedBytes = 0;
-					for(int EndIdx = 1; EndIdx < ReadBufferLength; EndIdx++)
-					{
-						if (ReadBuffer[EndIdx] == '\n')
+						// Copy as many lines as possible to the output
+						int ConvertedBytes = 0;
+						for (int EndIdx = 1; EndIdx < ReadBufferLength; EndIdx++)
 						{
-							WriteBufferLength = LogText.ConvertToPlainText(ReadBuffer.AsSpan(ConvertedBytes, EndIdx - ConvertedBytes), WriteBuffer, WriteBufferLength);
-							ConvertedBytes = EndIdx + 1;
+							if (ReadBuffer[EndIdx] == '\n')
+							{
+								WriteBufferLength = LogText.ConvertToPlainText(ReadBuffer.AsSpan(ConvertedBytes, EndIdx - ConvertedBytes), WriteBuffer, WriteBufferLength);
+								ConvertedBytes = EndIdx + 1;
+							}
 						}
-					}
 
-					Logger.LogInformation("LOG: Converted {NumBytes}", ConvertedBytes);
+						Logger.LogInformation("LOG: Converted {NumBytes}", ConvertedBytes);
 
-					// If there's anything in the write buffer, write it out
-					if (WriteBufferLength > 0)
-					{
-						if (Offset < WriteBufferLength)
+						// If there's anything in the write buffer, write it out
+						if (WriteBufferLength > 0)
 						{
-							int WriteLength = (int)Math.Min((long)WriteBufferLength - Offset, Length);
-							Logger.LogInformation("LOG: Writing {Start} -> {End}", Offset, WriteLength);
-							await OutputStream.WriteAsync(WriteBuffer.AsMemory((int)Offset, WriteLength));
-							Logger.LogInformation("LOG: Done");
-							Length -= WriteLength;
+							if (Offset < WriteBufferLength)
+							{
+								int WriteLength = (int)Math.Min((long)WriteBufferLength - Offset, Length);
+								Logger.LogInformation("LOG: Writing {Start} -> {End}", Offset, WriteLength);
+								await OutputStream.WriteAsync(WriteBuffer.AsMemory((int)Offset, WriteLength));
+								Logger.LogInformation("LOG: Done");
+								Length -= WriteLength;
+							}
+							Offset = Math.Max(Offset - WriteBufferLength, 0);
+							WriteBufferLength = 0;
 						}
-						Offset = Math.Max(Offset - WriteBufferLength, 0);
-						WriteBufferLength = 0;
-					}
 
-					// If we were able to read something, shuffle down the rest of the buffer. Otherwise expand the read buffer.
-					if (ConvertedBytes > 0)
-					{
-						Buffer.BlockCopy(ReadBuffer, ConvertedBytes, ReadBuffer, 0, ReadBufferLength - ConvertedBytes);
-						ReadBufferLength -= ConvertedBytes;
-						Logger.LogInformation("LOG: Removing {NumBytes}", ConvertedBytes);
-					}
-					else if(ReadBufferLength > 0)
-					{
-						Array.Resize(ref ReadBuffer, ReadBuffer.Length + 128);
-						WriteBuffer = new byte[ReadBuffer.Length];
-						Logger.LogInformation("LOG: Expanding to {Length}", ReadBuffer.Length);
-					}
+						// If we were able to read something, shuffle down the rest of the buffer. Otherwise expand the read buffer.
+						if (ConvertedBytes > 0)
+						{
+							Buffer.BlockCopy(ReadBuffer, ConvertedBytes, ReadBuffer, 0, ReadBufferLength - ConvertedBytes);
+							ReadBufferLength -= ConvertedBytes;
+							Logger.LogInformation("LOG: Removing {NumBytes}", ConvertedBytes);
+						}
+						else if (ReadBufferLength > 0)
+						{
+							Array.Resize(ref ReadBuffer, ReadBuffer.Length + 128);
+							WriteBuffer = new byte[ReadBuffer.Length];
+							Logger.LogInformation("LOG: Expanding to {Length}", ReadBuffer.Length);
+						}
 
-					// Exit if we didn't read anything in this iteration
-					if (ReadBytes == 0)
-					{
-						Logger.LogInformation("LOG: COMPLETE");
-						break;
+						// Exit if we didn't read anything in this iteration
+						if (ReadBytes == 0)
+						{
+							Logger.LogInformation("LOG: COMPLETE");
+							break;
+						}
 					}
 				}
 			}
+			catch (Exception Ex)
+			{
+				Logger.LogError(Ex, "Excception while reading log contents: {Ex}", Ex);
+			}
 		}
 	}
 	
@@ -390,6 +399,8 @@ namespace HordeServer.Services
 			/// </summary>
 			int SourceEnd;
 
+			ILogger Logger;
+
 			/// <summary>
 			/// Constructor
 			/// </summary>
@@ -397,7 +408,8 @@ namespace HordeServer.Services
 			/// <param name="LogFile"></param>
 			/// <param name="Offset"></param>
 			/// <param name="Length"></param>
-			public ResponseStream(LogFileService LogFileService, ILogFile LogFile, long Offset, long Length)
+			/// <param name="Logger"></param>
+			public ResponseStream(LogFileService LogFileService, ILogFile LogFile, long Offset, long Length, ILogger Logger)
 			{
 				this.LogFileService = LogFileService;
 				this.LogFile = LogFile;
@@ -409,6 +421,8 @@ namespace HordeServer.Services
 
 				this.ChunkIdx = LogFile.Chunks.GetChunkForOffset(Offset);
 				this.SourceBuffer = null!;
+
+				this.Logger = Logger;
 			}
 
 			/// <inheritdoc/>
@@ -473,17 +487,11 @@ namespace HordeServer.Services
 							ChunkIdx++;
 						}
 
-						// Get the end of this chunk
-						long NextOffset = ResponseOffset + ResponseLength;
-						if (ChunkIdx + 1 < LogFile.Chunks.Count)
-						{
-							ILogChunk NextChunk = LogFile.Chunks[ChunkIdx + 1];
-							NextOffset = Math.Min(NextOffset, NextChunk.Offset);
-						}
-
 						// Get the chunk data
 						ILogChunk Chunk = LogFile.Chunks[ChunkIdx];
+						Logger.LogInformation("Reading chunk {ChunkIdx}", ChunkIdx);
 						LogChunkData ChunkData = await LogFileService.ReadChunkAsync(LogFile, ChunkIdx);
+						Logger.LogInformation("Read chunk {ChunkIdx}: {NumBytes} bytes", ChunkIdx, ChunkData.Length);
 
 						// Figure out which sub-chunk to use
 						int SubChunkIdx = ChunkData.GetSubChunkForOffsetWithinChunk((int)(CurrentOffset - Chunk.Offset));
@@ -930,7 +938,7 @@ namespace HordeServer.Services
 			(_, long MaxOffset) = await GetLineOffsetAsync(LogFile, LineIndex + LineCount);
 
 			byte[] Data = new byte[MaxOffset - MinOffset];
-			using (Stream Stream = await OpenRawStreamAsync(LogFile, MinOffset, MaxOffset - MinOffset))
+			using (Stream Stream = await OpenRawStreamAsync(LogFile, MinOffset, MaxOffset - MinOffset, NullLogger.Instance))
 			{
 				int Length = await Stream.ReadAsync(Data.AsMemory());
 				if(Length != Data.Length)
@@ -975,7 +983,7 @@ namespace HordeServer.Services
 		}
 
 		/// <inheritdoc/>
-		public async Task<Stream> OpenRawStreamAsync(ILogFile LogFile, long Offset, long Length)
+		public async Task<Stream> OpenRawStreamAsync(ILogFile LogFile, long Offset, long Length, ILogger Logger)
 		{
 			if (LogFile.Chunks.Count == 0)
 			{
@@ -999,7 +1007,7 @@ namespace HordeServer.Services
 				}
 
 				// Create the new stream
-				return new ResponseStream(this, LogFile, Offset, Length);
+				return new ResponseStream(this, LogFile, Offset, Length, Logger);
 			}
 		}
 
