Horde: More logging hacks.

diff --git a/Engine/Source/Programs/Horde/Horde.Build.Tests/LogFileServiceTests.cs b/Engine/Source/Programs/Horde/Horde.Build.Tests/LogFileServiceTests.cs
--- a/Engine/Source/Programs/Horde/Horde.Build.Tests/LogFileServiceTests.cs
+++ b/Engine/Source/Programs/Horde/Horde.Build.Tests/LogFileServiceTests.cs
@@ -116,7 +116,7 @@ namespace HordeServerTests
         
         protected static async Task<string> ReadLogFile(ILogFileService LogFileService, ILogFile LogFile, long Offset, long Length)
         {
-	        Stream Stream = await LogFileService.OpenRawStreamAsync(LogFile, Offset, Length);
+	        Stream Stream = await LogFileService.OpenRawStreamAsync(LogFile, Offset, Length, NullLogger.Instance);
 	        return new StreamReader(Stream).ReadToEnd();
         }
         
diff --git a/Engine/Source/Programs/Horde/Horde.Build.Tests/LogIndexingTests.cs b/Engine/Source/Programs/Horde/Horde.Build.Tests/LogIndexingTests.cs
--- a/Engine/Source/Programs/Horde/Horde.Build.Tests/LogIndexingTests.cs
+++ b/Engine/Source/Programs/Horde/Horde.Build.Tests/LogIndexingTests.cs
@@ -23,6 +23,7 @@ using System.Text;
 using System.Threading.Tasks;
 using HordeServer.Collections.Impl;
 using HordeCommon;
+using Microsoft.Extensions.Logging.Abstractions;
 
 namespace HordeServerTests
 {
@@ -85,7 +86,7 @@ namespace HordeServerTests
 
 			// Read the data back out and check it's the same
 			byte[] ReadData = new byte[Data.Length];
-			using (Stream Stream = await LogFileService.OpenRawStreamAsync(LogFile, 0, Data.Length))
+			using (Stream Stream = await LogFileService.OpenRawStreamAsync(LogFile, 0, Data.Length, NullLogger.Instance))
 			{
 				int ReadSize = await Stream.ReadAsync(ReadData, 0, ReadData.Length);
 				Assert.AreEqual(ReadData.Length, ReadSize);
