Horde: Speculative fix for infinite loop when reading log files.

diff --git a/Engine/Source/Programs/Horde/Horde.Build/Logs/Impl/LogFileService.cs b/Engine/Source/Programs/Horde/Horde.Build/Logs/Impl/LogFileService.cs
--- a/Engine/Source/Programs/Horde/Horde.Build/Logs/Impl/LogFileService.cs
+++ b/Engine/Source/Programs/Horde/Horde.Build/Logs/Impl/LogFileService.cs
@@ -468,7 +468,7 @@ namespace HordeServer.Services
 			public override async ValueTask<int> ReadAsync(Memory<byte> Buffer, CancellationToken CancellationToken)
 			{
 				int ReadBytes = 0;
-				while (ReadBytes < Length)
+				while (ReadBytes < Buffer.Length)
 				{
 					if (SourcePos < SourceEnd)
 					{
