UGS: Remove all debug code added for logging.

diff --git a/Engine/Source/Programs/Horde/Horde.Build.Tests/LogFileServiceTests.cs b/Engine/Source/Programs/Horde/Horde.Build.Tests/LogFileServiceTests.cs
--- a/Engine/Source/Programs/Horde/Horde.Build.Tests/LogFileServiceTests.cs
+++ b/Engine/Source/Programs/Horde/Horde.Build.Tests/LogFileServiceTests.cs
@@ -116,7 +116,7 @@ namespace HordeServerTests
         
         protected static async Task<string> ReadLogFile(ILogFileService LogFileService, ILogFile LogFile, long Offset, long Length)
         {
-	        Stream Stream = await LogFileService.OpenRawStreamAsync(LogFile, Offset, Length, NullLogger.Instance);
+	        Stream Stream = await LogFileService.OpenRawStreamAsync(LogFile, Offset, Length);
 	        return new StreamReader(Stream).ReadToEnd();
         }
         
diff --git a/Engine/Source/Programs/Horde/Horde.Build.Tests/LogIndexingTests.cs b/Engine/Source/Programs/Horde/Horde.Build.Tests/LogIndexingTests.cs
--- a/Engine/Source/Programs/Horde/Horde.Build.Tests/LogIndexingTests.cs
+++ b/Engine/Source/Programs/Horde/Horde.Build.Tests/LogIndexingTests.cs
@@ -86,7 +86,7 @@ namespace HordeServerTests
 
 			// Read the data back out and check it's the same
 			byte[] ReadData = new byte[Data.Length];
-			using (Stream Stream = await LogFileService.OpenRawStreamAsync(LogFile, 0, Data.Length, NullLogger.Instance))
+			using (Stream Stream = await LogFileService.OpenRawStreamAsync(LogFile, 0, Data.Length))
 			{
 				int ReadSize = await Stream.ReadAsync(ReadData, 0, ReadData.Length);
 				Assert.AreEqual(ReadData.Length, ReadSize);
diff --git a/Engine/Source/Programs/Horde/Horde.Build/Controllers/LogsController.cs b/Engine/Source/Programs/Horde/Horde.Build/Controllers/LogsController.cs
--- a/Engine/Source/Programs/Horde/Horde.Build/Controllers/LogsController.cs
+++ b/Engine/Source/Programs/Horde/Horde.Build/Controllers/LogsController.cs
@@ -57,18 +57,16 @@ namespace HordeServer.Controllers
 		private readonly IIssueCollection IssueCollection;
 		private readonly AclService AclService;
 		private readonly JobService JobService;
-		private readonly ILogger Logger;
 
 		/// <summary>
 		/// Constructor
 		/// </summary>
-		public LogsController(ILogFileService LogFileService, IIssueCollection IssueCollection, AclService AclService, JobService JobService, ILogger<LogsController> Logger)
+		public LogsController(ILogFileService LogFileService, IIssueCollection IssueCollection, AclService AclService, JobService JobService)
 		{
 			this.LogFileService = LogFileService;
 			this.IssueCollection = IssueCollection;
 			this.AclService = AclService;
 			this.JobService = JobService;
-			this.Logger = Logger;
 		}
 
 		/// <summary>
@@ -110,8 +108,6 @@ namespace HordeServer.Controllers
 		[Route("/api/v1/logs/{LogFileId}/data")]
 		public async Task<ActionResult> GetLogData(LogId LogFileId, [FromQuery] LogOutputFormat Format = LogOutputFormat.Raw, [FromQuery] long Offset = 0, [FromQuery] long Length = long.MaxValue, [FromQuery] string? FileName = null, [FromQuery] bool Download = false)
 		{
-			Logger.LogInformation("Requesting log {LogId}", LogFileId);
-
 			ILogFile? LogFile = await LogFileService.GetLogFileAsync(LogFileId);
 			if (LogFile == null)
 			{
@@ -125,12 +121,10 @@ namespace HordeServer.Controllers
 			Func<Stream, ActionContext, Task> CopyTask;
 			if (Format == LogOutputFormat.Text && LogFile.Type == LogType.Json)
 			{
-				Logger.LogInformation("Creating plain text stream for {LogId}", LogFileId);
-				CopyTask = (OutputStream, Context) => LogFileService.CopyPlainTextStreamAsync(LogFile, Offset, Length, OutputStream, Logger);
+				CopyTask = (OutputStream, Context) => LogFileService.CopyPlainTextStreamAsync(LogFile, Offset, Length, OutputStream);
 			}
 			else
 			{
-				Logger.LogInformation("Creating raw stream for {LogId}", LogFileId);
 				CopyTask = (OutputStream, Context) => LogFileService.CopyRawStreamAsync(LogFile, Offset, Length, OutputStream);
 			}
 
@@ -166,7 +160,7 @@ namespace HordeServer.Controllers
 			Count = MaxIndex - MinIndex;
 
 			byte[] Result;
-			using (System.IO.Stream Stream = await LogFileService.OpenRawStreamAsync(LogFile, MinOffset, MaxOffset - MinOffset, NullLogger.Instance))
+			using (System.IO.Stream Stream = await LogFileService.OpenRawStreamAsync(LogFile, MinOffset, MaxOffset - MinOffset))
 			{
 				Result = new byte[Stream.Length];
 				await Stream.ReadFixedSizeDataAsync(Result, 0, Result.Length);
diff --git a/Engine/Source/Programs/Horde/Horde.Build/Logs/Impl/LogFileService.cs b/Engine/Source/Programs/Horde/Horde.Build/Logs/Impl/LogFileService.cs
--- a/Engine/Source/Programs/Horde/Horde.Build/Logs/Impl/LogFileService.cs
+++ b/Engine/Source/Programs/Horde/Horde.Build/Logs/Impl/LogFileService.cs
@@ -162,9 +162,8 @@ namespace HordeServer.Services
 		/// <param name="LogFile">The log file</param>
 		/// <param name="Offset">Offset of the data to return</param>
 		/// <param name="Length">Length of the data to return</param>
-		/// <param name="Logger"></param>
 		/// <returns>Data for the requested range</returns>
-		Task<Stream> OpenRawStreamAsync(ILogFile LogFile, long Offset, long Length, ILogger Logger);
+		Task<Stream> OpenRawStreamAsync(ILogFile LogFile, long Offset, long Length);
 
 		/// <summary>
 		/// Gets the offset of the given line number
@@ -202,7 +201,7 @@ namespace HordeServer.Services
 		/// <returns>Async text</returns>
 		public static async Task CopyRawStreamAsync(this ILogFileService LogFileService, ILogFile LogFile, long Offset, long Length, Stream OutputStream)
 		{
-			using (Stream Stream = await LogFileService.OpenRawStreamAsync(LogFile, Offset, Length, NullLogger.Instance))
+			using (Stream Stream = await LogFileService.OpenRawStreamAsync(LogFile, Offset, Length))
 			{
 				await Stream.CopyToAsync(OutputStream);
 			}
@@ -216,86 +215,66 @@ namespace HordeServer.Services
 		/// <param name="Offset">Offset within the data to copy from</param>
 		/// <param name="Length">Length of the data to copy</param>
 		/// <param name="OutputStream">Output stream to receive the text data</param>
-		/// <param name="Logger">Logger for output messages</param>
 		/// <returns>Async text</returns>
-		public static async Task CopyPlainTextStreamAsync(this ILogFileService LogFileService, ILogFile LogFile, long Offset, long Length, Stream OutputStream, ILogger Logger)
+		public static async Task CopyPlainTextStreamAsync(this ILogFileService LogFileService, ILogFile LogFile, long Offset, long Length, Stream OutputStream)
 		{
-			try
+			using (Stream Stream = await LogFileService.OpenRawStreamAsync(LogFile, 0, long.MaxValue))
 			{
-				Logger.LogInformation("LOG: Opening raw stream for log {LogId}", LogFile.Id);
-				using (Stream Stream = await LogFileService.OpenRawStreamAsync(LogFile, 0, long.MaxValue, Logger))
-				{
-					byte[] ReadBuffer = new byte[4096];
-					int ReadBufferLength = 0;
-
-					byte[] WriteBuffer = new byte[4096];
-					int WriteBufferLength = 0;
-
-					while (Length > 0)
-					{
-						Logger.LogInformation("LOG: Reading up to {MaxRead} bytes", ReadBuffer.Length - ReadBufferLength);
+				byte[] ReadBuffer = new byte[4096];
+				int ReadBufferLength = 0;
 
-						// Add more data to the buffer
-						int ReadBytes = await Stream.ReadAsync(ReadBuffer.AsMemory(ReadBufferLength, ReadBuffer.Length - ReadBufferLength));
-						ReadBufferLength += ReadBytes;
+				byte[] WriteBuffer = new byte[4096];
+				int WriteBufferLength = 0;
 
-						Logger.LogInformation("LOG: Read {NumBytes}", ReadBytes);
+				while (Length > 0)
+				{
+					// Add more data to the buffer
+					int ReadBytes = await Stream.ReadAsync(ReadBuffer.AsMemory(ReadBufferLength, ReadBuffer.Length - ReadBufferLength));
+					ReadBufferLength += ReadBytes;
 
-						// Copy as many lines as possible to the output
-						int ConvertedBytes = 0;
-						for (int EndIdx = 1; EndIdx < ReadBufferLength; EndIdx++)
+					// Copy as many lines as possible to the output
+					int ConvertedBytes = 0;
+					for (int EndIdx = 1; EndIdx < ReadBufferLength; EndIdx++)
+					{
+						if (ReadBuffer[EndIdx] == '\n')
 						{
-							if (ReadBuffer[EndIdx] == '\n')
-							{
-								WriteBufferLength = LogText.ConvertToPlainText(ReadBuffer.AsSpan(ConvertedBytes, EndIdx - ConvertedBytes), WriteBuffer, WriteBufferLength);
-								ConvertedBytes = EndIdx + 1;
-							}
+							WriteBufferLength = LogText.ConvertToPlainText(ReadBuffer.AsSpan(ConvertedBytes, EndIdx - ConvertedBytes), WriteBuffer, WriteBufferLength);
+							ConvertedBytes = EndIdx + 1;
 						}
+					}
 
-						Logger.LogInformation("LOG: Converted {NumBytes}", ConvertedBytes);
-
-						// If there's anything in the write buffer, write it out
-						if (WriteBufferLength > 0)
+					// If there's anything in the write buffer, write it out
+					if (WriteBufferLength > 0)
+					{
+						if (Offset < WriteBufferLength)
 						{
-							if (Offset < WriteBufferLength)
-							{
-								int WriteLength = (int)Math.Min((long)WriteBufferLength - Offset, Length);
-								Logger.LogInformation("LOG: Writing {Start} -> {End}", Offset, WriteLength);
-								await OutputStream.WriteAsync(WriteBuffer.AsMemory((int)Offset, WriteLength));
-								Logger.LogInformation("LOG: Done");
-								Length -= WriteLength;
-							}
-							Offset = Math.Max(Offset - WriteBufferLength, 0);
-							WriteBufferLength = 0;
+							int WriteLength = (int)Math.Min((long)WriteBufferLength - Offset, Length);
+							await OutputStream.WriteAsync(WriteBuffer.AsMemory((int)Offset, WriteLength));
+							Length -= WriteLength;
 						}
+						Offset = Math.Max(Offset - WriteBufferLength, 0);
+						WriteBufferLength = 0;
+					}
 
-						// If we were able to read something, shuffle down the rest of the buffer. Otherwise expand the read buffer.
-						if (ConvertedBytes > 0)
-						{
-							Buffer.BlockCopy(ReadBuffer, ConvertedBytes, ReadBuffer, 0, ReadBufferLength - ConvertedBytes);
-							ReadBufferLength -= ConvertedBytes;
-							Logger.LogInformation("LOG: Removing {NumBytes}", ConvertedBytes);
-						}
-						else if (ReadBufferLength > 0)
-						{
-							Array.Resize(ref ReadBuffer, ReadBuffer.Length + 128);
-							WriteBuffer = new byte[ReadBuffer.Length];
-							Logger.LogInformation("LOG: Expanding to {Length}", ReadBuffer.Length);
-						}
+					// If we were able to read something, shuffle down the rest of the buffer. Otherwise expand the read buffer.
+					if (ConvertedBytes > 0)
+					{
+						Buffer.BlockCopy(ReadBuffer, ConvertedBytes, ReadBuffer, 0, ReadBufferLength - ConvertedBytes);
+						ReadBufferLength -= ConvertedBytes;
+					}
+					else if (ReadBufferLength > 0)
+					{
+						Array.Resize(ref ReadBuffer, ReadBuffer.Length + 128);
+						WriteBuffer = new byte[ReadBuffer.Length];
+					}
 
-						// Exit if we didn't read anything in this iteration
-						if (ReadBytes == 0)
-						{
-							Logger.LogInformation("LOG: COMPLETE");
-							break;
-						}
+					// Exit if we didn't read anything in this iteration
+					if (ReadBytes == 0)
+					{
+						break;
 					}
 				}
 			}
-			catch (Exception Ex)
-			{
-				Logger.LogError(Ex, "Excception while reading log contents: {Ex}", Ex);
-			}
 		}
 	}
 	
@@ -399,8 +378,6 @@ namespace HordeServer.Services
 			/// </summary>
 			int SourceEnd;
 
-			ILogger Logger;
-
 			/// <summary>
 			/// Constructor
 			/// </summary>
@@ -408,8 +385,7 @@ namespace HordeServer.Services
 			/// <param name="LogFile"></param>
 			/// <param name="Offset"></param>
 			/// <param name="Length"></param>
-			/// <param name="Logger"></param>
-			public ResponseStream(LogFileService LogFileService, ILogFile LogFile, long Offset, long Length, ILogger Logger)
+			public ResponseStream(LogFileService LogFileService, ILogFile LogFile, long Offset, long Length)
 			{
 				this.LogFileService = LogFileService;
 				this.LogFile = LogFile;
@@ -421,8 +397,6 @@ namespace HordeServer.Services
 
 				this.ChunkIdx = LogFile.Chunks.GetChunkForOffset(Offset);
 				this.SourceBuffer = null!;
-
-				this.Logger = Logger;
 			}
 
 			/// <inheritdoc/>
@@ -489,9 +463,7 @@ namespace HordeServer.Services
 
 						// Get the chunk data
 						ILogChunk Chunk = LogFile.Chunks[ChunkIdx];
-						Logger.LogInformation("Reading chunk {ChunkIdx}", ChunkIdx);
 						LogChunkData ChunkData = await LogFileService.ReadChunkAsync(LogFile, ChunkIdx);
-						Logger.LogInformation("Read chunk {ChunkIdx}: {NumBytes} bytes", ChunkIdx, ChunkData.Length);
 
 						// Figure out which sub-chunk to use
 						int SubChunkIdx = ChunkData.GetSubChunkForOffsetWithinChunk((int)(CurrentOffset - Chunk.Offset));
@@ -938,7 +910,7 @@ namespace HordeServer.Services
 			(_, long MaxOffset) = await GetLineOffsetAsync(LogFile, LineIndex + LineCount);
 
 			byte[] Data = new byte[MaxOffset - MinOffset];
-			using (Stream Stream = await OpenRawStreamAsync(LogFile, MinOffset, MaxOffset - MinOffset, NullLogger.Instance))
+			using (Stream Stream = await OpenRawStreamAsync(LogFile, MinOffset, MaxOffset - MinOffset))
 			{
 				int Length = await Stream.ReadAsync(Data.AsMemory());
 				if(Length != Data.Length)
@@ -983,7 +955,7 @@ namespace HordeServer.Services
 		}
 
 		/// <inheritdoc/>
-		public async Task<Stream> OpenRawStreamAsync(ILogFile LogFile, long Offset, long Length, ILogger Logger)
+		public async Task<Stream> OpenRawStreamAsync(ILogFile LogFile, long Offset, long Length)
 		{
 			if (LogFile.Chunks.Count == 0)
 			{
@@ -1007,7 +979,7 @@ namespace HordeServer.Services
 				}
 
 				// Create the new stream
-				return new ResponseStream(this, LogFile, Offset, Length, Logger);
+				return new ResponseStream(this, LogFile, Offset, Length);
 			}
 		}
 
