Horde Storage - Fixed failing deserializing when blob is null (e.g. during a delete in the replication log)

diff --git a/Engine/Source/Programs/Horde/HordeStorage/Horde.Storage/Tests/UnitTests/BlobIdentifierTests.cs b/Engine/Source/Programs/Horde/HordeStorage/Horde.Storage/Tests/UnitTests/BlobIdentifierTests.cs
new file mode 100644
--- /dev/null
+++ b/Engine/Source/Programs/Horde/HordeStorage/Horde.Storage/Tests/UnitTests/BlobIdentifierTests.cs
@@ -0,0 +1,25 @@
+using Horde.Storage.Implementation;
+using Microsoft.VisualStudio.TestTools.UnitTesting;
+using Newtonsoft.Json;
+
+namespace Horde.Storage.Tests.Unit
+{
+    [TestClass]
+    public class BlobIdentifierTests
+    {
+        [TestMethod]
+        public void TestDeserialization()
+        {
+            string s =
+                @"{""namespace"":""ue.ddc"",""bucket"":""texture"",""key"":""e5d70ae6567ffa9b5b3e69b2f45280e04fc5227f"",""op"":0,""timestamp"":""2022-01-03T16:00:00Z"",""timeBucket"":""rep-132856992000000000"",""eventId"":""546d4d89-6cae-11ec-8461-c64372ba79e9"",""blob"":null}";
+
+            ReplicationLogEvent? e = JsonConvert.DeserializeObject<ReplicationLogEvent>(s);
+
+            Assert.IsNotNull(e);
+            Assert.AreEqual("ue.ddc", e.Namespace.ToString());
+            Assert.AreEqual("texture", e.Bucket.ToString());
+            Assert.AreEqual("e5d70ae6567ffa9b5b3e69b2f45280e04fc5227f", e.Key.ToString());
+            Assert.AreEqual(null, e.Blob);
+        }
+    }
+}
diff --git a/Engine/Source/Programs/Horde/HordeStorage/Jupiter.Common/Implementation/BlobIdentifier.cs b/Engine/Source/Programs/Horde/HordeStorage/Jupiter.Common/Implementation/BlobIdentifier.cs
--- a/Engine/Source/Programs/Horde/HordeStorage/Jupiter.Common/Implementation/BlobIdentifier.cs
+++ b/Engine/Source/Programs/Horde/HordeStorage/Jupiter.Common/Implementation/BlobIdentifier.cs
@@ -159,17 +159,21 @@ namespace Jupiter.Implementation
         }  
     }
 
-    public class BlobIdentifierConverter : JsonConverter<BlobIdentifier>
+    public class BlobIdentifierConverter : JsonConverter<BlobIdentifier?>
     {
         public override void WriteJson(JsonWriter writer, BlobIdentifier? value, JsonSerializer serializer)
         {
             writer.WriteValue(value!.ToString());
         }
 
-        public override BlobIdentifier ReadJson(JsonReader reader, Type objectType, BlobIdentifier? existingValue, bool hasExistingValue, Newtonsoft.Json.JsonSerializer serializer)
+        public override BlobIdentifier? ReadJson(JsonReader reader, Type objectType, BlobIdentifier? existingValue, bool hasExistingValue, JsonSerializer serializer)
         {
+            if (reader.TokenType == JsonToken.Null) 
+                return null;
             string? s = (string?)reader.Value;
 
+            if (s == null)
+                return null;
             return new BlobIdentifier(s!);
         }
     }
