Horde Storage - Fixed issue were HEAD checks against the compressed blob would always fail as we were checking for the wrong blob (was trying to resolve the content id identifier rather then the resolved blob identifiers)

Also did some general improvements and documentation as part of reading thru the code a few extra times.

diff --git a/Engine/Source/Programs/Horde/HordeStorage/Horde.Storage/Controllers/CompressedBlobController.cs b/Engine/Source/Programs/Horde/HordeStorage/Horde.Storage/Controllers/CompressedBlobController.cs
--- a/Engine/Source/Programs/Horde/HordeStorage/Horde.Storage/Controllers/CompressedBlobController.cs
+++ b/Engine/Source/Programs/Horde/HordeStorage/Horde.Storage/Controllers/CompressedBlobController.cs
@@ -48,7 +48,7 @@ namespace Horde.Storage.Controllers
         [Authorize("Storage.read")]
         [ProducesResponseType(type: typeof(byte[]), 200)]
         [ProducesResponseType(type: typeof(ValidationProblemDetails), 400)]
-        [Produces(MediaTypeNames.Application.Json, CustomMediaTypeNames.UnrealCompactBinary, CustomMediaTypeNames.UnrealCompressedBuffer)]
+        [Produces(CustomMediaTypeNames.UnrealCompressedBuffer)]
 
         public async Task<IActionResult> Get(
             [Required] NamespaceId ns,
@@ -91,15 +91,15 @@ namespace Horde.Storage.Controllers
                 return Forbid();
             }
             BlobIdentifier[]? chunks = await _contentIdStore.Resolve(ns, id);
-            if (chunks == null)
+            if (chunks == null || chunks.Length == 0)
             {
-                return NotFound(new ValidationProblemDetails { Title = $"Content-id {id} not found"});
+                return NotFound();
             }
 
             Task<bool>[] tasks = new Task<bool>[chunks.Length];
             for (int i = 0; i < chunks.Length; i++)
             {
-                tasks[i] = _storage.Exists(ns, id);
+                tasks[i] = _storage.Exists(ns, chunks[i]);
             }
 
             await Task.WhenAll(tasks);
@@ -108,7 +108,7 @@ namespace Horde.Storage.Controllers
 
             if (!exists)
             {
-                return NotFound(new ValidationProblemDetails { Title = $"Object {id} not found"});
+                return NotFound();
             }
 
             return Ok();
@@ -160,11 +160,18 @@ namespace Horde.Storage.Controllers
         private async Task<BlobContents> GetImpl(NamespaceId ns, BlobIdentifier contentId)
         {
             BlobIdentifier[]? chunks = await _contentIdStore.Resolve(ns, contentId);
-            if (chunks == null)
+            if (chunks == null || chunks.Length == 0)
             {
                 throw new ContentIdResolveException(contentId);
             }
 
+            // single chunk, we just return that chunk
+            if (chunks.Length == 1)
+            {
+                return await _storage.GetObject(ns, chunks[0]);
+            }
+
+            // chunked content, combine the chunks into a single stream
             using Scope _ = Tracer.Instance.StartActive("blob.combine");
             Task<BlobContents>[] tasks = new Task<BlobContents>[chunks.Length];
             for (int i = 0; i < chunks.Length; i++)
diff --git a/Engine/Source/Programs/Horde/HordeStorage/Horde.Storage/Implementation/ContentId/ScyllaContentIdStore.cs b/Engine/Source/Programs/Horde/HordeStorage/Horde.Storage/Implementation/ContentId/ScyllaContentIdStore.cs
--- a/Engine/Source/Programs/Horde/HordeStorage/Horde.Storage/Implementation/ContentId/ScyllaContentIdStore.cs
+++ b/Engine/Source/Programs/Horde/HordeStorage/Horde.Storage/Implementation/ContentId/ScyllaContentIdStore.cs
@@ -41,6 +41,7 @@ namespace Horde.Storage.Implementation
 
             Task<bool> blobStoreExistsTask = _blobStore.Exists(ns, contentId);
 
+            BlobIdentifier[]? resolvedBlobs = null;
             // lower content_weight means its a better candidate to resolve to
             foreach (ScyllaContentId? resolvedContentId in await _mapper.FetchAsync<ScyllaContentId>("WHERE content_id = ? ORDER BY content_weight DESC", new ScyllaBlobIdentifier(contentId)))
             {
@@ -57,6 +58,8 @@ namespace Horde.Storage.Implementation
                         return blobs;
                 }
                 // blobs are missing continue testing with the next content id in the weighted list as that might exist
+
+                resolvedBlobs = blobs;
             }
 
             // if no content id is found, but we have a blob that matches the content id (so a unchunked and uncompressed version of the data) we use that instead
@@ -64,6 +67,12 @@ namespace Horde.Storage.Implementation
 
             if (contentIdBlobExists)
                 return new[] { contentId };
+            
+            if (resolvedBlobs != null)
+            {
+                // content id found, but blobs is missing, still we know which blobs it was so we return that list
+                return resolvedBlobs;
+            }
 
             // unable to resolve the content id
             return null;
