Horde Storage - Refactored namespace policy resolving to allow us some more flexiblity.
Used this to include a default setting for the Jupiter internal namespace which is used by the snapshotting service

diff --git a/Engine/Source/Programs/Horde/HordeStorage/Horde.Storage/Implementation/Blob/AmazonS3Store.cs b/Engine/Source/Programs/Horde/HordeStorage/Horde.Storage/Implementation/Blob/AmazonS3Store.cs
--- a/Engine/Source/Programs/Horde/HordeStorage/Horde.Storage/Implementation/Blob/AmazonS3Store.cs
+++ b/Engine/Source/Programs/Horde/HordeStorage/Horde.Storage/Implementation/Blob/AmazonS3Store.cs
@@ -9,6 +9,7 @@ using Amazon.S3;
 using Amazon.S3.Model;
 using Horde.Storage.Implementation.Blob;
 using Jupiter;
+using Jupiter.Common;
 using Jupiter.Common.Implementation;
 using Jupiter.Implementation;
 using Microsoft.Extensions.Options;
@@ -22,16 +23,16 @@ namespace Horde.Storage.Implementation
         private readonly ILogger _logger = Log.ForContext<AmazonS3Store>();
         private readonly IAmazonS3 _amazonS3;
         private readonly IBlobIndex _blobIndex;
-        private readonly IOptionsMonitor<NamespaceSettings> _namespaceSettings;
+        private readonly INamespacePolicyResolver _namespacePolicyResolver;
         private readonly S3Settings _settings;
         private readonly HashSet<string> _bucketAccessPolicyApplied = new HashSet<string>();
         private readonly HashSet<string> _bucketExistenceChecked = new HashSet<string>();
 
-        public AmazonS3Store(IAmazonS3 amazonS3, IOptionsMonitor<S3Settings> settings, IBlobIndex blobIndex, IOptionsMonitor<NamespaceSettings> namespaceSettings)
+        public AmazonS3Store(IAmazonS3 amazonS3, IOptionsMonitor<S3Settings> settings, IBlobIndex blobIndex, INamespacePolicyResolver namespacePolicyResolver)
         {
             _amazonS3 = amazonS3;
             _blobIndex = blobIndex;
-            _namespaceSettings = namespaceSettings;
+            _namespacePolicyResolver = namespacePolicyResolver;
             _settings = settings.CurrentValue;
         }
 
@@ -126,7 +127,7 @@ namespace Horde.Storage.Implementation
         {
             try
             {
-                string storagePool = _namespaceSettings.CurrentValue.GetPoliciesForNs(ns).StoragePool;
+                string storagePool = _namespacePolicyResolver.GetPoliciesForNs(ns).StoragePool;
                 string storagePoolSuffix = string.IsNullOrEmpty(storagePool) ? "" : $"-{storagePool}";
                 return $"{_settings.BucketName}{storagePoolSuffix}";
             }
diff --git a/Engine/Source/Programs/Horde/HordeStorage/Horde.Storage/Implementation/GC/OrphanBlobCleanupRefs.cs b/Engine/Source/Programs/Horde/HordeStorage/Horde.Storage/Implementation/GC/OrphanBlobCleanupRefs.cs
--- a/Engine/Source/Programs/Horde/HordeStorage/Horde.Storage/Implementation/GC/OrphanBlobCleanupRefs.cs
+++ b/Engine/Source/Programs/Horde/HordeStorage/Horde.Storage/Implementation/GC/OrphanBlobCleanupRefs.cs
@@ -10,6 +10,7 @@ using Dasync.Collections;
 using Horde.Storage.Controllers;
 using Horde.Storage.Implementation.Blob;
 using Jupiter;
+using Jupiter.Common;
 using Jupiter.Implementation;
 using Microsoft.Extensions.Options;
 using Serilog;
@@ -23,18 +24,18 @@ namespace Horde.Storage.Implementation
         private readonly IBlobIndex _blobIndex;
         private readonly ILeaderElection _leaderElection;
         private readonly IOptionsMonitor<GCSettings> _gcSettings;
-        private readonly IOptionsMonitor<NamespaceSettings> _namespaceSettings;
+        private readonly INamespacePolicyResolver _namespacePolicyResolver;
         private readonly ILogger _logger = Log.ForContext<OrphanBlobCleanupRefs>();
 
         // ReSharper disable once UnusedMember.Global
-        public OrphanBlobCleanupRefs(IBlobService blobService, IObjectService objectService, IBlobIndex blobIndex, ILeaderElection leaderElection, IOptionsMonitor<GCSettings> gcSettings, IOptionsMonitor<NamespaceSettings> namespaceSettings)
+        public OrphanBlobCleanupRefs(IBlobService blobService, IObjectService objectService, IBlobIndex blobIndex, ILeaderElection leaderElection, IOptionsMonitor<GCSettings> gcSettings, INamespacePolicyResolver namespacePolicyResolver)
         {
             _blobService = blobService;
             _objectService = objectService;
             _blobIndex = blobIndex;
             _leaderElection = leaderElection;
             _gcSettings = gcSettings;
-            _namespaceSettings = namespaceSettings;
+            _namespacePolicyResolver = namespacePolicyResolver;
         }
 
         public async Task<List<BlobIdentifier>> Cleanup(CancellationToken cancellationToken)
@@ -67,10 +68,10 @@ namespace Horde.Storage.Implementation
                         break;
 
                     bool found = false;
-                    NamespaceSettings.PerNamespaceSettings policy = _namespaceSettings.CurrentValue.GetPoliciesForNs(@namespace);
+                    NamespaceSettings.PerNamespaceSettings policy = _namespacePolicyResolver.GetPoliciesForNs(@namespace);
                     
                     // check all other namespaces that share the same storage pool for presence of the blob
-                    foreach (NamespaceId blobNamespace in namespaces.Where(ns => _namespaceSettings.CurrentValue.GetPoliciesForNs(ns).StoragePool == policy.StoragePool))
+                    foreach (NamespaceId blobNamespace in namespaces.Where(ns => _namespacePolicyResolver.GetPoliciesForNs(ns).StoragePool == policy.StoragePool))
                     {
                         IBlobIndex.BlobInfo? blobIndex = await _blobIndex.GetBlobInfo(blobNamespace, blob);
 
diff --git a/Engine/Source/Programs/Horde/HordeStorage/Horde.Storage/Tests/UnitTests/S3StoreTests.cs b/Engine/Source/Programs/Horde/HordeStorage/Horde.Storage/Tests/UnitTests/S3StoreTests.cs
--- a/Engine/Source/Programs/Horde/HordeStorage/Horde.Storage/Tests/UnitTests/S3StoreTests.cs
+++ b/Engine/Source/Programs/Horde/HordeStorage/Horde.Storage/Tests/UnitTests/S3StoreTests.cs
@@ -10,6 +10,7 @@ using Amazon.S3.Model;
 using Horde.Storage.Implementation;
 using Horde.Storage.Implementation.Blob;
 using Jupiter;
+using Jupiter.Common;
 using Jupiter.Implementation;
 using Microsoft.Extensions.Options;
 using Microsoft.VisualStudio.TestTools.UnitTesting;
@@ -29,18 +30,13 @@ namespace Horde.Storage.UnitTests
             S3Settings settings = new S3Settings{BucketName = "tests-foo"};
             IOptionsMonitor<S3Settings> settingsMock = Mock.Of<IOptionsMonitor<S3Settings>>(_ => _.CurrentValue == settings);
 
-            NamespaceSettings namespaceSettings = new NamespaceSettings{ Policies = new Dictionary<string, NamespaceSettings.PerNamespaceSettings>()
-            {
-                {
-                    Namespace.ToString(), new NamespaceSettings.PerNamespaceSettings {
-                        StoragePool = ""
-                    }
-                }
-                
-            }};
-            IOptionsMonitor<NamespaceSettings> namespaceSettingsMock = Mock.Of<IOptionsMonitor<NamespaceSettings>>(_ => _.CurrentValue == namespaceSettings);
-
-            AmazonS3Store store = new AmazonS3Store(s3Mock.Object, settingsMock, Mock.Of<IBlobIndex>(), namespaceSettingsMock);
+            INamespacePolicyResolver namespacePolicyResolver = Mock.Of<INamespacePolicyResolver>(_ => _.GetPoliciesForNs(Namespace) ==
+                                                  new NamespaceSettings.PerNamespaceSettings
+                                                  {
+                                                      StoragePool = ""
+                                                  });
+
+            AmazonS3Store store = new AmazonS3Store(s3Mock.Object, settingsMock, Mock.Of<IBlobIndex>(), namespacePolicyResolver);
             byte[] content = Encoding.ASCII.GetBytes("test content");
             BlobIdentifier blobIdentifier = BlobIdentifier.FromBlob(content);
             Task task = store.PutObject(Namespace, content.AsMemory(), blobIdentifier);
@@ -61,19 +57,14 @@ namespace Horde.Storage.UnitTests
             Mock<IAmazonS3> s3Mock = new Mock<IAmazonS3>();
             S3Settings settings = new S3Settings{BucketName = "tests-foo"};
             IOptionsMonitor<S3Settings> settingsMock = Mock.Of<IOptionsMonitor<S3Settings>>(_ => _.CurrentValue == settings);
-
-            NamespaceSettings namespaceSettings = new NamespaceSettings{ Policies = new Dictionary<string, NamespaceSettings.PerNamespaceSettings>()
-            {
+            
+            INamespacePolicyResolver namespacePolicyResolver = Mock.Of<INamespacePolicyResolver>(_ => _.GetPoliciesForNs(Namespace) ==
+                new NamespaceSettings.PerNamespaceSettings
                 {
-                    Namespace.ToString(), new NamespaceSettings.PerNamespaceSettings {
-                        StoragePool = "storagepool"
-                    }
-                }
-                
-            }};
-            IOptionsMonitor<NamespaceSettings> namespaceSettingsMock = Mock.Of<IOptionsMonitor<NamespaceSettings>>(_ => _.CurrentValue == namespaceSettings);
-
-            AmazonS3Store store = new AmazonS3Store(s3Mock.Object, settingsMock, Mock.Of<IBlobIndex>(), namespaceSettingsMock);
+                    StoragePool = "storagepool"
+                });
+
+            AmazonS3Store store = new AmazonS3Store(s3Mock.Object, settingsMock, Mock.Of<IBlobIndex>(), namespacePolicyResolver);
             byte[] content = Encoding.ASCII.GetBytes("test content");
             BlobIdentifier blobIdentifier = BlobIdentifier.FromBlob(content);
             Task task = store.PutObject(Namespace, content.AsMemory(), blobIdentifier);
@@ -98,19 +89,13 @@ namespace Horde.Storage.UnitTests
             Mock<IAmazonS3> s3Mock = new Mock<IAmazonS3>();
             S3Settings settings = new S3Settings{BucketName = "tests-foo"};
             IOptionsMonitor<S3Settings> settingsMock = Mock.Of<IOptionsMonitor<S3Settings>>(_ => _.CurrentValue == settings);
-            NamespaceSettings namespaceSettings = new NamespaceSettings{ Policies = new Dictionary<string, NamespaceSettings.PerNamespaceSettings>()
-            {
+            INamespacePolicyResolver namespacePolicyResolver = Mock.Of<INamespacePolicyResolver>(_ => _.GetPoliciesForNs(Namespace) ==
+                new NamespaceSettings.PerNamespaceSettings
                 {
-                    Namespace.ToString(), new NamespaceSettings.PerNamespaceSettings {
-                        StoragePool = ""
-                    }
-                }
-                
-            }};
-            IOptionsMonitor<NamespaceSettings> namespaceSettingsMock = Mock.Of<IOptionsMonitor<NamespaceSettings>>(_ => _.CurrentValue == namespaceSettings);
-
+                    StoragePool = ""
+                });
             s3Mock.Setup(s3 => s3.PutObjectAsync(It.IsAny<PutObjectRequest>(), default)).Throws<Exception>();
-            AmazonS3Store store = new AmazonS3Store(s3Mock.Object, settingsMock, Mock.Of<IBlobIndex>(), namespaceSettingsMock);
+            AmazonS3Store store = new AmazonS3Store(s3Mock.Object, settingsMock, Mock.Of<IBlobIndex>(), namespacePolicyResolver);
             Task task = store.PutObject(Namespace, content, blob);
             await task;
 
@@ -131,18 +116,12 @@ namespace Horde.Storage.UnitTests
             S3Settings settings = new S3Settings{BucketName = "tests-foo"};
             IOptionsMonitor<S3Settings> settingsMock = Mock.Of<IOptionsMonitor<S3Settings>>(_ => _.CurrentValue == settings);
 
-            NamespaceSettings namespaceSettings = new NamespaceSettings{ Policies = new Dictionary<string, NamespaceSettings.PerNamespaceSettings>()
-            {
+            INamespacePolicyResolver namespacePolicyResolver = Mock.Of<INamespacePolicyResolver>(_ => _.GetPoliciesForNs(Namespace) ==
+                new NamespaceSettings.PerNamespaceSettings
                 {
-                    Namespace.ToString(), new NamespaceSettings.PerNamespaceSettings {
-                        StoragePool = ""
-                    }
-                }
-                
-            }};
-            IOptionsMonitor<NamespaceSettings> namespaceSettingsMock = Mock.Of<IOptionsMonitor<NamespaceSettings>>(_ => _.CurrentValue == namespaceSettings);
-
-            AmazonS3Store store = new AmazonS3Store(s3Mock.Object, settingsMock, Mock.Of<IBlobIndex>(), namespaceSettingsMock);
+                    StoragePool = ""
+                });
+            AmazonS3Store store = new AmazonS3Store(s3Mock.Object, settingsMock, Mock.Of<IBlobIndex>(), namespacePolicyResolver);
             await using BlobContents blobContents = await store.GetObject(Namespace, blob);
 
             s3Mock.Verify(s3 => s3.GetObjectAsync("tests-foo", blob.AsS3Key(), default));
@@ -157,18 +136,13 @@ namespace Horde.Storage.UnitTests
             S3Settings settings = new S3Settings{BucketName = "tests-foo"};
             IOptionsMonitor<S3Settings> settingsMock = Mock.Of<IOptionsMonitor<S3Settings>>(_ => _.CurrentValue == settings);
 
-            NamespaceSettings namespaceSettings = new NamespaceSettings{ Policies = new Dictionary<string, NamespaceSettings.PerNamespaceSettings>()
-            {
+            INamespacePolicyResolver namespacePolicyResolver = Mock.Of<INamespacePolicyResolver>(_ => _.GetPoliciesForNs(Namespace) ==
+                new NamespaceSettings.PerNamespaceSettings
                 {
-                    Namespace.ToString(), new NamespaceSettings.PerNamespaceSettings {
-                        StoragePool = ""
-                    }
-                }
-                
-            }};
-            IOptionsMonitor<NamespaceSettings> namespaceSettingsMock = Mock.Of<IOptionsMonitor<NamespaceSettings>>(_ => _.CurrentValue == namespaceSettings);
-
-            AmazonS3Store store = new AmazonS3Store(s3Mock.Object, settingsMock, Mock.Of<IBlobIndex>(), namespaceSettingsMock);
+                    StoragePool = ""
+                });
+
+            AmazonS3Store store = new AmazonS3Store(s3Mock.Object, settingsMock, Mock.Of<IBlobIndex>(), namespacePolicyResolver);
             Task task = store.DeleteObject(Namespace, blob);
             await task;
 
diff --git a/Engine/Source/Programs/Horde/HordeStorage/Jupiter.Common/Authentication/NamespaceAuthorizationHandler.cs b/Engine/Source/Programs/Horde/HordeStorage/Jupiter.Common/Authentication/NamespaceAuthorizationHandler.cs
--- a/Engine/Source/Programs/Horde/HordeStorage/Jupiter.Common/Authentication/NamespaceAuthorizationHandler.cs
+++ b/Engine/Source/Programs/Horde/HordeStorage/Jupiter.Common/Authentication/NamespaceAuthorizationHandler.cs
@@ -4,6 +4,7 @@ using System;
 using System.Collections.Generic;
 using System.ComponentModel.DataAnnotations;
 using System.Threading.Tasks;
+using Jupiter.Common;
 using Jupiter.Implementation;
 using Microsoft.AspNetCore.Authorization;
 using Microsoft.Extensions.Options;
@@ -13,11 +14,11 @@ namespace Jupiter
     // verifies that you have access to a namespace by checking if you have a corresponding claim to that namespace
     public class NamespaceAuthorizationHandler : AuthorizationHandler<NamespaceAccessRequirement, NamespaceId>
     {
-        private readonly IOptionsMonitor<NamespaceSettings> _namespaceSettings;
+        private readonly INamespacePolicyResolver _namespacePolicyResolver;
 
-        public NamespaceAuthorizationHandler(IOptionsMonitor<NamespaceSettings> namespaceSettings)
+        public NamespaceAuthorizationHandler(INamespacePolicyResolver namespacePolicyResolver)
         {
-            _namespaceSettings = namespaceSettings;
+            _namespacePolicyResolver = namespacePolicyResolver;
         }
 
         protected override Task HandleRequirementAsync(AuthorizationHandlerContext context, NamespaceAccessRequirement requirement,
@@ -29,7 +30,7 @@ namespace Jupiter
                 return Task.CompletedTask;
             }
 
-            NamespaceSettings.PerNamespaceSettings settings = _namespaceSettings.CurrentValue.GetPoliciesForNs(namespaceName);
+            NamespaceSettings.PerNamespaceSettings settings = _namespacePolicyResolver.GetPoliciesForNs(namespaceName);
             // These are ANDed, e.g. all claims needs to be present
             foreach (string expectedClaim in settings.Claims)
             {
diff --git a/Engine/Source/Programs/Horde/HordeStorage/Jupiter.Common/BaseStartup.cs b/Engine/Source/Programs/Horde/HordeStorage/Jupiter.Common/BaseStartup.cs
--- a/Engine/Source/Programs/Horde/HordeStorage/Jupiter.Common/BaseStartup.cs
+++ b/Engine/Source/Programs/Horde/HordeStorage/Jupiter.Common/BaseStartup.cs
@@ -11,6 +11,7 @@ using System.Text;
 using System.Threading.Tasks;
 using Amazon;
 using Datadog.Trace;
+using Jupiter.Common;
 using Jupiter.Implementation;
 using Microsoft.AspNetCore.Authentication;
 using Microsoft.AspNetCore.Authorization;
@@ -71,6 +72,8 @@ namespace Jupiter
             services.AddOptions<JupiterSettings>().Bind(Configuration.GetSection("Jupiter")).ValidateDataAnnotations();
             services.AddOptions<NamespaceSettings>().Bind(Configuration.GetSection("Namespaces")).ValidateDataAnnotations();
 
+            services.AddSingleton(typeof(INamespacePolicyResolver), typeof(NamespacePolicyResolver));
+
             // this is the same as invoke MvcBuilder.AddJsonOptions but with a service provider passed so we can use DI in the options creation
             // see https://stackoverflow.com/questions/53288633/net-core-api-custom-json-resolver-based-on-request-values
             // inject our custom json options and then pass a DI context to customize the serialization
@@ -472,18 +475,6 @@ namespace Jupiter
             public string[] Claims { get; set; } = Array.Empty<string>();
             public string StoragePool { get; set; } = "";
         }
-
-        public PerNamespaceSettings GetPoliciesForNs(NamespaceId ns)
-        {
-            return Policies[ns.ToString()];
-        /*
-            if(Policies.TryGetValue(ns.ToString(), out PerNamespaceSettings? settings))
-            {
-                return settings;
-            }
-
-            return DefaultNamespaceSettings;*/
-        }
     }
 
     public class DatadogTraceMiddleware
