Horde Storage - NamespacePolicyResolver was missing from change 18569766

diff --git a/Engine/Source/Programs/Horde/HordeStorage/Jupiter.Common/Implementation/NamespacePolicyResolver.cs b/Engine/Source/Programs/Horde/HordeStorage/Jupiter.Common/Implementation/NamespacePolicyResolver.cs
new file mode 100644
--- /dev/null
+++ b/Engine/Source/Programs/Horde/HordeStorage/Jupiter.Common/Implementation/NamespacePolicyResolver.cs
@@ -0,0 +1,41 @@
+// Copyright Epic Games, Inc. All Rights Reserved.
+
+using Jupiter.Implementation;
+using Microsoft.Extensions.Options;
+
+namespace Jupiter.Common
+{
+    public interface INamespacePolicyResolver
+    {
+        public NamespaceSettings.PerNamespaceSettings GetPoliciesForNs(NamespaceId ns);
+
+        static NamespaceId JupiterInternalNamespace
+        {
+            get { return new NamespaceId("jupiter-internal"); }
+        }
+    }
+    public class NamespacePolicyResolver : INamespacePolicyResolver
+    {
+        private readonly IOptionsMonitor<NamespaceSettings> _namespaceSettings;
+
+        public NamespacePolicyResolver(IOptionsMonitor<NamespaceSettings> namespaceSettings)
+        {
+            _namespaceSettings = namespaceSettings;
+        }
+
+        public NamespaceSettings.PerNamespaceSettings GetPoliciesForNs(NamespaceId ns)
+        {
+            if (ns == INamespacePolicyResolver.JupiterInternalNamespace)
+            {
+                return new NamespaceSettings.PerNamespaceSettings()
+                {
+                    // the internal jupiter namespace is sensitive so only admins should have access, but we can still use the default storage pool
+                    Claims = new string[] { "Admin" },
+                    StoragePool = ""
+                };
+            }
+
+            return _namespaceSettings.CurrentValue.Policies[ns.ToString()];
+        }
+    }
+}
