Horde: Add notice service
#rnx

diff --git a/Engine/Source/Programs/Horde/Horde.Build/Api/Notices.cs b/Engine/Source/Programs/Horde/Horde.Build/Api/Notices.cs
new file mode 100644
--- /dev/null
+++ b/Engine/Source/Programs/Horde/Horde.Build/Api/Notices.cs
@@ -0,0 +1,130 @@
+// Copyright Epic Games, Inc. All Rights Reserved.
+
+using System;
+using System.ComponentModel.DataAnnotations;
+using HordeServer.Models;
+
+namespace HordeServer.Api
+{
+
+	/// <summary>
+	/// Parameters required to create a notice
+	/// </summary>
+	public class CreateNoticeRequest
+	{
+		/// <summary>
+		/// Start time to display this message
+		/// </summary>
+		public DateTime? StartTime { get; set; }
+
+		/// <summary>
+		/// Finish time to display this message
+		/// </summary>
+		public DateTime? FinishTime { get; set; }
+
+		/// <summary>
+		/// Message to display
+		/// </summary>
+		[Required]
+		public string Message { get; set; }
+
+		/// <summary>
+		/// Constructor
+		/// </summary>
+		public CreateNoticeRequest(string Message)
+		{
+			this.Message = Message;
+		}
+
+	}
+
+	/// <summary>
+	/// Parameters required to update a notice
+	/// </summary>
+	public class UpdateNoticeRequest
+	{
+
+		/// <summary>
+		/// The id of the notice to update
+		/// </summary>
+		public string Id { get; set; }
+
+
+		/// <summary>
+		/// Start time to display this message
+		/// </summary>
+		public DateTime? StartTime { get; set; }
+
+		/// <summary>
+		/// Finish time to display this message
+		/// </summary>
+		public DateTime? FinishTime { get; set; }
+
+		/// <summary>
+		/// Message to display
+		/// </summary>
+		public string? Message { get; set; }
+
+		/// <summary>
+		/// Constructor
+		/// </summary>
+		public UpdateNoticeRequest(string Id)
+		{
+			this.Id = Id;
+		}
+
+	}
+
+
+	/// <summary>
+	/// Notice informtation
+	/// </summary>
+	public class GetNoticeResponse
+	{
+		/// <summary>
+		/// Object id of notice
+		/// </summary>
+		public string? Id { get; set; }
+
+		/// <summary>
+		/// Start time to display this message
+		/// </summary>
+		public DateTime? StartTime { get; set; }
+
+		/// <summary>
+		/// Finish time to display this message
+		/// </summary>
+		public DateTime? FinishTime { get; set; }
+
+		/// <summary>
+		/// User id who created the notice, otherwise null if a system message
+		/// </summary>
+		public GetThinUserInfoResponse? CreatedByUser { get; set; }
+
+		/// <summary>
+		/// Message to display
+		/// </summary>
+		[Required]
+		public string Message { get; set; } = String.Empty;
+
+		/// <summary>
+		/// Constructor
+		/// </summary>
+		public GetNoticeResponse()
+		{
+
+		}
+
+		/// <summary>
+		/// Constructor
+		/// </summary>
+		public GetNoticeResponse(INotice Notice, GetThinUserInfoResponse? CreateByUser = null)
+		{
+			this.Id = Notice.Id.ToString();
+			this.StartTime = Notice.StartTime;
+			this.FinishTime = Notice.FinishTime;
+			this.Message = Notice.Message;
+			this.CreatedByUser = CreateByUser;
+		}
+	}
+}
diff --git a/Engine/Source/Programs/Horde/Horde.Build/Controllers/NoticesController.cs b/Engine/Source/Programs/Horde/Horde.Build/Controllers/NoticesController.cs
--- a/Engine/Source/Programs/Horde/Horde.Build/Controllers/NoticesController.cs
+++ b/Engine/Source/Programs/Horde/Horde.Build/Controllers/NoticesController.cs
@@ -1,5 +1,7 @@
 // Copyright Epic Games, Inc. All Rights Reserved.
 
+using HordeServer.Api;
+using HordeServer.Collections;
 using HordeServer.Models;
 using HordeServer.Services;
 using HordeServer.Utilities;
@@ -10,14 +12,13 @@ using MongoDB.Bson;
 using System;
 using System.Collections.Generic;
 using System.Globalization;
-using System.Linq;
 using System.Threading.Tasks;
 using TimeZoneConverter;
 
 namespace HordeServer.Controllers
 {
 	/// <summary>
-	/// Controller for the /api/v1/status endpoint
+	/// Controller for the /api/v1/notices endpoint
 	/// </summary>
 	[ApiController]
 	[Authorize]
@@ -25,41 +26,83 @@ namespace HordeServer.Controllers
 	public class NoticesController : ControllerBase
 	{
 		/// <summary>
-		/// Timezone to use for display notices
+		/// 
 		/// </summary>
-		TimeZoneInfo TimeZone;
+		NoticeService NoticeService;
 
 		/// <summary>
-		/// The globals singleton
+		/// The acl service singleton
+		/// </summary>
+		AclService AclService;
+
+		/// <summary>
+		/// Reference to the globals singleton
 		/// </summary>
 		ISingletonDocument<Globals> Globals;
 
+		/// <summary>
+		/// Reference to the user collection
+		/// </summary>
+		private readonly IUserCollection UserCollection;
+
+		/// <summary>
+		/// Server settings
+		/// </summary>
+		IOptionsMonitor<ServerSettings> Settings;
+
 		/// <summary>
 		/// Constructor
 		/// </summary>
 		/// <param name="Settings">The server settings</param>
-		/// <param name="Globals">The globals singleton</param>
-		public NoticesController(IOptions<ServerSettings> Settings, ISingletonDocument<Globals> Globals)
+		/// <param name="NoticeService">The notice service singleton</param>
+		/// <param name="AclService">The acl service singleton</param>
+		/// <param name="UserCollection">The user collection singleton</param>
+		/// <param name="Globals">The global singleton</param>
+		public NoticesController(IOptionsMonitor<ServerSettings> Settings, NoticeService NoticeService, AclService AclService, IUserCollection UserCollection, ISingletonDocument<Globals> Globals)
 		{
-			string? ScheduleTimeZone = Settings.Value.ScheduleTimeZone;
-			this.TimeZone = (ScheduleTimeZone == null) ? TimeZoneInfo.Local : TZConvert.GetTimeZoneInfo(ScheduleTimeZone);
-
+			this.Settings = Settings;
+			this.NoticeService = NoticeService;
+			this.AclService = AclService;
+			this.UserCollection = UserCollection;
 			this.Globals = Globals;
 		}
 
 		/// <summary>
-		/// Update a status message
+		/// Add a status message
 		/// </summary>
-		/// <param name="Status"></param>
+		/// <param name="Request"></param>
 		/// <returns></returns>
 		[HttpPost("/api/v1/notices")]
-		public async Task<ActionResult> AddNoticeAsync(Notice Status)
+		public async Task<ActionResult> AddNoticeAsync(CreateNoticeRequest Request)
 		{
-			Status.Id = ObjectId.GenerateNewId().ToString();
-			await Globals.UpdateAsync(x => x.Notices.Add(Status));
-			return Ok();
+			if (!await AclService.AuthorizeAsync(AclAction.AdminWrite, User))
+			{
+				return Forbid();
+			}
+
+			INotice? Notice = await NoticeService.AddNoticeAsync(Request.Message, User.GetUserId(), Request.StartTime, Request.FinishTime);
+
+			return Notice == null ? NotFound() : Ok();
+
 		}
 
+		/// <summary>
+		/// Update a status message
+		/// </summary>
+		/// <param name="Request"></param>
+		/// <returns></returns>
+		[HttpPut("/api/v1/notices")]
+		public async Task<ActionResult> UpdateNoticeAsync(UpdateNoticeRequest Request)
+		{
+			if (!await AclService.AuthorizeAsync(AclAction.AdminWrite, User))
+			{
+				return Forbid();
+			}
+
+			return await NoticeService.UpdateNoticeAsync(new ObjectId(Request.Id), Request.Message, Request.StartTime, Request.FinishTime) ? Ok() : NotFound();
+		}
+
+
 		/// <summary>
 		/// Remove a manually added status message
 		/// </summary>
@@ -67,7 +110,13 @@ namespace HordeServer.Controllers
 		[HttpDelete("/api/v1/notices/{Id}")]
 		public async Task<ActionResult> DeleteNoticeAsync(string Id)
 		{
-			await Globals.UpdateAsync(x => x.Notices.RemoveAll(x => x.Id == Id));
+			if (!await AclService.AuthorizeAsync(AclAction.AdminWrite, User))
+			{
+				return Forbid();
+			}
+
+			await NoticeService.RemoveNoticeAsync(new ObjectId(Id));
+			
 			return Ok();
 		}
 
@@ -76,43 +125,65 @@ namespace HordeServer.Controllers
 		/// </summary>
 		/// <returns>The status messages</returns>
 		[HttpGet("/api/v1/notices")]
-		public async Task<List<Notice>> GetNoticesAsync()
-		{
-			Globals Value = await Globals.GetAsync();
-
-			List<Notice> Messages = new List<Notice>(Value.Notices);
+		public async Task<List<GetNoticeResponse>> GetNoticesAsync()
+		{			
+			string? ScheduleTimeZone = Settings.CurrentValue.ScheduleTimeZone;
+			TimeZoneInfo TimeZone = (ScheduleTimeZone == null) ? TimeZoneInfo.Local : TZConvert.GetTimeZoneInfo(ScheduleTimeZone);
+			
+			Globals GlobalsValue = await Globals.GetAsync();
+			List<GetNoticeResponse> Messages = new List<GetNoticeResponse>();
 
-			if(Value.ScheduledDowntime.Count > 0)
+			// Add system downtime notices
+			for (int Idx = 0; Idx < GlobalsValue.ScheduledDowntime.Count; Idx++)
 			{
-				DateTimeOffset Now = DateTimeOffset.Now;
+				ScheduledDowntime Downtime = GlobalsValue.ScheduledDowntime[Idx];
 
-				(DateTimeOffset StartTime, DateTimeOffset FinishTime) = Value.ScheduledDowntime[0].GetNext(Now);
-				for (int Idx = 1; Idx < Value.ScheduledDowntime.Count; Idx++)
+				// @todo: handle this, though likely will have a custom message
+				if (Downtime.Frequency == ScheduledDowntimeFrequency.Once)
 				{
-					(DateTimeOffset NextStartTime, DateTimeOffset NextFinishTime) = Value.ScheduledDowntime[Idx].GetNext(Now);
-					if (NextFinishTime < StartTime)
-					{
-						StartTime = NextStartTime;
-					}
-					if (NextStartTime < FinishTime && NextFinishTime > FinishTime)
-					{
-						FinishTime = NextFinishTime;
-					}
+					continue;
 				}
 
-				if (StartTime < Now)
+				DateTimeOffset Now = TimeZoneInfo.ConvertTime(DateTimeOffset.Now, Settings.CurrentValue.TimeZoneInfo);
+				
+				// @todo: this time and formatting should be on client
+				if (Downtime.IsActive(Now))
 				{
+					(DateTimeOffset StartTime, DateTimeOffset FinishTime) = Downtime.GetNext(Now);
 					string FinishTimeString = String.Format(CultureInfo.CurrentCulture, "{0:t}", TimeZoneInfo.ConvertTime(FinishTime, TimeZone));
-					Messages.Add(new Notice { StartTime = StartTime.UtcDateTime, FinishTime = FinishTime.UtcDateTime, Message = $"Horde is currently in scheduled downtime. Jobs will resume execution at {FinishTimeString} {TimeZone.Id}." });
+					Messages.Add(new GetNoticeResponse() { StartTime = StartTime.UtcDateTime, FinishTime = FinishTime.UtcDateTime, Message = $"Horde is currently in scheduled downtime. Jobs will resume execution at {FinishTimeString} {TimeZone.Id}." });
 				}
-				else if(StartTime < Now + TimeSpan.FromHours(12.0))
+				else 
 				{
-					string StartTimeString = String.Format(CultureInfo.CurrentCulture, "{0:t}", StartTime);
-					string FinishTimeString = String.Format(CultureInfo.CurrentCulture, "{0:t}", FinishTime);
-					Messages.Add(new Notice { StartTime = StartTime.UtcDateTime, FinishTime = FinishTime.UtcDateTime, Message = $"Downtime is scheduled from {StartTimeString} to {FinishTimeString} {TimeZone.Id}. No jobs will run during this time." });
+					// add a week to get the actual next scheduled time
+					if (Downtime.Frequency == ScheduledDowntimeFrequency.Weekly)
+					{
+						Now += TimeSpan.FromDays(7);
+					}
+
+					(DateTimeOffset StartTime, DateTimeOffset FinishTime) = Downtime.GetNext(Now);
+					string StartTimeString = String.Format(CultureInfo.CurrentCulture, "{0:g}", StartTime);
+					string FinishTimeString = String.Format(CultureInfo.CurrentCulture, "{0:g}", FinishTime);
+					Messages.Add(new GetNoticeResponse() { StartTime = StartTime.UtcDateTime, FinishTime = FinishTime.UtcDateTime, Message = $"Downtime is scheduled from {StartTimeString} to {FinishTimeString} {TimeZone.Id}. No jobs will run during this time." });
 				}
 			}
 
+			// Add user notices
+			List<INotice> Notices = await NoticeService.GetNoticesAsync();
+
+			for (int i = 0; i < Notices.Count; i++)
+			{
+				INotice Notice = Notices[i];
+				GetThinUserInfoResponse? UserInfo = null;
+
+				if (Notice.UserId != null)
+				{
+					UserInfo = new GetThinUserInfoResponse(await UserCollection.GetCachedUserAsync(Notice.UserId));
+				}								
+
+				Messages.Add(new GetNoticeResponse(Notice, UserInfo));
+			}
+
 			return Messages;
 		}
 	}
diff --git a/Engine/Source/Programs/Horde/Horde.Build/Models/Globals.cs b/Engine/Source/Programs/Horde/Horde.Build/Models/Globals.cs
--- a/Engine/Source/Programs/Horde/Horde.Build/Models/Globals.cs
+++ b/Engine/Source/Programs/Horde/Horde.Build/Models/Globals.cs
@@ -16,6 +16,8 @@ using System.Threading.Tasks;
 
 namespace HordeServer.Models
 {
+	using UserId = ObjectId<IUser>;
+
 	/// <summary>
 	/// Base class for singleton documents
 	/// </summary>
@@ -152,32 +154,6 @@ namespace HordeServer.Models
 		}
 	}
 
-	/// <summary>
-	/// User notice
-	/// </summary>
-	public class Notice
-	{
-		/// <summary>
-		/// Unique id for this notice
-		/// </summary>
-		public string Id { get; set; } = String.Empty;
-
-		/// <summary>
-		/// Start time to display this message
-		/// </summary>
-		public DateTime? StartTime { get; set; }
-
-		/// <summary>
-		/// Finish time to display this message
-		/// </summary>
-		public DateTime? FinishTime { get; set; }
-
-		/// <summary>
-		/// Message to display
-		/// </summary>
-		public string Message { get; set; } = String.Empty;
-	}
-
 	/// <summary>
 	/// Path to a platform and stream to use for syncing AutoSDK
 	/// </summary>
@@ -315,11 +291,6 @@ namespace HordeServer.Models
 		/// </summary>
 		public string? ConfigRevision { get; set; }
 
-		/// <summary>
-		/// Manually added status messages
-		/// </summary>
-		public List<Notice> Notices { get; set; } = new List<Notice>();
-
 		/// <summary>
 		/// List of Perforce clusters
 		/// </summary>
diff --git a/Engine/Source/Programs/Horde/Horde.Build/Notices/INotice.cs b/Engine/Source/Programs/Horde/Horde.Build/Notices/INotice.cs
new file mode 100644
--- /dev/null
+++ b/Engine/Source/Programs/Horde/Horde.Build/Notices/INotice.cs
@@ -0,0 +1,51 @@
+// Copyright Epic Games, Inc. All Rights Reserved.
+
+using System;
+using MongoDB.Bson.Serialization.Attributes;
+using HordeServer.Utilities;
+using MongoDB.Bson;
+
+namespace HordeServer.Models
+{
+	using UserId = ObjectId<IUser>;
+
+	/// <summary>
+	/// User notice
+	/// </summary>
+	public interface INotice
+	{
+		/// <summary>
+		/// Unique id for this reservation
+		/// </summary>
+		public ObjectId Id { get; }
+
+		/// <summary>
+		/// User id for who created this notice, null for system created notices
+		/// </summary>
+		[BsonIgnoreIfNull]
+		public UserId? UserId { get; }
+
+		/// <summary>
+		/// Start time to display this message
+		/// </summary>
+		public DateTime? StartTime { get; }
+
+		/// <summary>
+		/// Finish time to display this message
+		/// </summary>
+		public DateTime? FinishTime { get;  }
+
+		/// <summary>
+		/// Message to display
+		/// </summary>
+		public string Message { get; }
+	}
+
+}
+
+
+
+
+
+
+
diff --git a/Engine/Source/Programs/Horde/Horde.Build/Notices/INoticeCollection.cs b/Engine/Source/Programs/Horde/Horde.Build/Notices/INoticeCollection.cs
new file mode 100644
--- /dev/null
+++ b/Engine/Source/Programs/Horde/Horde.Build/Notices/INoticeCollection.cs
@@ -0,0 +1,62 @@
+// Copyright Epic Games, Inc. All Rights Reserved.
+
+using HordeServer.Models;
+using HordeServer.Utilities;
+using MongoDB.Bson;
+using System;
+using System.Collections.Generic;
+using System.Threading.Tasks;
+
+namespace HordeServer.Collections
+{
+	using UserId = ObjectId<IUser>;
+
+	/// <summary>
+	/// A collection of user specified notices
+	/// </summary>
+	public interface INoticeCollection
+	{
+		/// <summary>
+		/// Add a notice to the collection
+		/// </summary>
+		/// <param name="Message"></param>
+		/// <param name="UserId"></param>
+		/// <param name="StartTime"></param>
+		/// <param name="FinishTime"></param>
+		/// <returns></returns>
+		Task<INotice?> AddNoticeAsync(string Message, UserId? UserId, DateTime? StartTime, DateTime? FinishTime);
+
+		/// <summary>
+		/// Update an existing notice
+		/// </summary>
+		/// <param name="Id"></param>
+		/// <param name="Message"></param>
+		/// <param name="StartTime"></param>
+		/// <param name="FinishTime"></param>
+		/// <returns></returns>
+		Task<bool> UpdateNoticeAsync(ObjectId Id, string? Message, DateTime? StartTime, DateTime? FinishTime);
+
+		/// <summary>
+		/// Get a notice by id
+		/// </summary>
+		/// <param name="NoticeId"></param>
+		/// <returns></returns>
+		Task<INotice?> GetNoticeAsync(ObjectId NoticeId);
+
+		/// <summary>
+		/// Get all notices
+		/// </summary>
+		/// <returns></returns>
+		Task<List<INotice>> GetNoticesAsync();
+
+		/// <summary>
+		/// Remove a notice
+		/// </summary>
+		/// <param name="Id"></param>
+		/// <returns></returns>
+		Task<bool> RemoveNoticeAsync(ObjectId Id);
+
+
+	}
+
+}
\ No newline at end of file
diff --git a/Engine/Source/Programs/Horde/Horde.Build/Notices/Impl/NoticeCollection.cs b/Engine/Source/Programs/Horde/Horde.Build/Notices/Impl/NoticeCollection.cs
new file mode 100644
--- /dev/null
+++ b/Engine/Source/Programs/Horde/Horde.Build/Notices/Impl/NoticeCollection.cs
@@ -0,0 +1,125 @@
+// Copyright Epic Games, Inc. All Rights Reserved.
+
+using HordeServer.Services;
+using MongoDB.Bson.Serialization.Attributes;
+using System.Threading.Tasks;
+
+using HordeServer.Models;
+using MongoDB.Driver;
+using System;
+using HordeServer.Utilities;
+using System.Linq;
+
+using MongoDB.Bson;
+using System.Collections.Generic;
+
+namespace HordeServer.Collections.Impl
+{
+	using UserId = ObjectId<IUser>;
+
+	/// <summary>
+	/// Collection of notice documents
+	/// </summary>
+	public class NoticeCollection : INoticeCollection
+	{
+
+		/// <summary>
+		/// Document representing a notice 
+		/// </summary>
+		class NoticeDocument : INotice
+		{
+			[BsonRequired, BsonId]
+			public ObjectId Id { get; set; }
+
+			[BsonIgnoreIfNull]
+			public UserId? UserId { get; set; }
+
+			[BsonIgnoreIfNull]
+			public DateTime? StartTime { get; set; }
+
+			[BsonIgnoreIfNull]
+			public DateTime? FinishTime { get; set; }
+
+			public string Message { get; set; } = String.Empty;
+
+			[BsonConstructor]
+			private NoticeDocument()
+			{
+				
+			}
+
+			public NoticeDocument(ObjectId Id, string Message, UserId? UserId, DateTime? StartTime, DateTime? FinishTime)
+			{
+				this.Id = Id;
+				this.Message = Message;
+				this.UserId = UserId;
+				this.StartTime = StartTime;
+				this.FinishTime = FinishTime;
+			}
+
+		}
+
+		readonly IMongoCollection<NoticeDocument> Notices;
+
+		/// <summary>
+		/// Constructor
+		/// </summary>
+		public NoticeCollection(DatabaseService DatabaseService)
+		{
+			Notices = DatabaseService.GetCollection<NoticeDocument>("Notices");
+		}
+
+		/// <inheritdoc/>
+		public async Task<INotice?> AddNoticeAsync(string Message, UserId? UserId, DateTime? StartTime, DateTime? FinishTime)
+		{
+			NoticeDocument NewNotice = new NoticeDocument(ObjectId.GenerateNewId(), Message, UserId, StartTime, FinishTime);
+			await Notices.InsertOneAsync(NewNotice);
+			return NewNotice;
+		}
+
+		/// <inheritdoc/>
+		public async Task<bool> UpdateNoticeAsync(ObjectId Id, string? Message, DateTime? StartTime, DateTime? FinishTime)
+		{
+
+			UpdateDefinitionBuilder<NoticeDocument> UpdateBuilder = Builders<NoticeDocument>.Update;
+			List<UpdateDefinition<NoticeDocument>> Updates = new List<UpdateDefinition<NoticeDocument>>();
+
+			if (Message != null)
+			{
+				Updates.Add(UpdateBuilder.Set(x => x.Message, Message));
+			}
+
+			Updates.Add(UpdateBuilder.Set(x => x.StartTime, StartTime));
+
+			Updates.Add(UpdateBuilder.Set(x => x.FinishTime, FinishTime));
+
+			NoticeDocument? Document = await Notices.FindOneAndUpdateAsync(x => x.Id == Id, UpdateBuilder.Combine(Updates));
+
+			return Document != null;
+		}
+
+		/// <inheritdoc/>
+		public async Task<INotice?> GetNoticeAsync(ObjectId NoticeId)
+		{
+			return await Notices.Find(x => x.Id == NoticeId).FirstOrDefaultAsync();
+		}
+
+		/// <inheritdoc/>
+		public async Task<List<INotice>> GetNoticesAsync()
+		{
+			List<NoticeDocument> Results = await Notices.Find(x => true).ToListAsync();
+			return Results.Select<NoticeDocument, INotice>(x => x).ToList();
+		}
+
+
+		/// <inheritdoc/>
+		public async Task<bool> RemoveNoticeAsync(ObjectId Id)
+		{
+			DeleteResult Result = await Notices.DeleteOneAsync(x => x.Id == Id);
+			return Result.DeletedCount > 0;
+		}
+
+	}
+
+}
+
diff --git a/Engine/Source/Programs/Horde/Horde.Build/Notices/Impl/NoticeService.cs b/Engine/Source/Programs/Horde/Horde.Build/Notices/Impl/NoticeService.cs
new file mode 100644
--- /dev/null
+++ b/Engine/Source/Programs/Horde/Horde.Build/Notices/Impl/NoticeService.cs
@@ -0,0 +1,107 @@
+// Copyright Epic Games, Inc. All Rights Reserved.
+
+using System.Threading.Tasks;
+
+using HordeServer.Models;
+using HordeServer.Collections;
+using System;
+using HordeServer.Utilities;
+using System.Threading;
+
+using Microsoft.Extensions.Logging;
+using HordeCommon;
+using MongoDB.Bson;
+using System.Collections.Generic;
+using System.Security.Claims;
+using System.Diagnostics.CodeAnalysis;
+using HordeServer.Notifications;
+using System.Linq;
+using HordeServer.Api;
+using Microsoft.Extensions.Hosting;
+
+namespace HordeServer.Services
+{
+	using UserId = ObjectId<IUser>;
+
+	/// <summary>
+	/// Notice service
+	/// </summary>
+	public sealed class NoticeService 
+	{
+		/// <summary>
+		/// Log output writer
+		/// </summary>
+		ILogger<NoticeService> Logger;
+
+		/// <summary>
+		/// Notice collection
+		/// </summary>
+		INoticeCollection Notices;
+
+		/// <summary>
+		/// Device service constructor
+		/// </summary>
+		public NoticeService(INoticeCollection Notices, ILogger<NoticeService> Logger)
+		{
+			this.Notices = Notices;
+            this.Logger = Logger;			
+		}
+
+		/// <summary>
+		/// Add a notice
+		/// </summary>
+		/// <param name="Message"></param>
+		/// <param name="UserId"></param>
+		/// <param name="StartTime"></param>
+		/// <param name="FinishTime"></param>
+		/// <returns></returns>
+		public Task<INotice?> AddNoticeAsync(string Message, UserId? UserId, DateTime? StartTime, DateTime? FinishTime)
+		{
+			return this.Notices.AddNoticeAsync(Message, UserId, StartTime, FinishTime);
+		}
+
+		/// <summary>
+		/// Update an existing notice
+		/// </summary>
+		/// <param name="Id"></param>
+		/// <param name="Message"></param>
+		/// <param name="StartTime"></param>
+		/// <param name="FinishTime"></param>
+		/// <returns></returns>
+		public Task<bool> UpdateNoticeAsync(ObjectId Id, string? Message, DateTime? StartTime, DateTime? FinishTime)
+		{
+			return this.Notices.UpdateNoticeAsync(Id, Message, StartTime, FinishTime);
+		}
+
+
+		/// <summary>
+		/// Get a notice by id
+		/// </summary>
+		/// <param name="NoticeId"></param>
+		/// <returns></returns>
+		public Task<INotice?> GetNoticeAsync(ObjectId NoticeId)
+		{
+			return this.Notices.GetNoticeAsync(NoticeId);
+		}
+
+		/// <summary>
+		/// Get all notices
+		/// </summary>
+		/// <returns></returns>
+		public Task<List<INotice>> GetNoticesAsync()
+		{
+			return this.Notices.GetNoticesAsync();
+		}
+
+		/// <summary>
+		/// Remove an existing notice
+		/// </summary>
+		/// <param name="Id"></param>
+		/// <returns></returns>
+		public Task<bool> RemoveNoticeAsync(ObjectId Id)
+		{
+			return this.Notices.RemoveNoticeAsync(Id);
+		}
+
+	}
+}
\ No newline at end of file
diff --git a/Engine/Source/Programs/Horde/Horde.Build/Server/Config.cs b/Engine/Source/Programs/Horde/Horde.Build/Server/Config.cs
--- a/Engine/Source/Programs/Horde/Horde.Build/Server/Config.cs
+++ b/Engine/Source/Programs/Horde/Horde.Build/Server/Config.cs
@@ -28,11 +28,6 @@ namespace HordeServer.Models
 		/// </summary>
 		public List<ProjectConfigRef> Projects { get; set; } = new List<ProjectConfigRef>();
 
-		/// <summary>
-		/// Manually added status messages
-		/// </summary>
-		public List<Notice> Notices { get; set; } = new List<Notice>();
-
 		/// <summary>
 		/// List of scheduled downtime
 		/// </summary>
diff --git a/Engine/Source/Programs/Horde/Horde.Build/Server/ConfigService.cs b/Engine/Source/Programs/Horde/Horde.Build/Server/ConfigService.cs
--- a/Engine/Source/Programs/Horde/Horde.Build/Server/ConfigService.cs
+++ b/Engine/Source/Programs/Horde/Horde.Build/Server/ConfigService.cs
@@ -140,7 +140,6 @@ namespace HordeServer.Services
 				}
 
 				Globals.ConfigRevision = Revision;
-				Globals.Notices = CachedGlobalConfig.Notices;
 				Globals.PerforceClusters = CachedGlobalConfig.PerforceClusters;
 				Globals.ScheduledDowntime = CachedGlobalConfig.Downtime;
 				Globals.MaxConformCount = CachedGlobalConfig.MaxConformCount;
diff --git a/Engine/Source/Programs/Horde/Horde.Build/Startup.cs b/Engine/Source/Programs/Horde/Horde.Build/Startup.cs
--- a/Engine/Source/Programs/Horde/Horde.Build/Startup.cs
+++ b/Engine/Source/Programs/Horde/Horde.Build/Startup.cs
@@ -316,6 +316,7 @@ namespace HordeServer
 			Services.AddSingleton<IUgsMetadataCollection, UgsMetadataCollection>();
 			Services.AddSingleton<IUserCollection, UserCollectionV2>();
 			Services.AddSingleton<IDeviceCollection, DeviceCollection>();
+			Services.AddSingleton<INoticeCollection, NoticeCollection>();
 
 			// Auditing
 			Services.AddSingleton<IAuditLog<AgentId>>(SP => SP.GetRequiredService<IAuditLogFactory<AgentId>>().Create("Agents.Log", "AgentId"));
@@ -388,6 +389,7 @@ namespace HordeServer
 			Services.AddSingleton<UpgradeService>();
 			Services.AddSingleton<DeviceService>();			
 			Services.AddSingleton<JiraService>();
+			Services.AddSingleton<NoticeService>();
 
 			AWSOptions AwsOptions = Configuration.GetAWSOptions();
 			if (Settings.S3CredentialType == "AssumeRole" && Settings.S3AssumeArn != null)
