Horde Storage - Replication log option to view it the newest bucket

diff --git a/Engine/Source/Programs/Horde/HordeStorage/Horde.Storage/Controllers/ReplicationLogController.cs b/Engine/Source/Programs/Horde/HordeStorage/Horde.Storage/Controllers/ReplicationLogController.cs
--- a/Engine/Source/Programs/Horde/HordeStorage/Horde.Storage/Controllers/ReplicationLogController.cs
+++ b/Engine/Source/Programs/Horde/HordeStorage/Horde.Storage/Controllers/ReplicationLogController.cs
@@ -93,7 +93,7 @@ namespace Horde.Storage.Controllers
                 return Forbid();
             }
 
-            if ((lastBucket == null && lastEvent.HasValue) || (lastBucket != null && !lastEvent.HasValue))
+            if (((lastBucket == null && lastEvent.HasValue) || (lastBucket != null && !lastEvent.HasValue)) && lastBucket != "now")
             {
                 return BadRequest(new ProblemDetails
                 {
diff --git a/Engine/Source/Programs/Horde/HordeStorage/Horde.Storage/Implementation/TransactionLog/ScyllaReplicationLog.cs b/Engine/Source/Programs/Horde/HordeStorage/Horde.Storage/Implementation/TransactionLog/ScyllaReplicationLog.cs
--- a/Engine/Source/Programs/Horde/HordeStorage/Horde.Storage/Implementation/TransactionLog/ScyllaReplicationLog.cs
+++ b/Engine/Source/Programs/Horde/HordeStorage/Horde.Storage/Implementation/TransactionLog/ScyllaReplicationLog.cs
@@ -160,7 +160,14 @@ namespace Horde.Storage.Implementation
         {
             using Scope findReplicationBucketScope = Tracer.Instance.StartActive("scylla.find_replication_buckets");
 
+
             DateTime startBucketTime;
+            if (lastBucket == "now")
+            {
+                // for debug purposes we allow you to list the latest bucket
+                lastBucket = DateTime.Now.ToHourlyBucket().ToReplicationBucketIdentifier();
+            }
+
             if (lastBucket != null)
             {
                 long bucket = FromReplicationBucketIdentifier(lastBucket);
