Horde: Add a SharedCredentials setting for S3 client auth

Allows reading credentials from the default ".aws/credentials" file that AWS CLI/SDK uses. Primarily for development purposes.

diff --git a/Engine/Source/Programs/Horde/Horde.Build/ServerSettings.cs b/Engine/Source/Programs/Horde/Horde.Build/ServerSettings.cs
--- a/Engine/Source/Programs/Horde/Horde.Build/ServerSettings.cs
+++ b/Engine/Source/Programs/Horde/Horde.Build/ServerSettings.cs
@@ -338,9 +338,14 @@ namespace HordeServer
 		/// S3 bucket region for logfile storage
 		/// </summary>
 		public string S3BucketRegion { get; set; } = null!;
+		
+		/// <summary>
+		/// AWS profile name. Used when S3CredentialType is set to "SharedCredentials".
+		/// </summary>
+		public string? S3AwsProfile { get; set; }
 
 		/// <summary>
-		/// Arn to assume for s3.  "Basic", "AssumeRole", "AssumeRoleWebIdentity" only
+		/// Arn to assume for s3.  "Basic", "AssumeRole", "AssumeRoleWebIdentity", "SharedCredentials" only
 		/// </summary>
 		public string? S3CredentialType { get; set; }
 
diff --git a/Engine/Source/Programs/Horde/Horde.Build/Startup.cs b/Engine/Source/Programs/Horde/Horde.Build/Startup.cs
--- a/Engine/Source/Programs/Horde/Horde.Build/Startup.cs
+++ b/Engine/Source/Programs/Horde/Horde.Build/Startup.cs
@@ -81,6 +81,8 @@ using EpicGames.Horde.Compute;
 using HordeServer.Compute.Impl;
 using HordeServer.Compute;
 using System.Net.Http.Headers;
+using Amazon.SecurityToken.Model;
+using Horde.Build.Utilities;
 using Serilog.Events;
 using HordeServer.Jobs;
 using Microsoft.AspNetCore.Mvc;
@@ -405,6 +407,12 @@ namespace HordeServer
 				// Using the fallback credentials from the AWS SDK, it will pick up credentials through a number of default mechanisms.
 				AwsOptions.Credentials = FallbackCredentialsFactory.GetCredentials();
 			}
+			else if(Settings.S3CredentialType == "SharedCredentials" && Settings.S3AwsProfile != null)
+			{
+				// This SharedCredentials option is primarily for development purposes.
+				var (AccessKey, SecretAccessKey, SecretToken) = AwsHelper.ReadAwsCredentials(Settings.S3AwsProfile);
+				AwsOptions.Credentials = new Credentials(AccessKey, SecretAccessKey, SecretToken, DateTime.Now + TimeSpan.FromHours(12));
+			}
 			
 			Services.AddSingleton(AwsOptions);
 
diff --git a/Engine/Source/Programs/Horde/Horde.Build/Utilities/Aws.cs b/Engine/Source/Programs/Horde/Horde.Build/Utilities/Aws.cs
new file mode 100644
--- /dev/null
+++ b/Engine/Source/Programs/Horde/Horde.Build/Utilities/Aws.cs
@@ -0,0 +1,49 @@
+// Copyright Epic Games, Inc. All Rights Reserved.
+
+using System;
+using System.IO;
+
+namespace Horde.Build.Utilities
+{
+	/// <summary>
+	/// Helper functions for interacting with Amazon Web Services (AWS) 
+	/// </summary>
+	public static class AwsHelper
+	{
+		/// <summary>
+		/// Reads AWS credentials from 'credentials' file
+		///
+		/// The AWS SDK provided functionality was proven too much work, so this is a simplified version.
+		/// </summary>
+		/// <param name="ProfileName">Name of the AWS profile</param>
+		/// <param name="CredentialsFilePath">Override for the credentials file path, set to null for default path</param>
+		/// <returns>Credentials as a tuple</returns>
+		/// <exception cref="Exception"></exception>
+		public static (string AccessKey, string SecretAccessKey, string SecretToken) ReadAwsCredentials(string ProfileName, string? CredentialsFilePath = null)
+		{
+			CredentialsFilePath ??= Path.Join(Environment.GetFolderPath(Environment.SpecialFolder.UserProfile), ".aws", "credentials");
+			string[] Lines = File.ReadAllLines(CredentialsFilePath);
+		    
+			string ReadValue(string Line, string ExpectedKey)
+			{
+				if (Line.StartsWith(ExpectedKey, StringComparison.Ordinal)) return Line.Split("=")[1].Trim();
+				throw new Exception($"Unable to read key/value on line {Line} for key {ExpectedKey}");
+			}
+
+			for (int i = 0; i < Lines.Length; i++)
+			{
+				string Line = Lines[i];
+				if (Line == $"[{ProfileName}]")
+				{
+					string AccessKey = ReadValue(Lines[i + 1], "aws_access_key_id");
+					string SecretAccessKey = ReadValue(Lines[i + 2], "aws_secret_access_key");
+					string SessionToken = ReadValue(Lines[i + 3], "aws_session_token");
+
+					return (AccessKey, SecretAccessKey, SessionToken);
+				}
+			}
+
+			throw new Exception($"Unable to find profile {ProfileName} in file {CredentialsFilePath}");
+		}
+	}
+}
\ No newline at end of file
