Horde: More logging for new user creation.

#skipci

diff --git a/Engine/Source/Programs/Horde/Horde.Build/Users/Impl/UserCollectionV2.cs b/Engine/Source/Programs/Horde/Horde.Build/Users/Impl/UserCollectionV2.cs
--- a/Engine/Source/Programs/Horde/Horde.Build/Users/Impl/UserCollectionV2.cs
+++ b/Engine/Source/Programs/Horde/Horde.Build/Users/Impl/UserCollectionV2.cs
@@ -257,7 +257,8 @@ namespace HordeServer.Collections.Impl
 				Logger.LogWarning("Potentially invalid login name: {Login} (from {Trace})", Login, Environment.StackTrace);
 			}
 
-			UpdateDefinition<UserDocument> Update = Builders<UserDocument>.Update.SetOnInsert(x => x.Id, UserId.GenerateNewId()).SetOnInsert(x => x.Login, Login).Unset(x => x.Hidden);
+			UserId NewUserId = UserId.GenerateNewId();
+			UpdateDefinition<UserDocument> Update = Builders<UserDocument>.Update.SetOnInsert(x => x.Id, NewUserId).SetOnInsert(x => x.Login, Login).Unset(x => x.Hidden);
 
 			if (Name == null)
 			{
@@ -274,7 +275,14 @@ namespace HordeServer.Collections.Impl
 			}
 
 			string LoginUpper = Login.ToUpperInvariant();
-			return await Users.FindOneAndUpdateAsync<UserDocument>(x => x.LoginUpper == LoginUpper, Update, new FindOneAndUpdateOptions<UserDocument, UserDocument> { IsUpsert = true, ReturnDocument = ReturnDocument.After });
+
+			IUser User = await Users.FindOneAndUpdateAsync<UserDocument>(x => x.LoginUpper == LoginUpper, Update, new FindOneAndUpdateOptions<UserDocument, UserDocument> { IsUpsert = true, ReturnDocument = ReturnDocument.After });
+			if (User.Id == NewUserId)
+			{
+				Logger.LogInformation("Added new user {Name} ({UserId}, {Login}, {Email})", User.Name, User.Id, User.Login, User.Email);
+			}
+
+			return User;
 		}
 
 		/// <inheritdoc/>
