Horde: Update cached list of open issues on every server pod, not just the elected leader.

#skipci

diff --git a/Engine/Source/Programs/Horde/Horde.Build/Issues/Impl/IssueService.cs b/Engine/Source/Programs/Horde/Horde.Build/Issues/Impl/IssueService.cs
--- a/Engine/Source/Programs/Horde/Horde.Build/Issues/Impl/IssueService.cs
+++ b/Engine/Source/Programs/Horde/Horde.Build/Issues/Impl/IssueService.cs
@@ -205,7 +205,7 @@ namespace HordeServer.Services.Impl
 			this.Perforce = Perforce;
 			this.LogFileService = LogFileService;
 			this.Clock = Clock;
-			this.Ticker = Clock.AddSharedTicker<IssueService>(TimeSpan.FromMinutes(1.0), TickLeaderAsync, Logger);
+			this.Ticker = Clock.AddTicker(TimeSpan.FromMinutes(1.0), TickAsync, Logger);
 			this.Logger = Logger;
 
 			// Create all the issue factories
@@ -235,7 +235,7 @@ namespace HordeServer.Services.Impl
 		/// </summary>
 		/// <param name="StoppingToken">Token to indicate that the service should stop</param>
 		/// <returns>Async task</returns>
-		async ValueTask TickLeaderAsync(CancellationToken StoppingToken)
+		async ValueTask TickAsync(CancellationToken StoppingToken)
 		{
 			DateTime UtcNow = Clock.UtcNow;
 
