Horde Storage - Replication log should use utc everywhere - resolves issue with events ending up in the wrong bucket

diff --git a/Engine/Source/Programs/Horde/HordeStorage/Horde.Storage/Implementation/TransactionLog/ScyllaReplicationLog.cs b/Engine/Source/Programs/Horde/HordeStorage/Horde.Storage/Implementation/TransactionLog/ScyllaReplicationLog.cs
--- a/Engine/Source/Programs/Horde/HordeStorage/Horde.Storage/Implementation/TransactionLog/ScyllaReplicationLog.cs
+++ b/Engine/Source/Programs/Horde/HordeStorage/Horde.Storage/Implementation/TransactionLog/ScyllaReplicationLog.cs
@@ -165,7 +165,7 @@ namespace Horde.Storage.Implementation
             if (lastBucket == "now")
             {
                 // for debug purposes we allow you to list the latest bucket
-                lastBucket = DateTime.Now.ToHourlyBucket().ToReplicationBucketIdentifier();
+                lastBucket = DateTime.UtcNow.ToHourlyBucket().ToReplicationBucketIdentifier();
             }
 
             if (lastBucket != null)
@@ -184,10 +184,10 @@ namespace Horde.Storage.Implementation
             }
 
             // ignore any bucket that is older then a cutoff, as that can cause us to end up scanning thru a lot of hours that will never exist (incremental logs are deleted after 7 days)
-            DateTime oldCutoff = DateTime.Now.AddDays(-14);
+            DateTime oldCutoff = DateTime.UtcNow.AddDays(-14);
             // we returned the start bucket earlier so now we start with the next one
             DateTime bucketTime = startBucketTime.AddHours(1.0).ToHourlyBucket();
-            while(bucketTime < DateTime.Now && bucketTime > oldCutoff)
+            while(bucketTime < DateTime.UtcNow && bucketTime > oldCutoff)
             {
                 using Scope _ = Tracer.Instance.StartActive("scylla.determine_replication_bucket_exists");
                 // fetch all the buckets that exists and sort them based on time
@@ -278,7 +278,7 @@ namespace Horde.Storage.Implementation
             Namespace = @namespace;
             Bucket = bucket;
             Key = key;
-            ReplicationBucket = lastTimestamp.ToHourlyBucket().ToFileTimeUtc();
+            ReplicationBucket = lastTimestamp.ToUniversalTime().ToHourlyBucket().ToFileTimeUtc();
             ReplicationId = TimeUuid.NewId(lastTimestamp);
             Type = (int)opType;
             ObjectIdentifier = objectIdentifier != null ? new ScyllaBlobIdentifier(objectIdentifier) : null;
