Horde: Return log events in ascending order, to prevent them being incorrectly ordered on notifications.

diff --git a/Engine/Source/Programs/Horde/Horde.Build/Logs/Impl/LogEventCollection.cs b/Engine/Source/Programs/Horde/Horde.Build/Logs/Impl/LogEventCollection.cs
--- a/Engine/Source/Programs/Horde/Horde.Build/Logs/Impl/LogEventCollection.cs
+++ b/Engine/Source/Programs/Horde/Horde.Build/Logs/Impl/LogEventCollection.cs
@@ -146,7 +146,7 @@ namespace HordeServer.Collections.Impl
 
 			FilterDefinition<LogEventDocument> Filter = Builder.Eq(x => x.Id.LogId, LogId);
 
-			IFindFluent<LogEventDocument, LogEventDocument> Results = LogEvents.Find(Filter);
+			IFindFluent<LogEventDocument, LogEventDocument> Results = LogEvents.Find(Filter).SortBy(x => x.Id);
 			if (Index != null)
 			{
 				Results = Results.Skip(Index.Value);
