Horde Storage - Fixes for the replication log lookup of lastBucket=now was not properly determining the bucket to use in all places.

diff --git a/Engine/Source/Programs/Horde/HordeStorage/Horde.Storage/Implementation/TransactionLog/MemoryReplicationLog.cs b/Engine/Source/Programs/Horde/HordeStorage/Horde.Storage/Implementation/TransactionLog/MemoryReplicationLog.cs
--- a/Engine/Source/Programs/Horde/HordeStorage/Horde.Storage/Implementation/TransactionLog/MemoryReplicationLog.cs
+++ b/Engine/Source/Programs/Horde/HordeStorage/Horde.Storage/Implementation/TransactionLog/MemoryReplicationLog.cs
@@ -182,7 +182,7 @@ namespace Horde.Storage.Implementation
 
         public static string ToReplicationBucketIdentifier(this DateTime timestamp)
         {
-            return $"rep-{timestamp.ToFileTime()}";
+            return $"rep-{timestamp.ToFileTimeUtc()}";
         }
     }
 }
diff --git a/Engine/Source/Programs/Horde/HordeStorage/Horde.Storage/Implementation/TransactionLog/ScyllaReplicationLog.cs b/Engine/Source/Programs/Horde/HordeStorage/Horde.Storage/Implementation/TransactionLog/ScyllaReplicationLog.cs
--- a/Engine/Source/Programs/Horde/HordeStorage/Horde.Storage/Implementation/TransactionLog/ScyllaReplicationLog.cs
+++ b/Engine/Source/Programs/Horde/HordeStorage/Horde.Storage/Implementation/TransactionLog/ScyllaReplicationLog.cs
@@ -67,7 +67,7 @@ namespace Horde.Storage.Implementation
         {
             using Scope _ = Tracer.Instance.StartActive("scylla.insert_add_event");
             Task addNamespaceTask = PotentiallyAddNamespace(ns);
-            DateTime timeBucket = timestamp.GetValueOrDefault(DateTime.Now);
+            DateTime timeBucket = timestamp.GetValueOrDefault(DateTime.UtcNow);
             ScyllaReplicationLogEvent log = new ScyllaReplicationLogEvent(ns.ToString(), bucket.ToString(), key.ToString(), timeBucket, ScyllaReplicationLogEvent.OpType.Added, objectBlob);
             await _mapper.InsertAsync<ScyllaReplicationLogEvent>(log, insertNulls: false,  ttl: (int)_settings.CurrentValue.ReplicationLogTimeToLive.TotalSeconds);
 
@@ -79,7 +79,7 @@ namespace Horde.Storage.Implementation
         {
             using Scope _ = Tracer.Instance.StartActive("scylla.insert_delete_event");
             Task addNamespaceTask = PotentiallyAddNamespace(ns);
-            DateTime timeBucket = timestamp.GetValueOrDefault(DateTime.Now);
+            DateTime timeBucket = timestamp.GetValueOrDefault(DateTime.UtcNow);
             ScyllaReplicationLogEvent log = new ScyllaReplicationLogEvent(ns.ToString(), bucket.ToString(), key.ToString(), timeBucket, ScyllaReplicationLogEvent.OpType.Deleted, null);
             await _mapper.InsertAsync<ScyllaReplicationLogEvent>(log, insertNulls: false,  ttl: (int)_settings.CurrentValue.ReplicationLogTimeToLive.TotalSeconds);
 
@@ -95,6 +95,13 @@ namespace Horde.Storage.Implementation
         public async IAsyncEnumerable<ReplicationLogEvent> Get(NamespaceId ns, string? lastBucket, Guid? lastEvent)
         {
             using Scope getReplicationLogScope = Tracer.Instance.StartActive("scylla.get_replication_log");
+
+            if (lastBucket == "now")
+            {
+                // for debug purposes we allow you to list the latest bucket
+                lastBucket = DateTime.UtcNow.ToHourlyBucket().ToReplicationBucketIdentifier();
+            }
+
             IAsyncEnumerable<long> buckets = FindReplicationBuckets(ns, lastBucket);
 
             // loop thru the buckets starting with the oldest to try and find were lastBucket refers to
@@ -162,11 +169,6 @@ namespace Horde.Storage.Implementation
 
 
             DateTime startBucketTime;
-            if (lastBucket == "now")
-            {
-                // for debug purposes we allow you to list the latest bucket
-                lastBucket = DateTime.UtcNow.ToHourlyBucket().ToReplicationBucketIdentifier();
-            }
 
             if (lastBucket != null)
             {
@@ -179,7 +181,7 @@ namespace Horde.Storage.Implementation
                 using Scope _ = Tracer.Instance.StartActive("scylla.determine_first_replication_bucket");
 
                 // we should have no data older then the ttl to lets just assume that the bucket to start searching from is now - time to live
-                DateTime oldestTimestamp = DateTime.Now.AddSeconds(-1 * _settings.CurrentValue.ReplicationLogTimeToLive.TotalSeconds);
+                DateTime oldestTimestamp = DateTime.UtcNow.AddSeconds(-1 * _settings.CurrentValue.ReplicationLogTimeToLive.TotalSeconds);
                 startBucketTime = oldestTimestamp;
             }
 
