Horde: Add null check for shelved files
#rnx

diff --git a/Engine/Source/Programs/Horde/Horde.Build/Perforce/Impl/PerforceService.cs b/Engine/Source/Programs/Horde/Horde.Build/Perforce/Impl/PerforceService.cs
--- a/Engine/Source/Programs/Horde/Horde.Build/Perforce/Impl/PerforceService.cs
+++ b/Engine/Source/Programs/Horde/Horde.Build/Perforce/Impl/PerforceService.cs
@@ -913,7 +913,7 @@ namespace HordeServer.Services
 				{
 					return CheckShelfResult.NoChange;
 				}
-				if (Change.ShelvedFiles.Count == 0)
+				if (Change.ShelvedFiles == null || Change.ShelvedFiles.Count == 0)
 				{
 					return CheckShelfResult.NoShelvedFiles;
 				}
