Horde Storage -  use constant for internal namespace

diff --git a/Engine/Source/Programs/Horde/HordeStorage/Horde.Storage/Implementation/TransactionLog/ReplicationSnapshotService.cs b/Engine/Source/Programs/Horde/HordeStorage/Horde.Storage/Implementation/TransactionLog/ReplicationSnapshotService.cs
--- a/Engine/Source/Programs/Horde/HordeStorage/Horde.Storage/Implementation/TransactionLog/ReplicationSnapshotService.cs
+++ b/Engine/Source/Programs/Horde/HordeStorage/Horde.Storage/Implementation/TransactionLog/ReplicationSnapshotService.cs
@@ -8,6 +8,7 @@ using System.Threading.Tasks;
 using Dasync.Collections;
 using Horde.Storage.Implementation.TransactionLog;
 using Jupiter;
+using Jupiter.Common;
 using Jupiter.Implementation;
 using Microsoft.Extensions.DependencyInjection;
 using Microsoft.Extensions.Options;
@@ -94,6 +95,6 @@ namespace Horde.Storage.Implementation
         [Required]
         public int SnapshotFrequencySeconds { get; set; } = 3600; // default to once a day
 
-        public NamespaceId SnapshotStorageNamespace { get; set; } = new NamespaceId("jupiter-internal");
+        public NamespaceId SnapshotStorageNamespace { get; set; } = INamespacePolicyResolver.JupiterInternalNamespace;
     }
 }
