Horde Storage - Exposed the new replicator state in the replicator description

diff --git a/Engine/Source/Programs/Horde/HordeStorage/Horde.Storage/Controllers/ReplicationController.cs b/Engine/Source/Programs/Horde/HordeStorage/Horde.Storage/Controllers/ReplicationController.cs
--- a/Engine/Source/Programs/Horde/HordeStorage/Horde.Storage/Controllers/ReplicationController.cs
+++ b/Engine/Source/Programs/Horde/HordeStorage/Horde.Storage/Controllers/ReplicationController.cs
@@ -61,6 +61,9 @@ namespace Horde.Storage.Controllers
                     Offset = replicator.Info.State.ReplicatorOffset ?? 0,
                     Generation = replicator.Info.State.ReplicatingGeneration ?? Guid.Empty,
 
+                    LastEvent = replicator.Info.State.LastEvent ?? Guid.Empty,
+                    LastBucket = replicator.Info.State.LastBucket ?? "",
+
                     LastReplicationRun = replicator.Info.LastRun,
                     CountOfRunningReplications = replicator.Info.CountOfRunningReplications,
                 })
@@ -203,5 +206,7 @@ namespace Horde.Storage.Controllers
         public Guid? Generation { get; set; }
         public DateTime LastReplicationRun { get; set; }
         public int CountOfRunningReplications { get; set; }
+        public string? LastBucket { get; set; }
+        public Guid LastEvent { get; set; }
     }
 }
diff --git a/Engine/Source/Programs/Horde/HordeStorage/Horde.Storage/Implementation/Replication/IReplicator.cs b/Engine/Source/Programs/Horde/HordeStorage/Horde.Storage/Implementation/Replication/IReplicator.cs
--- a/Engine/Source/Programs/Horde/HordeStorage/Horde.Storage/Implementation/Replication/IReplicator.cs
+++ b/Engine/Source/Programs/Horde/HordeStorage/Horde.Storage/Implementation/Replication/IReplicator.cs
@@ -12,6 +12,8 @@ namespace Horde.Storage.Implementation
         {
             public long? ReplicatorOffset { get; set; }
             public Guid? ReplicatingGeneration { get; set; }
+            public string? LastBucket { get; set; }
+            public Guid? LastEvent { get; set; }
         }
 
         public class ReplicatorInfo
diff --git a/Engine/Source/Programs/Horde/HordeStorage/Horde.Storage/Implementation/Replication/RefsReplicator.cs b/Engine/Source/Programs/Horde/HordeStorage/Horde.Storage/Implementation/Replication/RefsReplicator.cs
--- a/Engine/Source/Programs/Horde/HordeStorage/Horde.Storage/Implementation/Replication/RefsReplicator.cs
+++ b/Engine/Source/Programs/Horde/HordeStorage/Horde.Storage/Implementation/Replication/RefsReplicator.cs
@@ -539,8 +539,5 @@ namespace Horde.Storage.Implementation
             LastEvent = null;
             LastBucket = null;
         }
-
-        public Guid? LastEvent { get; set; }
-        public string? LastBucket { get; set; }
     }
 }
