Horde Storage - Added more debug logging from the refs replicator to investigate why it looks to be running more instances then it is configured for
Also fixed a potential cause of the issue by making sure we cleanup our list of running requests even if a exception occurs.

diff --git a/Engine/Source/Programs/Horde/HordeStorage/Horde.Storage/Implementation/Replication/RefsReplicator.cs b/Engine/Source/Programs/Horde/HordeStorage/Horde.Storage/Implementation/Replication/RefsReplicator.cs
--- a/Engine/Source/Programs/Horde/HordeStorage/Horde.Storage/Implementation/Replication/RefsReplicator.cs
+++ b/Engine/Source/Programs/Horde/HordeStorage/Horde.Storage/Implementation/Replication/RefsReplicator.cs
@@ -262,6 +262,8 @@ namespace Horde.Storage.Implementation
             int maxParallelism = _replicatorSettings.MaxParallelReplications != -1 ? _replicatorSettings.MaxParallelReplications : 0;
             await snapshot.LiveObjects.ParallelForEachAsync(async snapshotLiveObject =>
             {
+                using Scope scope = Tracer.Instance.StartActive("replicator.replicate_op_snapshot");
+                scope.Span.ResourceName = $"{ns}.{snapshotLiveObject.Blob}";
                 Interlocked.Increment(ref countOfObjectsCurrentlyReplicating);
                 LogReplicationHeartbeat(countOfObjectsCurrentlyReplicating);
 
@@ -301,8 +303,13 @@ namespace Horde.Storage.Implementation
             // if MaxParallelReplications is set to not limit we use the default behavior of ParallelForEachAsync which is to limit based on CPUs 
             int maxParallelism = _replicatorSettings.MaxParallelReplications != -1 ? _replicatorSettings.MaxParallelReplications : 0;
 
+            _logger.Information("{Name} Starting incremental replication maxParallelism: {MaxParallelism} Last event: {LastEvent} Last Bucket {LastBucket}", _name, maxParallelism, lastEvent, lastBucket);
+
             await GetRefEvents(ns, lastBucket, lastEvent, replicationToken).ParallelForEachAsync(async (ReplicationLogEvent @event) =>
             {
+                using Scope scope = Tracer.Instance.StartActive("replicator.replicate_op_incremental");
+                scope.Span.ResourceName = $"[REDACTED_EMAIL]}[REDACTED_EMAIL]}";
+
                 _logger.Information("{Name} New transaction to replicate found. New op: {@Op} . Count of running replications: {CurrentReplications}", _name, @event, replicationTasks.Count);
                 
                 Info.CountOfRunningReplications = replicationTasks.Count;
@@ -331,40 +338,37 @@ namespace Horde.Storage.Implementation
                     // and the 3rd site is the whole reason we even add this to replication log in the first place
                     // await AddToReplicationLog(@event.Namespace, @event.Bucket, @event.Key, @event.Blob);
 
-                    try
+                    bool wasOldest;
+                    lock (replicationTasks)
                     {
-                        bool wasOldest;
-                        lock (replicationTasks)
-                        {
-                            wasOldest = currentOffset <= replicationTasks.First();
-                        }
-
-                        // only update the state when we have replicated everything up that point
-                        if (wasOldest)
-                        {
-                            // we have replicated everything up to a point and can persist this in the state
-                            _refsState.LastBucket = eventBucket;
-                            _refsState.LastEvent = eventId;
-                            SaveState(_stateFile, _refsState);
-
-                            _logger.Information("{Name} replicated all events up to {Time} . Bucket: {EventBucket} Id: {EventId}", _name, @event.Timestamp, eventBucket, eventId);
-                        }
-
-                        Info.LastRun = DateTime.Now;
-                        LogReplicationHeartbeat(replicationTasks.Count);
+                        wasOldest = currentOffset <= replicationTasks.First();
                     }
-                    finally
+
+                    // only update the state when we have replicated everything up that point
+                    if (wasOldest)
                     {
-                        lock (replicationTasks)
-                        {
-                            replicationTasks.Remove(currentOffset);
-                        }
+                        // we have replicated everything up to a point and can persist this in the state
+                        _refsState.LastBucket = eventBucket;
+                        _refsState.LastEvent = eventId;
+                        SaveState(_stateFile, _refsState);
+
+                        _logger.Information("{Name} replicated all events up to {Time} . Bucket: {EventBucket} Id: {EventId}", _name, @event.Timestamp, eventBucket, eventId);
                     }
+
+                    Info.LastRun = DateTime.Now;
+                    LogReplicationHeartbeat(replicationTasks.Count);
                 }
                 catch (BlobNotFoundException)
                 {
                     _logger.Warning("{Name} Failed to replicate {@Op} in {Namespace} because blob was not present in remote store. Skipping.", _name, @event, Info.NamespaceToReplicate);
                 }
+                finally
+                {
+                    lock (replicationTasks)
+                    {
+                        replicationTasks.Remove(currentOffset);
+                    }
+                }
             } , maxParallelism, breakLoopOnException: true, cancellationToken: replicationToken);
 
             return countOfReplicationsDone;
