Horde: Fix issue where last success would select user jobs (which  potentially have modified targets and don't reflect state of template)
#rnx

diff --git a/Engine/Source/Programs/Horde/Horde.Build/Controllers/JobsController.cs b/Engine/Source/Programs/Horde/Horde.Build/Controllers/JobsController.cs
--- a/Engine/Source/Programs/Horde/Horde.Build/Controllers/JobsController.cs
+++ b/Engine/Source/Programs/Horde/Horde.Build/Controllers/JobsController.cs
@@ -218,7 +218,7 @@ namespace HordeServer.Controllers
 		/// <returns></returns>
 		async Task<int> ExecuteChangeQueryAsync(IStream Stream, TemplateRefId TemplateId, string? Target, List<JobStepOutcome> Outcomes)
 		{
-			IList<IJob> Jobs = await JobService.FindJobsAsync(StreamId: Stream.Id, Templates: new[] { TemplateId }, Target: Target, State: new[] { JobStepState.Completed }, Outcome: Outcomes.ToArray(), Count: 1);
+			IList<IJob> Jobs = await JobService.FindJobsAsync(StreamId: Stream.Id, Templates: new[] { TemplateId }, Target: Target, State: new[] { JobStepState.Completed }, Outcome: Outcomes.ToArray(), Count: 1, ExcludeUserJobs: true);
 			if (Jobs.Count == 0)
 			{
 				Logger.LogInformation("Unable to find successful build of {TemplateId} target {Target}. Using latest change instead", TemplateId, Target);
diff --git a/Engine/Source/Programs/Horde/Horde.Build/Jobs/IJobCollection.cs b/Engine/Source/Programs/Horde/Horde.Build/Jobs/IJobCollection.cs
--- a/Engine/Source/Programs/Horde/Horde.Build/Jobs/IJobCollection.cs
+++ b/Engine/Source/Programs/Horde/Horde.Build/Jobs/IJobCollection.cs
@@ -108,8 +108,9 @@ namespace HordeServer.Collections
 		/// <param name="Count">Number of results to return</param>
 		/// <param name="ConsistentRead">If the database read should be made to the replica server</param>
 		/// <param name="IndexHint">Name of index to be specified as a hint to the database query planner</param>
+		/// <param name="ExcludeUserJobs">Whether to exclude user jobs from the find</param>
 		/// <returns>List of jobs matching the given criteria</returns>
-		Task<List<IJob>> FindAsync(JobId[]? JobIds = null, StreamId? StreamId = null, string? Name = null, TemplateRefId[]? Templates = null, int? MinChange = null, int? MaxChange = null, int? PreflightChange = null, bool? PreflightOnly = null, UserId? PreflightStartedByUser = null, UserId? StartedByUser = null, DateTimeOffset? MinCreateTime = null, DateTimeOffset? MaxCreateTime = null, DateTimeOffset? ModifiedBefore = null, DateTimeOffset? ModifiedAfter = null, int? Index = null, int? Count = null, bool ConsistentRead = true, string? IndexHint = null);
+		Task<List<IJob>> FindAsync(JobId[]? JobIds = null, StreamId? StreamId = null, string? Name = null, TemplateRefId[]? Templates = null, int? MinChange = null, int? MaxChange = null, int? PreflightChange = null, bool? PreflightOnly = null, UserId? PreflightStartedByUser = null, UserId? StartedByUser = null, DateTimeOffset? MinCreateTime = null, DateTimeOffset? MaxCreateTime = null, DateTimeOffset? ModifiedBefore = null, DateTimeOffset? ModifiedAfter = null, int? Index = null, int? Count = null, bool ConsistentRead = true, string? IndexHint = null, bool? ExcludeUserJobs = null);
 
 		/// <summary>
 		/// Searches for jobs in a specified stream and templates
diff --git a/Engine/Source/Programs/Horde/Horde.Build/Jobs/Impl/JobCollection.cs b/Engine/Source/Programs/Horde/Horde.Build/Jobs/Impl/JobCollection.cs
--- a/Engine/Source/Programs/Horde/Horde.Build/Jobs/Impl/JobCollection.cs
+++ b/Engine/Source/Programs/Horde/Horde.Build/Jobs/Impl/JobCollection.cs
@@ -516,7 +516,7 @@ namespace HordeServer.Collections.Impl
 		}
 
 		/// <inheritdoc/>
-		public async Task<List<IJob>> FindAsync(JobId[]? JobIds, StreamId? StreamId, string? Name, TemplateRefId[]? Templates, int? MinChange, int? MaxChange, int? PreflightChange, bool? PreflightOnly, UserId ? PreflightStartedByUser, UserId? StartedByUser, DateTimeOffset? MinCreateTime, DateTimeOffset? MaxCreateTime, DateTimeOffset? ModifiedBefore, DateTimeOffset? ModifiedAfter, int? Index, int? Count, bool ConsistentRead, string? IndexHint)
+		public async Task<List<IJob>> FindAsync(JobId[]? JobIds, StreamId? StreamId, string? Name, TemplateRefId[]? Templates, int? MinChange, int? MaxChange, int? PreflightChange, bool? PreflightOnly, UserId ? PreflightStartedByUser, UserId? StartedByUser, DateTimeOffset? MinCreateTime, DateTimeOffset? MaxCreateTime, DateTimeOffset? ModifiedBefore, DateTimeOffset? ModifiedAfter, int? Index, int? Count, bool ConsistentRead, string? IndexHint, bool? ExcludeUserJobs)
 		{
 			FilterDefinitionBuilder<JobDocument> FilterBuilder = Builders<JobDocument>.Filter;
 
@@ -553,13 +553,20 @@ namespace HordeServer.Collections.Impl
 			{
 				Filter &= FilterBuilder.Ne(x => x.PreflightChange, 0);
 			}
-			if (PreflightStartedByUser != null)
+			if (ExcludeUserJobs != null && ExcludeUserJobs.Value)
 			{
-				Filter &= FilterBuilder.Or(FilterBuilder.Eq(x => x.PreflightChange, 0), FilterBuilder.Eq(x => x.StartedByUserId, PreflightStartedByUser));
+				Filter &= FilterBuilder.Eq(x => x.StartedByUserId, null);
 			}
-			if (StartedByUser != null)
+			else
 			{
-				Filter &= FilterBuilder.Eq(x => x.StartedByUserId, StartedByUser);
+				if (PreflightStartedByUser != null)
+				{
+					Filter &= FilterBuilder.Or(FilterBuilder.Eq(x => x.PreflightChange, 0), FilterBuilder.Eq(x => x.StartedByUserId, PreflightStartedByUser));
+				}
+				if (StartedByUser != null)
+				{
+					Filter &= FilterBuilder.Eq(x => x.StartedByUserId, StartedByUser);
+				}
 			}
 			if (MinCreateTime != null)
 			{
diff --git a/Engine/Source/Programs/Horde/Horde.Build/Jobs/Impl/JobService.cs b/Engine/Source/Programs/Horde/Horde.Build/Jobs/Impl/JobService.cs
--- a/Engine/Source/Programs/Horde/Horde.Build/Jobs/Impl/JobService.cs
+++ b/Engine/Source/Programs/Horde/Horde.Build/Jobs/Impl/JobService.cs
@@ -545,8 +545,9 @@ namespace HordeServer.Services
 		/// <param name="Index">Index of the first result to return</param>
 		/// <param name="Count">Number of results to return</param>
 		/// <param name="ConsistentRead">If the database read should be made to the replica server</param>
+		/// <param name="ExcludeUserJobs">Whether to exclude user jobs from the find</param>
 		/// <returns>List of jobs matching the given criteria</returns>
-		public async Task<List<IJob>> FindJobsAsync(JobId[]? JobIds = null, StreamId? StreamId = null, string? Name = null, TemplateRefId[]? Templates = null, int? MinChange = null, int? MaxChange = null, int? PreflightChange = null, bool? PreflightOnly = null, UserId? PreflightStartedByUser = null, UserId? StartedByUser = null, DateTimeOffset ? MinCreateTime = null, DateTimeOffset? MaxCreateTime = null, string? Target = null, JobStepState[]? State = null, JobStepOutcome[]? Outcome = null, DateTimeOffset? ModifiedBefore = null, DateTimeOffset? ModifiedAfter = null, int? Index = null, int? Count = null, bool ConsistentRead = true)
+		public async Task<List<IJob>> FindJobsAsync(JobId[]? JobIds = null, StreamId? StreamId = null, string? Name = null, TemplateRefId[]? Templates = null, int? MinChange = null, int? MaxChange = null, int? PreflightChange = null, bool? PreflightOnly = null, UserId? PreflightStartedByUser = null, UserId? StartedByUser = null, DateTimeOffset ? MinCreateTime = null, DateTimeOffset? MaxCreateTime = null, string? Target = null, JobStepState[]? State = null, JobStepOutcome[]? Outcome = null, DateTimeOffset? ModifiedBefore = null, DateTimeOffset? ModifiedAfter = null, int? Index = null, int? Count = null, bool ConsistentRead = true, bool? ExcludeUserJobs = null)
 		{
 			using IScope Scope = GlobalTracer.Instance.BuildSpan("JobService.FindJobsAsync").StartActive();
 			Scope.Span.SetTag("JobIds", JobIds);
@@ -570,7 +571,7 @@ namespace HordeServer.Services
 			
 			if (Target == null && (State == null || State.Length == 0) && (Outcome == null || Outcome.Length == 0))
 			{
-				return await Jobs.FindAsync(JobIds, StreamId, Name, Templates, MinChange, MaxChange, PreflightChange, PreflightOnly, PreflightStartedByUser, StartedByUser, MinCreateTime, MaxCreateTime, ModifiedBefore, ModifiedAfter, Index, Count, ConsistentRead);
+				return await Jobs.FindAsync(JobIds, StreamId, Name, Templates, MinChange, MaxChange, PreflightChange, PreflightOnly, PreflightStartedByUser, StartedByUser, MinCreateTime, MaxCreateTime, ModifiedBefore, ModifiedAfter, Index, Count, ConsistentRead, null, ExcludeUserJobs);
 			}
 			else
 			{
@@ -580,7 +581,7 @@ namespace HordeServer.Services
 				int MaxCount = (Count ?? 1);
 				while (Results.Count < MaxCount)
 				{
-					List<IJob> ScanJobs = await Jobs.FindAsync(JobIds, StreamId, Name, Templates, MinChange, MaxChange, PreflightChange, PreflightOnly, PreflightStartedByUser, StartedByUser, MinCreateTime, MaxCreateTime, ModifiedBefore, ModifiedAfter, 0, 5, ConsistentRead);
+					List<IJob> ScanJobs = await Jobs.FindAsync(JobIds, StreamId, Name, Templates, MinChange, MaxChange, PreflightChange, PreflightOnly, PreflightStartedByUser, StartedByUser, MinCreateTime, MaxCreateTime, ModifiedBefore, ModifiedAfter, 0, 5, ConsistentRead, null, ExcludeUserJobs);
 					if (ScanJobs.Count == 0)
 					{
 						break;
