Horde: Add some additional guards for running in read-only mode.

#skipci

diff --git a/Engine/Source/Programs/Horde/Horde.Build/Fleet/AutoscaleService.cs b/Engine/Source/Programs/Horde/Horde.Build/Fleet/AutoscaleService.cs
--- a/Engine/Source/Programs/Horde/Horde.Build/Fleet/AutoscaleService.cs
+++ b/Engine/Source/Programs/Horde/Horde.Build/Fleet/AutoscaleService.cs
@@ -105,7 +105,7 @@ namespace HordeServer.Services
 		/// <summary>
 		/// Constructor
 		/// </summary>
-		public AutoscaleService(IAgentCollection AgentCollection, IPoolCollection PoolCollection, ILeaseCollection LeaseCollection, IFleetManager FleetManager, IDogStatsd DogStatsd, IClock Clock, ILogger<AutoscaleService> Logger)
+		public AutoscaleService(IAgentCollection AgentCollection, IPoolCollection PoolCollection, ILeaseCollection LeaseCollection, IFleetManager FleetManager, IDogStatsd DogStatsd, DatabaseService DatabaseService, IClock Clock, ILogger<AutoscaleService> Logger)
 		{
 			this.AgentCollection = AgentCollection;
 			this.PoolCollection = PoolCollection;
@@ -113,7 +113,14 @@ namespace HordeServer.Services
 			this.FleetManager = FleetManager;
 			this.DogStatsd = DogStatsd;
 			this.Clock = Clock;
-			this.Ticker = Clock.AddSharedTicker<AutoscaleService>(TimeSpan.FromMinutes(5.0), TickLeaderAsync, Logger);
+			if (DatabaseService.ReadOnlyMode)
+			{
+				this.Ticker = new NullTicker();
+			}
+			else
+			{
+				this.Ticker = Clock.AddSharedTicker<AutoscaleService>(TimeSpan.FromMinutes(5.0), TickLeaderAsync, Logger);
+			}
 			this.Logger = Logger;
 		}
 
diff --git a/Engine/Source/Programs/Horde/Horde.Build/Perforce/PerforceLoadBalancer.cs b/Engine/Source/Programs/Horde/Horde.Build/Perforce/PerforceLoadBalancer.cs
--- a/Engine/Source/Programs/Horde/Horde.Build/Perforce/PerforceLoadBalancer.cs
+++ b/Engine/Source/Programs/Horde/Horde.Build/Perforce/PerforceLoadBalancer.cs
@@ -152,7 +152,14 @@ namespace HordeServer.Services
 			this.LeaseCollection = LeaseCollection;
 			this.ServerListSingleton = new SingletonDocument<PerforceServerList>(DatabaseService);
 			this.Logger = Logger;
-			this.Ticker = Clock.AddSharedTicker<PerforceLoadBalancer>(TimeSpan.FromMinutes(1.0), TickInternalAsync, Logger);
+			if (DatabaseService.ReadOnlyMode)
+			{
+				this.Ticker = new NullTicker();
+			}
+			else
+			{
+				this.Ticker = Clock.AddSharedTicker<PerforceLoadBalancer>(TimeSpan.FromMinutes(1.0), TickInternalAsync, Logger);
+			}
 		}
 
 		/// <inheritdoc/>
diff --git a/Engine/Source/Programs/Horde/Horde.Build/Server/ConfigService.cs b/Engine/Source/Programs/Horde/Horde.Build/Server/ConfigService.cs
--- a/Engine/Source/Programs/Horde/Horde.Build/Server/ConfigService.cs
+++ b/Engine/Source/Programs/Horde/Horde.Build/Server/ConfigService.cs
@@ -81,7 +81,14 @@ namespace HordeServer.Services
 			this.PoolService = PoolService;
 			this.AgentService = AgentService;
 			this.Settings = Settings;
-			this.Ticker = Clock.AddSharedTicker<ConfigService>(TimeSpan.FromMinutes(1.0), TickLeaderAsync, Logger);
+			if (DatabaseService.ReadOnlyMode)
+			{
+				this.Ticker = new NullTicker();
+			}
+			else
+			{
+				this.Ticker = Clock.AddSharedTicker<ConfigService>(TimeSpan.FromMinutes(1.0), TickLeaderAsync, Logger);
+			}
 			this.Logger = Logger;
 
 			// This will trigger if the local Horde.json user configuration is changed
diff --git a/Engine/Source/Programs/Horde/Horde.Build/Utilities/Clock.cs b/Engine/Source/Programs/Horde/Horde.Build/Utilities/Clock.cs
--- a/Engine/Source/Programs/Horde/Horde.Build/Utilities/Clock.cs
+++ b/Engine/Source/Programs/Horde/Horde.Build/Utilities/Clock.cs
@@ -29,6 +29,21 @@ namespace HordeCommon
 		Task StopAsync();
 	}
 
+	/// <summary>
+	/// Placeholder interface for ITicker
+	/// </summary>
+	public sealed class NullTicker : ITicker
+	{
+		/// <inheritdoc/>
+		public void Dispose() { }
+
+		/// <inheritdoc/>
+		public Task StartAsync() => Task.CompletedTask;
+
+		/// <inheritdoc/>
+		public Task StopAsync() => Task.CompletedTask;
+	}
+
 	/// <summary>
 	/// Interface representing time and scheduling events which is pluggable during testing. In normal use, the Clock implementation below is used. 
 	/// </summary>
