Horde: Disable schedule service in read-only mode.

#skipci

diff --git a/Engine/Source/Programs/Horde/Horde.Build/Projects/ScheduleService.cs b/Engine/Source/Programs/Horde/Horde.Build/Projects/ScheduleService.cs
--- a/Engine/Source/Programs/Horde/Horde.Build/Projects/ScheduleService.cs
+++ b/Engine/Source/Programs/Horde/Horde.Build/Projects/ScheduleService.cs
@@ -154,23 +154,12 @@ namespace HordeServer.Services
 		RedisKey BaseLockKey;
 		RedisKey TickLockKey; // Lock to tick the queue
 		RedisSortedSet<QueueItem> Queue; // Items to tick, ordered by time
-		BackgroundTask BackgroundTask;
+		ITicker Ticker;
 
 		/// <summary>
 		/// Constructor
 		/// </summary>
-		/// <param name="Redis"></param>
-		/// <param name="Graphs">Collection of graph documents</param>
-		/// <param name="DowntimeService">The downtime service</param>
-		/// <param name="Perforce">The Perforce service</param>
-		/// <param name="JobCollection">Job collection singleton</param>
-		/// <param name="JobService">The job service instance</param>
-		/// <param name="StreamService">The stream service instance</param>
-		/// <param name="TemplateCollection">The template service instance</param>
-		/// <param name="Clock"></param>
-		/// <param name="Settings">Settings for the server</param>
-		/// <param name="Logger">Logging instance</param>
-		public ScheduleService(RedisService Redis, IGraphCollection Graphs, IDowntimeService DowntimeService, IPerforceService Perforce, IJobCollection JobCollection, JobService JobService, StreamService StreamService, ITemplateCollection TemplateCollection, IClock Clock, IOptionsMonitor<ServerSettings> Settings, ILogger<ScheduleService> Logger)
+		public ScheduleService(RedisService Redis, IGraphCollection Graphs, IDowntimeService DowntimeService, IPerforceService Perforce, IJobCollection JobCollection, JobService JobService, StreamService StreamService, ITemplateCollection TemplateCollection, DatabaseService DatabaseService, IClock Clock, IOptionsMonitor<ServerSettings> Settings, ILogger<ScheduleService> Logger)
 		{
 			this.Graphs = Graphs;
 			this.DowntimeService = DowntimeService;
@@ -190,46 +179,36 @@ namespace HordeServer.Services
 			this.BaseLockKey = "scheduler/locks";
 			this.TickLockKey = BaseLockKey.Append("/tick");
 			this.Queue = new RedisSortedSet<QueueItem>(Redis.Database, "scheduler/queue");
-			this.BackgroundTask = new BackgroundTask(RunAsync);
+			if (DatabaseService.ReadOnlyMode)
+			{
+				Ticker = new NullTicker();
+			}
+			else
+			{
+				Ticker = Clock.AddTicker(TimeSpan.FromMinutes(1.0), TickAsync, Logger);
+			}
 		}
 
 		/// <inheritdoc/>
 		public async Task StartAsync(CancellationToken CancellationToken)
 		{
 			await UpdateQueueAsync(Clock.UtcNow);
-			BackgroundTask.Start();
+			await Ticker.StartAsync();
 		}
 
 		/// <inheritdoc/>
 		public async Task StopAsync(CancellationToken CancellationToken)
 		{
-			await BackgroundTask.StopAsync();
+			await Ticker.StopAsync();
 		}
 
 		/// <inheritdoc/>
 		public void Dispose()
 		{
-			BackgroundTask?.Dispose();
-		}
-
-		async Task RunAsync(CancellationToken CancellationToken)
-		{
-			await UpdateQueueAsync(Clock.UtcNow);
-
-			for(; ;)
-			{
-				Stopwatch Timer = Stopwatch.StartNew();
-				await TickAsync(CancellationToken);
-
-				TimeSpan DelayTime = TimeSpan.FromMinutes(1.0) - Timer.Elapsed;
-				if(DelayTime > TimeSpan.Zero)
-				{
-					await Task.Delay(DelayTime, CancellationToken);
-				}
-			}
+			Ticker.Dispose();
 		}
 
-		async Task TickAsync(CancellationToken CancellationToken)
+		async ValueTask TickAsync(CancellationToken CancellationToken)
 		{
 			DateTime UtcNow = Clock.UtcNow;
 
