Horde: Additional logging for job scheduling failures.

#skipci

diff --git a/Engine/Source/Programs/Horde/Horde.Build/Jobs/Impl/JobCollection.cs b/Engine/Source/Programs/Horde/Horde.Build/Jobs/Impl/JobCollection.cs
--- a/Engine/Source/Programs/Horde/Horde.Build/Jobs/Impl/JobCollection.cs
+++ b/Engine/Source/Programs/Horde/Horde.Build/Jobs/Impl/JobCollection.cs
@@ -1207,6 +1207,8 @@ namespace HordeServer.Collections.Impl
 		/// <inheritdoc/>
 		public async Task<IJob?> TryCancelLeaseAsync(IJob Job, int BatchIdx)
 		{
+			Logger.LogDebug("Cancelling lease {LeaseId} for agent {AgentId}", Job.Batches[BatchIdx].LeaseId, Job.Batches[BatchIdx].AgentId);
+
 			JobDocument JobDocument = Clone((JobDocument)Job);
 
 			UpdateDefinition<JobDocument> Update = Builders<JobDocument>.Update.Unset(x => x.Batches[BatchIdx].AgentId).Unset(x => x.Batches[BatchIdx].SessionId).Unset(x => x.Batches[BatchIdx].LeaseId);
@@ -1228,6 +1230,7 @@ namespace HordeServer.Collections.Impl
 		{
 			JobDocument JobDocument = Clone((JobDocument)Job);
 			JobStepBatchDocument Batch = JobDocument.Batches[BatchIdx];
+			Logger.LogDebug("Failing batch {JobId}:{BatchId} with error {Error}", Job.Id, Batch.Id, Error);
 
 			UpdateDefinitionBuilder<JobDocument> UpdateBuilder = new UpdateDefinitionBuilder<JobDocument>();
 			List<UpdateDefinition<JobDocument>> Updates = new List<UpdateDefinition<JobDocument>>();
diff --git a/Engine/Source/Programs/Horde/Horde.Build/Jobs/Impl/JobTaskSource.cs b/Engine/Source/Programs/Horde/Horde.Build/Jobs/Impl/JobTaskSource.cs
--- a/Engine/Source/Programs/Horde/Horde.Build/Jobs/Impl/JobTaskSource.cs
+++ b/Engine/Source/Programs/Horde/Horde.Build/Jobs/Impl/JobTaskSource.cs
@@ -50,6 +50,7 @@ namespace HordeServer.Tasks.Impl
 		/// <summary>
 		/// An item in the queue to be executed
 		/// </summary>
+		[DebuggerDisplay("{Job.Id}:{Batch.Id} ({PoolId})")]
 		internal class QueueItem
 		{
 			/// <summary>
@@ -496,7 +497,7 @@ namespace HordeServer.Tasks.Impl
 			}
 		}
 
-		private async Task<IJob?> SkipBatchAsync(IJob Job, SubResourceId BatchId, IGraph Graph, JobStepBatchError Reason)
+	private async Task<IJob?> SkipBatchAsync(IJob Job, SubResourceId BatchId, IGraph Graph, JobStepBatchError Reason)
 		{
 			Logger.LogInformation("Skipping batch {BatchId} for job {JobId} (reason: {Reason})", BatchId, Job.Id, Reason);
 
@@ -716,7 +717,7 @@ namespace HordeServer.Tasks.Impl
 			IJob Job = Item.Job;
 			IJobStepBatch Batch = Item.Batch;
 			IAgent Agent = Waiter.Agent;
-			Logger.LogDebug("Assigning job {JobId}, batch {BatchId} to waiter (agent {AgentID})", Job.Id, Batch.Id, Agent.Id);
+			Logger.LogDebug("Assigning job {JobId}, batch {BatchId} to waiter (agent {AgentId})", Job.Id, Batch.Id, Agent.Id);
 
 			// Generate a new unique id for the lease
 			LeaseId LeaseId = LeaseId.GenerateNewId();
@@ -768,7 +769,7 @@ namespace HordeServer.Tasks.Impl
 			else
 			{
 				// Unable to assign job
-				Logger.LogTrace("Refreshing queue entries for job {JobId}", Job.Id);
+				Logger.LogDebug("Failed to assign job {JobId}, batch {BatchId} to agent {AgentId}. Refreshing queue entries.", Job.Id, Batch.Id, Agent.Id);
 
 				// Get the new copy of the job
 				NewJob = await Jobs.GetAsync(Job.Id);
