Horde: Add step name to job search
#rnx

diff --git a/Engine/Source/Programs/Horde/HordeDashboard/src/backend/Api.ts b/Engine/Source/Programs/Horde/HordeDashboard/src/backend/Api.ts
--- a/Engine/Source/Programs/Horde/HordeDashboard/src/backend/Api.ts
+++ b/Engine/Source/Programs/Horde/HordeDashboard/src/backend/Api.ts
@@ -275,6 +275,7 @@ export type JobStreamQuery = {
 };
 
 export type JobQuery = {
+	id?: string[],
 	filter?: string;
 	streamId?: string;
 	name?: string;
