Horde Storage - Fixed snapshot service frequency to be once a day as intended

diff --git a/Engine/Source/Programs/Horde/HordeStorage/Horde.Storage/Implementation/TransactionLog/ReplicationSnapshotService.cs b/Engine/Source/Programs/Horde/HordeStorage/Horde.Storage/Implementation/TransactionLog/ReplicationSnapshotService.cs
--- a/Engine/Source/Programs/Horde/HordeStorage/Horde.Storage/Implementation/TransactionLog/ReplicationSnapshotService.cs
+++ b/Engine/Source/Programs/Horde/HordeStorage/Horde.Storage/Implementation/TransactionLog/ReplicationSnapshotService.cs
@@ -93,7 +93,7 @@ namespace Horde.Storage.Implementation
         /// The frequency at which to poll for new replication events
         /// </summary>
         [Required]
-        public int SnapshotFrequencySeconds { get; set; } = 3600; // default to once a day
+        public int SnapshotFrequencySeconds { get; set; } = (int)TimeSpan.FromDays(1.0).TotalSeconds; // default to once a day
 
         public NamespaceId SnapshotStorageNamespace { get; set; } = INamespacePolicyResolver.JupiterInternalNamespace;
     }
