Horde Storage - Upgraded to datadog C# tracer 2.1.0

Should fix missing automatic instrumentation from WebClient

diff --git a/Engine/Source/Programs/Horde/HordeStorage/Callisto/Implementation/FileTransactionLog.cs b/Engine/Source/Programs/Horde/HordeStorage/Callisto/Implementation/FileTransactionLog.cs
--- a/Engine/Source/Programs/Horde/HordeStorage/Callisto/Implementation/FileTransactionLog.cs
+++ b/Engine/Source/Programs/Horde/HordeStorage/Callisto/Implementation/FileTransactionLog.cs
@@ -85,8 +85,8 @@ namespace Callisto.Implementation
         private async Task<long> CommitNewEvent(OpRecord eventRecord)
         {
             // datadog instrumentation of this method
-            using Scope scope = Tracer.Instance.StartActive("log.write");
-            Span span = scope.Span;
+            using IScope scope = Tracer.Instance.StartActive("log.write");
+            ISpan span = scope.Span;
             span.ResourceName = $"write {_namespace}";
             span.OperationName = "WRITE";
 
@@ -175,8 +175,8 @@ namespace Callisto.Implementation
         public async Task<TransactionEvents> Get(long index, int count, string? notSeenAtSite = null)
         {
             // datadog instrumentation of this method
-            using Scope scope = Tracer.Instance.StartActive("log.get");
-            Span span = scope.Span;
+            using IScope scope = Tracer.Instance.StartActive("log.get");
+            ISpan span = scope.Span;
             span.ResourceName = $"get {_namespace}";
             span.OperationName = "GET";
 
diff --git a/Engine/Source/Programs/Horde/HordeStorage/Callisto/Implementation/FileTransactionLogIndex.cs b/Engine/Source/Programs/Horde/HordeStorage/Callisto/Implementation/FileTransactionLogIndex.cs
--- a/Engine/Source/Programs/Horde/HordeStorage/Callisto/Implementation/FileTransactionLogIndex.cs
+++ b/Engine/Source/Programs/Horde/HordeStorage/Callisto/Implementation/FileTransactionLogIndex.cs
@@ -80,7 +80,7 @@ namespace Callisto.Implementation
                     return index;
                 }
 
-                using Scope _ = Tracer.Instance.StartActive("update.index");
+                using IScope _ = Tracer.Instance.StartActive("update.index");
                 int newIndex = _nextLocation++;
                 _locationToIdentifier.Add(location, newIndex);
                 Save();
diff --git a/Engine/Source/Programs/Horde/HordeStorage/Horde.Storage/Controllers/CompressedBlobController.cs b/Engine/Source/Programs/Horde/HordeStorage/Horde.Storage/Controllers/CompressedBlobController.cs
--- a/Engine/Source/Programs/Horde/HordeStorage/Horde.Storage/Controllers/CompressedBlobController.cs
+++ b/Engine/Source/Programs/Horde/HordeStorage/Horde.Storage/Controllers/CompressedBlobController.cs
@@ -172,7 +172,7 @@ namespace Horde.Storage.Controllers
             }
 
             // chunked content, combine the chunks into a single stream
-            using Scope _ = Tracer.Instance.StartActive("blob.combine");
+            using IScope _ = Tracer.Instance.StartActive("blob.combine");
             Task<BlobContents>[] tasks = new Task<BlobContents>[chunks.Length];
             for (int i = 0; i < chunks.Length; i++)
             {
@@ -241,7 +241,7 @@ namespace Horde.Storage.Controllers
 
             BlobIdentifier identifierCompressedPayload;
             {
-                using Scope _ = Tracer.Instance.StartActive("web.hash");
+                using IScope _ = Tracer.Instance.StartActive("web.hash");
                 await using Stream hashStream = payload.GetStream();
                 identifierCompressedPayload = await BlobIdentifier.FromStream(hashStream);
             }
diff --git a/Engine/Source/Programs/Horde/HordeStorage/Horde.Storage/Controllers/DDCRefController.cs b/Engine/Source/Programs/Horde/HordeStorage/Horde.Storage/Controllers/DDCRefController.cs
--- a/Engine/Source/Programs/Horde/HordeStorage/Horde.Storage/Controllers/DDCRefController.cs
+++ b/Engine/Source/Programs/Horde/HordeStorage/Horde.Storage/Controllers/DDCRefController.cs
@@ -113,7 +113,7 @@ namespace Horde.Storage.Controllers
         {
             if (!ShouldDoAuth())
             {
-                using (Scope _ = Tracer.Instance.StartActive("authorize"))
+                using (IScope _ = Tracer.Instance.StartActive("authorize"))
                 {
                     AuthorizationResult authorizationResult = await _authorizationService.AuthorizeAsync(User, ns, NamespaceAccessRequirement.Name);
 
@@ -142,7 +142,7 @@ namespace Horde.Storage.Controllers
    
                     if (isRaw)
                     {
-                        using (Scope _ = Tracer.Instance.StartActive("body.write"))
+                        using (IScope _ = Tracer.Instance.StartActive("body.write"))
                         {
                             const int BufferSize = 64 * 1024;
                             var outputStream = Response.Body;
@@ -250,7 +250,7 @@ namespace Horde.Storage.Controllers
                             httpContext.Response.Headers["Content-Length"] = Convert.ToString(blob.Length);
                             httpContext.Response.Headers[CommonHeaders.HashHeaderName] = refRes.ContentHash.ToString();
 
-                            using (Scope _ = Tracer.Instance.StartActive("body.write"))
+                            using (IScope _ = Tracer.Instance.StartActive("body.write"))
                             {
                                 await blob.Stream.CopyToAsync(httpContext.Response.Body);
                             }
@@ -317,12 +317,12 @@ namespace Horde.Storage.Controllers
 
                 using (StreamReader sr = new StreamReader(httpContext.Request.Body))
                 {
-                    using Scope _ = Tracer.Instance.StartActive("body.read");
+                    using IScope _ = Tracer.Instance.StartActive("body.read");
                     requestBody = await sr.ReadToEndAsync();
                 }
 
                 BatchGetOp batch;
-                using (Scope _ = Tracer.Instance.StartActive("json.deserialize"))
+                using (IScope _ = Tracer.Instance.StartActive("json.deserialize"))
                 {
                     batch = JsonConvert.DeserializeObject<BatchGetOp>(requestBody)!;
                 }
diff --git a/Engine/Source/Programs/Horde/HordeStorage/Horde.Storage/Controllers/ReferencesController.cs b/Engine/Source/Programs/Horde/HordeStorage/Horde.Storage/Controllers/ReferencesController.cs
--- a/Engine/Source/Programs/Horde/HordeStorage/Horde.Storage/Controllers/ReferencesController.cs
+++ b/Engine/Source/Programs/Horde/HordeStorage/Horde.Storage/Controllers/ReferencesController.cs
@@ -97,7 +97,7 @@ namespace Horde.Storage.Controllers
         {
             if (ShouldDoAuth())
             {
-                using (Scope _ = Tracer.Instance.StartActive("authorize"))
+                using (IScope _ = Tracer.Instance.StartActive("authorize"))
                 {
                     AuthorizationResult authorizationResult = await _authorizationService.AuthorizeAsync(User, ns, NamespaceAccessRequirement.Name);
 
@@ -126,7 +126,7 @@ namespace Horde.Storage.Controllers
 
                 async Task WriteBody(BlobContents blobContents, string contentType)
                 {
-                    using Scope scope = Tracer.Instance.StartActive("body.write");
+                    using IScope scope = Tracer.Instance.StartActive("body.write");
                     long contentLength = blobContents.Length;
                     scope.Span.SetTag("content-length", contentLength.ToString());
                     const int BufferSize = 64 * 1024;
@@ -175,7 +175,7 @@ namespace Horde.Storage.Controllers
                     {
                         byte[] blobMemory;
                         {
-                            using Scope scope = Tracer.Instance.StartActive("json.readblob");
+                            using IScope scope = Tracer.Instance.StartActive("json.readblob");
                             blobMemory = await blob.Stream.ToByteArray();
                         }
                         ReadOnlyMemory<byte> localMemory = new ReadOnlyMemory<byte>(blobMemory);
@@ -226,7 +226,7 @@ namespace Horde.Storage.Controllers
         {
             if (ShouldDoAuth())
             {
-                using (Scope _ = Tracer.Instance.StartActive("authorize"))
+                using (IScope _ = Tracer.Instance.StartActive("authorize"))
                 {
                     AuthorizationResult authorizationResult = await _authorizationService.AuthorizeAsync(User, ns, NamespaceAccessRequirement.Name);
 
@@ -302,7 +302,7 @@ namespace Horde.Storage.Controllers
         {
             if (ShouldDoAuth())
             {
-                using (Scope _ = Tracer.Instance.StartActive("authorize"))
+                using (IScope _ = Tracer.Instance.StartActive("authorize"))
                 {
                     AuthorizationResult authorizationResult = await _authorizationService.AuthorizeAsync(User, ns, NamespaceAccessRequirement.Name);
 
@@ -402,7 +402,7 @@ namespace Horde.Storage.Controllers
         {
             if (ShouldDoAuth())
             {
-                using (Scope _ = Tracer.Instance.StartActive("authorize"))
+                using (IScope _ = Tracer.Instance.StartActive("authorize"))
                 {
                     AuthorizationResult authorizationResult = await _authorizationService.AuthorizeAsync(User, ns, NamespaceAccessRequirement.Name);
 
diff --git a/Engine/Source/Programs/Horde/HordeStorage/Horde.Storage/Endpoints/DDCEndpoints.cs b/Engine/Source/Programs/Horde/HordeStorage/Horde.Storage/Endpoints/DDCEndpoints.cs
--- a/Engine/Source/Programs/Horde/HordeStorage/Horde.Storage/Endpoints/DDCEndpoints.cs
+++ b/Engine/Source/Programs/Horde/HordeStorage/Horde.Storage/Endpoints/DDCEndpoints.cs
@@ -36,7 +36,7 @@ namespace Horde.Storage
                 return;
             }
             
-            using (Scope _ = Tracer.Instance.StartActive("authorize"))
+            using (IScope _ = Tracer.Instance.StartActive("authorize"))
             {
                 {
                     AuthorizationResult authorizationResult = await _authorizationService.AuthorizeAsync(httpContext.User, ns, NamespaceAccessRequirement.Name);
@@ -74,7 +74,7 @@ namespace Horde.Storage
                 httpContext.Response.Headers["Content-Length"] = Convert.ToString(blob.Length);
                 httpContext.Response.Headers[CommonHeaders.HashHeaderName] = refRes.ContentHash.ToString();
 
-                using (Scope _ = Tracer.Instance.StartActive("body.write"))
+                using (IScope _ = Tracer.Instance.StartActive("body.write"))
                 {
                     await blob.Stream.CopyToAsync(httpContext.Response.Body);
                 }
diff --git a/Engine/Source/Programs/Horde/HordeStorage/Horde.Storage/Implementation/Blob/ConsistencyCheckService.cs b/Engine/Source/Programs/Horde/HordeStorage/Horde.Storage/Implementation/Blob/ConsistencyCheckService.cs
--- a/Engine/Source/Programs/Horde/HordeStorage/Horde.Storage/Implementation/Blob/ConsistencyCheckService.cs
+++ b/Engine/Source/Programs/Horde/HordeStorage/Horde.Storage/Implementation/Blob/ConsistencyCheckService.cs
@@ -55,7 +55,7 @@ namespace Horde.Storage.Implementation
                 return false;
             }
 
-            using Scope scope = Tracer.Instance.StartActive("consistency_check.run");
+            using IScope scope = Tracer.Instance.StartActive("consistency_check.run");
 
             await RunConsistencyCheck();
 
diff --git a/Engine/Source/Programs/Horde/HordeStorage/Horde.Storage/Implementation/Blob/IBlobService.cs b/Engine/Source/Programs/Horde/HordeStorage/Horde.Storage/Implementation/Blob/IBlobService.cs
--- a/Engine/Source/Programs/Horde/HordeStorage/Horde.Storage/Implementation/Blob/IBlobService.cs
+++ b/Engine/Source/Programs/Horde/HordeStorage/Horde.Storage/Implementation/Blob/IBlobService.cs
@@ -89,7 +89,7 @@ public class BlobService : IBlobService
     {
         BlobIdentifier blobHash;
         {
-            using Scope _ = Tracer.Instance.StartActive("web.hash");
+            using IScope _ = Tracer.Instance.StartActive("web.hash");
             blobHash = await BlobIdentifier.FromStream(content);
         }
 
@@ -103,7 +103,7 @@ public class BlobService : IBlobService
 
     public async Task<BlobIdentifier> PutObject(NamespaceId ns, IBufferedPayload payload, BlobIdentifier identifier)
     {
-        using Scope _ = Tracer.Instance.StartActive("put_blob");
+        using IScope _ = Tracer.Instance.StartActive("put_blob");
 
         await using Stream hashStream = payload.GetStream();
         BlobIdentifier id = await VerifyContentMatchesHash(hashStream, identifier);
@@ -119,7 +119,7 @@ public class BlobService : IBlobService
 
     public async Task<BlobIdentifier> PutObject(NamespaceId ns, byte[] payload, BlobIdentifier identifier)
     {
-        using Scope _ = Tracer.Instance.StartActive("put_blob");
+        using IScope _ = Tracer.Instance.StartActive("put_blob");
 
         await using Stream hashStream = new MemoryStream(payload);
         BlobIdentifier id = await VerifyContentMatchesHash(hashStream, identifier);
@@ -135,7 +135,7 @@ public class BlobService : IBlobService
     public async Task<BlobIdentifier> PutObjectKnownHash(NamespaceId ns, IBufferedPayload content,
         BlobIdentifier identifier)
     {
-        using Scope _ = Tracer.Instance.StartActive("put_blob");
+        using IScope _ = Tracer.Instance.StartActive("put_blob");
         Task<BlobIdentifier> putObjectTask = PutObjectToStores(ns, content, identifier);
         Task addToBlobIndexTask = _blobIndex.AddBlobToIndex(ns, identifier);
 
@@ -150,7 +150,7 @@ public class BlobService : IBlobService
         {
             await Parallel.ForEachAsync(_blobStores, async (store, cancellationToken) =>
             {
-                using Scope scope = Tracer.Instance.StartActive("put_blob_to_store");
+                using IScope scope = Tracer.Instance.StartActive("put_blob_to_store");
                 scope.Span.ResourceName = identifier.ToString();
                 scope.Span.SetTag("store", store.ToString());
 
@@ -162,7 +162,7 @@ public class BlobService : IBlobService
         {
             foreach (IBlobStore store in _blobStores)
             {
-                using Scope scope = Tracer.Instance.StartActive("put_blob_to_store");
+                using IScope scope = Tracer.Instance.StartActive("put_blob_to_store");
                 scope.Span.ResourceName = identifier.ToString();
                 scope.Span.SetTag("store", store.ToString());
 
@@ -213,7 +213,7 @@ public class BlobService : IBlobService
         BlobContents? blobContents = null;
         foreach (IBlobStore store in _blobStores)
         {
-            using Scope scope = Tracer.Instance.StartActive("HierarchicalStore.GetObject");
+            using IScope scope = Tracer.Instance.StartActive("HierarchicalStore.GetObject");
             scope.Span.SetTag("BlobStore", store.GetType().Name);
             scope.Span.SetTag("ObjectFound", false.ToString());
             try
@@ -251,7 +251,7 @@ public class BlobService : IBlobService
 
         if (numStoreMisses >= 1)
         {
-            using Scope _ = Tracer.Instance.StartActive("HierarchicalStore.Populate");
+            using IScope _ = Tracer.Instance.StartActive("HierarchicalStore.Populate");
             await using MemoryStream tempStream = new MemoryStream();
             await blobContents.Stream.CopyToAsync(tempStream);
             byte[] data = tempStream.ToArray();
@@ -260,7 +260,7 @@ public class BlobService : IBlobService
             for (int i = 0; i < numStoreMisses; i++)
             {
                 var blobStore = _blobStores[i];
-                using Scope scope = Tracer.Instance.StartActive("HierarchicalStore.PopulateStore");
+                using IScope scope = Tracer.Instance.StartActive("HierarchicalStore.PopulateStore");
                 scope.Span.SetTag("BlobStore", blobStore.GetType().Name);
                 // Populate each store traversed that did not have the content found lower in the hierarchy
                 await blobStore.PutObject(ns, data, blob);
@@ -284,7 +284,7 @@ public class BlobService : IBlobService
         {
             getTasks.Add(Task.Run( async () =>
             {
-                using Scope scope = Tracer.Instance.StartActive("HierarchicalStore.GetObject");
+                using IScope scope = Tracer.Instance.StartActive("HierarchicalStore.GetObject");
                 scope.Span.SetTag("BlobStore", store.GetType().Name);
                 scope.Span.SetTag("ObjectFound", false.ToString());
                 try
@@ -350,7 +350,7 @@ public class BlobService : IBlobService
         // if it wasn't the first store we need to populate the stores before it with the blob
         if (indexOfStoreWithBlob != 0)
         {
-            using Scope _ = Tracer.Instance.StartActive("HierarchicalStore.Populate");
+            using IScope _ = Tracer.Instance.StartActive("HierarchicalStore.Populate");
             // buffer the blob contents as it could be to large for us to buffer into memory before forwarding to the blob stores
             using FilesystemBufferedPayload payload = await FilesystemBufferedPayload.Create(blobContents.Stream);
 
@@ -358,7 +358,7 @@ public class BlobService : IBlobService
             // Populate each store traversed that did not have the content found lower in the hierarchy
             await Parallel.ForEachAsync(_blobStores.GetRange(0, indexOfStoreWithBlob), async (store, token) =>
             {
-                using Scope scope = Tracer.Instance.StartActive("HierarchicalStore.PopulateStore");
+                using IScope scope = Tracer.Instance.StartActive("HierarchicalStore.PopulateStore");
                 scope.Span.SetTag("BlobStore", store.GetType().Name);
 
                 await using Stream s = payload.GetStream();
@@ -388,7 +388,7 @@ public class BlobService : IBlobService
     {
         foreach (IBlobStore store in _blobStores)
         {
-            using Scope scope = Tracer.Instance.StartActive("HierarchicalStore.ObjectExists");
+            using IScope scope = Tracer.Instance.StartActive("HierarchicalStore.ObjectExists");
             scope.Span.SetTag("BlobStore", store.GetType().Name);
             if (await store.Exists(ns, blob))
             {
@@ -410,7 +410,7 @@ public class BlobService : IBlobService
         {
             existsTasks.Add(Task.Run( async () =>
             {
-                using Scope scope = Tracer.Instance.StartActive("HierarchicalStore.ObjectExists");
+                using IScope scope = Tracer.Instance.StartActive("HierarchicalStore.ObjectExists");
                 scope.Span.SetTag("BlobStore", store.GetType().Name);
                 if (await store.Exists(ns, blob))
                 {
@@ -450,7 +450,7 @@ public class BlobService : IBlobService
         {
             try
             {
-                using Scope scope = Tracer.Instance.StartActive("HierarchicalStore.DeleteObject");
+                using IScope scope = Tracer.Instance.StartActive("HierarchicalStore.DeleteObject");
                 scope.Span.SetTag("BlobStore", store.GetType().Name);
                 await store.DeleteObject(ns, blob);
                 deletedAtLeastOnce = true;
@@ -481,7 +481,7 @@ public class BlobService : IBlobService
         bool deletedAtLeastOnce = false;
         foreach (IBlobStore store in _blobStores)
         {
-            using Scope scope = Tracer.Instance.StartActive("HierarchicalStore.DeleteNamespace");
+            using IScope scope = Tracer.Instance.StartActive("HierarchicalStore.DeleteNamespace");
             scope.Span.SetTag("BlobStore", store.GetType().Name);
             try
             {
@@ -543,7 +543,7 @@ public class BlobService : IBlobService
 
     public async Task<BlobContents> GetObjects(NamespaceId ns, BlobIdentifier[] blobs)
     {
-        using Scope _ = Tracer.Instance.StartActive("blob.combine");
+        using IScope _ = Tracer.Instance.StartActive("blob.combine");
         Task<BlobContents>[] tasks = new Task<BlobContents>[blobs.Length];
         for (int i = 0; i < blobs.Length; i++)
         {
diff --git a/Engine/Source/Programs/Horde/HordeStorage/Horde.Storage/Implementation/Blob/ScyllaBlobIndex.cs b/Engine/Source/Programs/Horde/HordeStorage/Horde.Storage/Implementation/Blob/ScyllaBlobIndex.cs
--- a/Engine/Source/Programs/Horde/HordeStorage/Horde.Storage/Implementation/Blob/ScyllaBlobIndex.cs
+++ b/Engine/Source/Programs/Horde/HordeStorage/Horde.Storage/Implementation/Blob/ScyllaBlobIndex.cs
@@ -36,7 +36,7 @@ public class ScyllaBlobIndex : IBlobIndex
 
     public async Task AddBlobToIndex(NamespaceId ns, BlobIdentifier id)
     {
-        using Scope _ = Tracer.Instance.StartActive("scylla.insert_blob_index");
+        using IScope _ = Tracer.Instance.StartActive("scylla.insert_blob_index");
 
         await _mapper.UpdateAsync<ScyllaBlobIndexTable>("SET regions = regions + ? WHERE namespace = ? AND blob_id = ?",
             new string[] { _jupiterSettings.CurrentValue.CurrentSite }, ns.ToString(), new ScyllaBlobIdentifier(id));
@@ -44,7 +44,7 @@ public class ScyllaBlobIndex : IBlobIndex
 
     public async Task<IBlobIndex.BlobInfo?> GetBlobInfo(NamespaceId ns, BlobIdentifier id)
     {
-        using Scope _ = Tracer.Instance.StartActive("scylla.fetch_blob_index");
+        using IScope _ = Tracer.Instance.StartActive("scylla.fetch_blob_index");
 
         ScyllaBlobIndexTable? blobIndex =
             await _mapper.FirstOrDefaultAsync<ScyllaBlobIndexTable>("WHERE namespace = ? AND blob_id = ?",
@@ -64,7 +64,7 @@ public class ScyllaBlobIndex : IBlobIndex
     public async Task<bool> RemoveBlobFromIndex(NamespaceId ns, BlobIdentifier id)
     {
         // TODO: Should this only remove the current region, and the actual row isnt removed until all regions have been removed? Seems overly complicated
-        using Scope _ = Tracer.Instance.StartActive("scylla.remove_from_blob_index");
+        using IScope _ = Tracer.Instance.StartActive("scylla.remove_from_blob_index");
         await _mapper.DeleteAsync<ScyllaBlobIndexTable>("WHERE namespace = ? AND blob_id = ?", ns.ToString(),
             new ScyllaBlobIdentifier(id));
 
@@ -79,7 +79,7 @@ public class ScyllaBlobIndex : IBlobIndex
 
     public Task AddRefToBlobs(NamespaceId ns, BucketId bucket, IoHashKey key, BlobIdentifier[] blobs)
     {
-        using Scope _ = Tracer.Instance.StartActive("scylla.add_ref_blobs");
+        using IScope _ = Tracer.Instance.StartActive("scylla.add_ref_blobs");
 
         string nsAsString = ns.ToString();
         ScyllaObjectReference r = new ScyllaObjectReference(bucket, key);
diff --git a/Engine/Source/Programs/Horde/HordeStorage/Horde.Storage/Implementation/ContentId/CompressedBufferUtils.cs b/Engine/Source/Programs/Horde/HordeStorage/Horde.Storage/Implementation/ContentId/CompressedBufferUtils.cs
--- a/Engine/Source/Programs/Horde/HordeStorage/Horde.Storage/Implementation/ContentId/CompressedBufferUtils.cs
+++ b/Engine/Source/Programs/Horde/HordeStorage/Horde.Storage/Implementation/ContentId/CompressedBufferUtils.cs
@@ -178,7 +178,7 @@ namespace Horde.Storage.Implementation
                 throw new Exception("Did not decompress the full payload");
 
             {
-                using Scope _ = Tracer.Instance.StartActive("cb.hash");
+                using IScope _ = Tracer.Instance.StartActive("cb.hash");
                 using Hasher hasher = Hasher.New();
                 hasher.UpdateWithJoin(decompressedPayload);
                 Hash blake3Hash = hasher.Finalize();
diff --git a/Engine/Source/Programs/Horde/HordeStorage/Horde.Storage/Implementation/ContentId/ScyllaContentIdStore.cs b/Engine/Source/Programs/Horde/HordeStorage/Horde.Storage/Implementation/ContentId/ScyllaContentIdStore.cs
--- a/Engine/Source/Programs/Horde/HordeStorage/Horde.Storage/Implementation/ContentId/ScyllaContentIdStore.cs
+++ b/Engine/Source/Programs/Horde/HordeStorage/Horde.Storage/Implementation/ContentId/ScyllaContentIdStore.cs
@@ -36,7 +36,7 @@ namespace Horde.Storage.Implementation
 
         public async Task<BlobIdentifier[]?> Resolve(NamespaceId ns, BlobIdentifier contentId)
         {
-            using Scope scope = Tracer.Instance.StartActive("ScyllaContentIdStore.ResolveContentId");
+            using IScope scope = Tracer.Instance.StartActive("ScyllaContentIdStore.ResolveContentId");
             scope.Span.ResourceName = contentId.ToString();
 
             Task<bool> blobStoreExistsTask = _blobStore.Exists(ns, contentId);
@@ -51,7 +51,7 @@ namespace Horde.Storage.Implementation
                 BlobIdentifier[] blobs = resolvedContentId.Chunks.Select(b => b.AsBlobIdentifier()).ToArray();
 
                 {
-                    using Scope _ = Tracer.Instance.StartActive("ScyllaContentIdStore.FindMissingBlobs");
+                    using IScope _ = Tracer.Instance.StartActive("ScyllaContentIdStore.FindMissingBlobs");
 
                     BlobIdentifier[] missingBlobs = await _blobStore.FilterOutKnownBlobs(ns, blobs);
                     if (missingBlobs.Length == 0)
diff --git a/Engine/Source/Programs/Horde/HordeStorage/Horde.Storage/Implementation/DDCRefService.cs b/Engine/Source/Programs/Horde/HordeStorage/Horde.Storage/Implementation/DDCRefService.cs
--- a/Engine/Source/Programs/Horde/HordeStorage/Horde.Storage/Implementation/DDCRefService.cs
+++ b/Engine/Source/Programs/Horde/HordeStorage/Horde.Storage/Implementation/DDCRefService.cs
@@ -94,7 +94,7 @@ namespace Horde.Storage.Implementation
 
             string resource = $"{ns}.{bucket}.{key}";
             RefRecord? record;
-            using (Scope scope = Tracer.Instance.StartActive("ref.get"))
+            using (IScope scope = Tracer.Instance.StartActive("ref.get"))
             {
                 scope.Span.ResourceName = resource;
                 record = await _refsStore.Get(ns, bucket, key, flags);
@@ -119,7 +119,7 @@ namespace Horde.Storage.Implementation
 
             if (needsBlob)
             {
-                using Scope scope = Tracer.Instance.StartActive("blob.get");
+                using IScope scope = Tracer.Instance.StartActive("blob.get");
                 scope.Span.ResourceName = resource;
                 scope.Span.SetTag("BlobCount", record.Blobs.Length.ToString());
 
diff --git a/Engine/Source/Programs/Horde/HordeStorage/Horde.Storage/Implementation/GC/BlobCleanupService.cs b/Engine/Source/Programs/Horde/HordeStorage/Horde.Storage/Implementation/GC/BlobCleanupService.cs
--- a/Engine/Source/Programs/Horde/HordeStorage/Horde.Storage/Implementation/GC/BlobCleanupService.cs
+++ b/Engine/Source/Programs/Horde/HordeStorage/Horde.Storage/Implementation/GC/BlobCleanupService.cs
@@ -85,7 +85,7 @@ namespace Horde.Storage.Implementation
             {
                 string type = blobCleanup.GetType().ToString();
                 _logger.Information("Blob cleanup running for {BlobCleanup}", type);
-                using Scope scope = Tracer.Instance.StartActive("gc.blob");
+                using IScope scope = Tracer.Instance.StartActive("gc.blob");
                 scope.Span.SetTag("type", type);
 
                 _logger.Information("Attempting to run Blob Cleanup {BlobCleanup}. ", type);
diff --git a/Engine/Source/Programs/Horde/HordeStorage/Horde.Storage/Implementation/GC/CacheCleanup.cs b/Engine/Source/Programs/Horde/HordeStorage/Horde.Storage/Implementation/GC/CacheCleanup.cs
--- a/Engine/Source/Programs/Horde/HordeStorage/Horde.Storage/Implementation/GC/CacheCleanup.cs
+++ b/Engine/Source/Programs/Horde/HordeStorage/Horde.Storage/Implementation/GC/CacheCleanup.cs
@@ -36,7 +36,7 @@ namespace Horde.Storage.Implementation
 
         public Task<List<OldRecord>> Cleanup(NamespaceId ns, CancellationToken cancellationToken)
         {
-            using Scope scope = Tracer.Instance.StartActive("gc.refs.namespace");
+            using IScope scope = Tracer.Instance.StartActive("gc.refs.namespace");
             scope.Span.ResourceName = ns.ToString();
             if (_settings.CurrentValue.CleanNamespacesV1.Contains(ns.ToString()))
             {
@@ -66,7 +66,7 @@ namespace Horde.Storage.Implementation
                     continue;
 
                 _logger.Information("Attempting to delete object {Namespace} {Bucket} {Name} as it was last updated {LastAccessTime} which is older then {CutoffTime}", ns, bucket, name, lastAccessTime, cutoffTime);
-                using Scope scope = Tracer.Instance.StartActive("refCleanup.delete_record");
+                using IScope scope = Tracer.Instance.StartActive("refCleanup.delete_record");
                 scope.Span.ResourceName = $"{ns}:{bucket}.{name}";
                 // delete the old record from the ref refs
                 Task storeDelete = _referencesStore.Delete(ns, bucket, name);
diff --git a/Engine/Source/Programs/Horde/HordeStorage/Horde.Storage/Implementation/GC/RefCleanupService.cs b/Engine/Source/Programs/Horde/HordeStorage/Horde.Storage/Implementation/GC/RefCleanupService.cs
--- a/Engine/Source/Programs/Horde/HordeStorage/Horde.Storage/Implementation/GC/RefCleanupService.cs
+++ b/Engine/Source/Programs/Horde/HordeStorage/Horde.Storage/Implementation/GC/RefCleanupService.cs
@@ -62,7 +62,7 @@ namespace Horde.Storage.Implementation
                 }
                 await foreach (NamespaceId ns in state.Refs.GetNamespaces().WithCancellation(cancellationToken))
                 {
-                    using Scope scope = Tracer.Instance.StartActive("gc.refs");
+                    using IScope scope = Tracer.Instance.StartActive("gc.refs");
                     scope.Span.ResourceName = ns.ToString();
 
                     _logger.Information("Attempting to run Refs Cleanup of {Namespace}. ", ns);
@@ -80,7 +80,7 @@ namespace Horde.Storage.Implementation
                 List<NamespaceId>? namespaces = await _referencesStore.GetNamespaces().ToListAsync(cancellationToken);
                 await foreach (NamespaceId ns in namespaces)
                 {
-                    using Scope scope = Tracer.Instance.StartActive("gc.refs");
+                    using IScope scope = Tracer.Instance.StartActive("gc.refs");
                     scope.Span.ResourceName = ns.ToString();
 
                     _logger.Information("Attempting to run Refs Cleanup of {Namespace}. ", ns);
diff --git a/Engine/Source/Programs/Horde/HordeStorage/Horde.Storage/Implementation/IReferenceResolver.cs b/Engine/Source/Programs/Horde/HordeStorage/Horde.Storage/Implementation/IReferenceResolver.cs
--- a/Engine/Source/Programs/Horde/HordeStorage/Horde.Storage/Implementation/IReferenceResolver.cs
+++ b/Engine/Source/Programs/Horde/HordeStorage/Horde.Storage/Implementation/IReferenceResolver.cs
@@ -141,7 +141,7 @@ namespace Horde.Storage.Implementation
 
         private async Task<(BlobIdentifier, BlobIdentifier[]?)> ResolveContentId(NamespaceId ns, BlobIdentifier blobIdentifier)
         {
-            using Scope scope = Tracer.Instance.StartActive("ReferenceResolver.ResolveContentId");
+            using IScope scope = Tracer.Instance.StartActive("ReferenceResolver.ResolveContentId");
             scope.Span.ResourceName = blobIdentifier.ToString();
             BlobIdentifier[]? resolvedBlobs = await _contentIdStore.Resolve(ns, blobIdentifier);
             return (blobIdentifier, resolvedBlobs);
diff --git a/Engine/Source/Programs/Horde/HordeStorage/Horde.Storage/Implementation/LastAccess/LastAccessService.cs b/Engine/Source/Programs/Horde/HordeStorage/Horde.Storage/Implementation/LastAccess/LastAccessService.cs
--- a/Engine/Source/Programs/Horde/HordeStorage/Horde.Storage/Implementation/LastAccess/LastAccessService.cs
+++ b/Engine/Source/Programs/Horde/HordeStorage/Horde.Storage/Implementation/LastAccess/LastAccessService.cs
@@ -61,12 +61,12 @@ namespace Horde.Storage.Implementation
 
         internal async Task<List<(RefRecord, DateTime)>> ProcessLastAccessRecords()
         {
-            using Scope _ = Tracer.Instance.StartActive("lastAccess.update");
+            using IScope _ = Tracer.Instance.StartActive("lastAccess.update");
             _logger.Information("Running Last Access Aggregation");
             List<(RefRecord, DateTime)> records = await _lastAccessCacheRecord.GetLastAccessedRecords();
             foreach ((RefRecord record, DateTime lastAccessTime) in records)
             {
-                using Scope scope = Tracer.Instance.StartActive("lastAccess.update.record");
+                using IScope scope = Tracer.Instance.StartActive("lastAccess.update.record");
                 scope.Span.ResourceName = $"{record.Namespace}:{record.Bucket}.{record.RefName}";
                 _logger.Debug("Updating last access time to {LastAccessTime} for {Record}", lastAccessTime, record);
                 await _refsStore.UpdateLastAccessTime(record, lastAccessTime);
diff --git a/Engine/Source/Programs/Horde/HordeStorage/Horde.Storage/Implementation/LastAccess/LastAccessServiceReferences.cs b/Engine/Source/Programs/Horde/HordeStorage/Horde.Storage/Implementation/LastAccess/LastAccessServiceReferences.cs
--- a/Engine/Source/Programs/Horde/HordeStorage/Horde.Storage/Implementation/LastAccess/LastAccessServiceReferences.cs
+++ b/Engine/Source/Programs/Horde/HordeStorage/Horde.Storage/Implementation/LastAccess/LastAccessServiceReferences.cs
@@ -61,12 +61,12 @@ namespace Horde.Storage.Implementation
 
         internal async Task<List<(LastAccessRecord, DateTime)>> ProcessLastAccessRecords()
         {
-            using Scope _ = Tracer.Instance.StartActive("lastAccessRefs.update");
+            using IScope _ = Tracer.Instance.StartActive("lastAccessRefs.update");
             _logger.Information("Running Last Access Aggregation for refs");
             List<(LastAccessRecord, DateTime)> records = await _lastAccessCacheRecord.GetLastAccessedRecords();
             foreach ((LastAccessRecord record, DateTime lastAccessTime) in records)
             {
-                using Scope scope = Tracer.Instance.StartActive("lastAccess.update.record");
+                using IScope scope = Tracer.Instance.StartActive("lastAccess.update.record");
                 scope.Span.ResourceName = $"{record.Namespace}:{record.Bucket}.{record.Key}";
                 _logger.Debug("Updating last access time to {LastAccessTime} for {Record}", lastAccessTime, record);
                 await _referencesStore.UpdateLastAccessTime(record.Namespace, record.Bucket, record.Key, lastAccessTime);
diff --git a/Engine/Source/Programs/Horde/HordeStorage/Horde.Storage/Implementation/LastAccess/LastAccessTracker.cs b/Engine/Source/Programs/Horde/HordeStorage/Horde.Storage/Implementation/LastAccess/LastAccessTracker.cs
--- a/Engine/Source/Programs/Horde/HordeStorage/Horde.Storage/Implementation/LastAccess/LastAccessTracker.cs
+++ b/Engine/Source/Programs/Horde/HordeStorage/Horde.Storage/Implementation/LastAccess/LastAccessTracker.cs
@@ -58,7 +58,7 @@ namespace Horde.Storage.Implementation
         {
             return Task.Run(() =>
             {
-                using Scope _ = Tracer.Instance.StartActive("lastAccessTracker.track");
+                using IScope _ = Tracer.Instance.StartActive("lastAccessTracker.track");
                 try
                 {
                     _rwLock.AcquireReaderLock(-1);
diff --git a/Engine/Source/Programs/Horde/HordeStorage/Horde.Storage/Implementation/ObjectService.cs b/Engine/Source/Programs/Horde/HordeStorage/Horde.Storage/Implementation/ObjectService.cs
--- a/Engine/Source/Programs/Horde/HordeStorage/Horde.Storage/Implementation/ObjectService.cs
+++ b/Engine/Source/Programs/Horde/HordeStorage/Horde.Storage/Implementation/ObjectService.cs
@@ -104,7 +104,7 @@ namespace Horde.Storage.Implementation
             BlobIdentifier[] missingReferences = Array.Empty<BlobIdentifier>();
             if (hasReferences)
             {
-                using Scope _ = Tracer.Instance.StartActive("ObjectService.ResolveReferences");
+                using IScope _ = Tracer.Instance.StartActive("ObjectService.ResolveReferences");
                 try
                 {
                     IAsyncEnumerable<BlobIdentifier> references = _referenceResolver.ResolveReferences(ns, payload);
@@ -148,7 +148,7 @@ namespace Horde.Storage.Implementation
             BlobIdentifier[] missingReferences = Array.Empty<BlobIdentifier>();
             if (hasReferences)
             {
-                using Scope _ = Tracer.Instance.StartActive("ObjectService.ResolveReferences");
+                using IScope _ = Tracer.Instance.StartActive("ObjectService.ResolveReferences");
                 try
                 {
                     IAsyncEnumerable<BlobIdentifier> references = _referenceResolver.ResolveReferences(ns, payload);
diff --git a/Engine/Source/Programs/Horde/HordeStorage/Horde.Storage/Implementation/Objects/ScyllaReferencesStore.cs b/Engine/Source/Programs/Horde/HordeStorage/Horde.Storage/Implementation/Objects/ScyllaReferencesStore.cs
--- a/Engine/Source/Programs/Horde/HordeStorage/Horde.Storage/Implementation/Objects/ScyllaReferencesStore.cs
+++ b/Engine/Source/Programs/Horde/HordeStorage/Horde.Storage/Implementation/Objects/ScyllaReferencesStore.cs
@@ -65,7 +65,7 @@ namespace Horde.Storage.Implementation
 
         public async Task<ObjectRecord> Get(NamespaceId ns, BucketId bucket, IoHashKey name, IReferencesStore.FieldFlags flags)
         {
-            using Scope _ = Tracer.Instance.StartActive("scylla.get");
+            using IScope _ = Tracer.Instance.StartActive("scylla.get");
 
             ScyllaObject? o;
             bool includePayload = (flags & IReferencesStore.FieldFlags.IncludePayload) != 0;
@@ -93,7 +93,7 @@ namespace Horde.Storage.Implementation
 
         public async Task Put(NamespaceId ns, BucketId bucket, IoHashKey name, BlobIdentifier blobHash, byte[] blob, bool isFinalized)
         {
-            using Scope _ = Tracer.Instance.StartActive("scylla.put");
+            using IScope _ = Tracer.Instance.StartActive("scylla.put");
             if (blob.LongLength > _settings.CurrentValue.InlineBlobMaxSize)
             {
                 // do not inline large blobs
@@ -112,14 +112,14 @@ namespace Horde.Storage.Implementation
 
         public async Task Finalize(NamespaceId ns, BucketId bucket, IoHashKey name)
         {
-            using Scope _ = Tracer.Instance.StartActive("scylla.finalize");
+            using IScope _ = Tracer.Instance.StartActive("scylla.finalize");
 
             await _mapper.UpdateAsync<ScyllaObject>("SET is_finalized=true WHERE namespace=? AND bucket=? AND name=?", ns.ToString(), bucket.ToString(), name.ToString());
         }
 
         private async Task CreateTTLRecord(NamespaceId ns, BucketId bucket, IoHashKey name, sbyte partitionIndex)
         {
-            using Scope _ = Tracer.Instance.StartActive("scylla.create_ttl_record");
+            using IScope _ = Tracer.Instance.StartActive("scylla.create_ttl_record");
             DateTime now = DateTime.Now;
             LocalDate lastAccessAt = new LocalDate(now.Year, now.Month, now.Day);
             await _mapper.UpdateAsync<ScyllaObjectLastAccessAt>("SET objects = objects + ? WHERE accessed_at = ? AND namespace = ? AND partition_index = ? ", new string[] {$"{bucket}#{name}"},  lastAccessAt, ns.ToString(), partitionIndex);
@@ -132,14 +132,14 @@ namespace Horde.Storage.Implementation
 
         private async Task RemoveTTLRecord(NamespaceId ns, BucketId bucket, IoHashKey name, DateTime lastAccess, sbyte partitionIndex)
         {
-            using Scope _ = Tracer.Instance.StartActive("scylla.remove_ttl_record");
+            using IScope _ = Tracer.Instance.StartActive("scylla.remove_ttl_record");
             LocalDate lastAccessAt = new LocalDate(lastAccess.Year, lastAccess.Month, lastAccess.Day);
             await _mapper.UpdateAsync<ScyllaObjectLastAccessAt>("SET objects = objects - ? WHERE accessed_at = ? AND namespace = ? AND partition_index = ?", new string[] {$"{bucket}#{name}"},  lastAccessAt, ns.ToString(), partitionIndex);
         }
 
         public async Task UpdateLastAccessTime(NamespaceId ns, BucketId bucket, IoHashKey name, DateTime lastAccessTime)
         {
-            using Scope _ = Tracer.Instance.StartActive("scylla.update_last_access_time");
+            using IScope _ = Tracer.Instance.StartActive("scylla.update_last_access_time");
             // fetch the old record
             /* ObjectRecord objectRecord = await Get(ns, bucket, name);
 
@@ -166,15 +166,15 @@ namespace Horde.Storage.Implementation
 
         /*public async IAsyncEnumerable<ObjectRecord> GetRecords(NamespaceId ns)
         {
-            using Scope _ = Tracer.Instance.StartActive("scylla.get_records");
+            using IScope _ = Tracer.Instance.StartActive("scylla.get_records");
             for (sbyte partitionIndex = sbyte.MinValue; partitionIndex < sbyte.MaxValue; partitionIndex++)
             {
-                using Scope scope = Tracer.Instance.StartActive("scylla.get_records_partition");
+                using IScope scope = Tracer.Instance.StartActive("scylla.get_records_partition");
                 scope.Span.ResourceName = $"partition_{partitionIndex}";
 
                 SortedSet<DateTime> dateBuckets = new SortedSet<DateTime>();
                 {
-                    using Scope dateBucketsScope = Tracer.Instance.StartActive("scylla.get_records.date_buckets");
+                    using IScope dateBucketsScope = Tracer.Instance.StartActive("scylla.get_records.date_buckets");
 
                     RowSet rowSetDateBuckets = await _session.ExecuteAsync(new SimpleStatement("SELECT accessed_at FROM object_last_access WHERE namespace = ? AND partition_index = ? PER PARTITION LIMIT 1 ALLOW FILTERING", ns.ToString(), partitionIndex));
                     foreach (Row row in rowSetDateBuckets)
@@ -220,7 +220,7 @@ namespace Horde.Storage.Implementation
         
         public async IAsyncEnumerable<(BucketId, IoHashKey, DateTime)> GetRecords(NamespaceId ns)
         {
-            using Scope _ = Tracer.Instance.StartActive("scylla.get_records");
+            using IScope _ = Tracer.Instance.StartActive("scylla.get_records");
 
             RowSet rowSet = await _session.ExecuteAsync(_getObjectsStatement.Bind(ns.ToString()));
             foreach (Row row in rowSet)
@@ -241,7 +241,7 @@ namespace Horde.Storage.Implementation
 
         public async IAsyncEnumerable<NamespaceId> GetNamespaces()
         {
-            using Scope _ = Tracer.Instance.StartActive("scylla.get_namespaces");
+            using IScope _ = Tracer.Instance.StartActive("scylla.get_namespaces");
             RowSet rowSet = await _session.ExecuteAsync(_getNamespacesStatement.Bind());
 
             foreach (Row row in rowSet)
@@ -255,7 +255,7 @@ namespace Horde.Storage.Implementation
 
         public async Task<long> Delete(NamespaceId ns, BucketId bucket, IoHashKey key)
         {
-            using Scope _ = Tracer.Instance.StartActive("scylla.delete_record");
+            using IScope _ = Tracer.Instance.StartActive("scylla.delete_record");
             ObjectRecord record;
             try
             {
@@ -279,7 +279,7 @@ namespace Horde.Storage.Implementation
 
         public async Task<long> DropNamespace(NamespaceId ns)
         {
-            using Scope _ = Tracer.Instance.StartActive("scylla.delete_namespace");
+            using IScope _ = Tracer.Instance.StartActive("scylla.delete_namespace");
             RowSet rowSet = await _session.ExecuteAsync(new SimpleStatement("SELECT bucket, name FROM objects WHERE namespace = ? ALLOW FILTERING;", ns.ToString()));
             long deletedCount = 0;
             foreach (Row row in rowSet)
@@ -300,7 +300,7 @@ namespace Horde.Storage.Implementation
 
         public async Task<long> DeleteBucket(NamespaceId ns, BucketId bucket)
         {
-            using Scope _ = Tracer.Instance.StartActive("scylla.delete_bucket");
+            using IScope _ = Tracer.Instance.StartActive("scylla.delete_bucket");
             RowSet rowSet = await _session.ExecuteAsync(new SimpleStatement("SELECT name FROM objects WHERE namespace = ? AND bucket = ? ALLOW FILTERING;", ns.ToString(), bucket.ToString()));
             long deletedCount = 0;
             foreach (Row row in rowSet)
@@ -319,7 +319,7 @@ namespace Horde.Storage.Implementation
 
         private async Task MaybeAddBucket(NamespaceId ns, BucketId bucket)
         {
-            using Scope _ = Tracer.Instance.StartActive("scylla.add_bucket");
+            using IScope _ = Tracer.Instance.StartActive("scylla.add_bucket");
             await _mapper.UpdateAsync<ScyllaBucket>("SET bucket = bucket + ? WHERE namespace = ?", new string[] {bucket.ToString()}, ns.ToString());
         }
     }
diff --git a/Engine/Source/Programs/Horde/HordeStorage/Horde.Storage/Implementation/Refs/DynamoDbRefsStore.cs b/Engine/Source/Programs/Horde/HordeStorage/Horde.Storage/Implementation/Refs/DynamoDbRefsStore.cs
--- a/Engine/Source/Programs/Horde/HordeStorage/Horde.Storage/Implementation/Refs/DynamoDbRefsStore.cs
+++ b/Engine/Source/Programs/Horde/HordeStorage/Horde.Storage/Implementation/Refs/DynamoDbRefsStore.cs
@@ -75,8 +75,8 @@ namespace Horde.Storage.Implementation
         {
             await Initialize();
 
-            using Scope scope = Tracer.Instance.StartActive("dynamo.get");
-            Span span = scope.Span;
+            using IScope scope = Tracer.Instance.StartActive("dynamo.get");
+            ISpan span = scope.Span;
             span.ResourceName = $"{ns}.{bucket}.{key}";
 
             Task<DynamoBaseRefRecord> baseRefRecordTask = Context.LoadAsync<DynamoBaseRefRecord>(DynamoBaseRefRecord.BuildCacheKey(ns, bucket, key));
@@ -111,8 +111,8 @@ namespace Horde.Storage.Implementation
         {
             await Initialize();
 
-            using Scope scope = Tracer.Instance.StartActive("dynamo.add");
-            Span span = scope.Span;
+            using IScope scope = Tracer.Instance.StartActive("dynamo.add");
+            ISpan span = scope.Span;
             span.ResourceName = $"{record.Namespace}.{record.Bucket}.{record.RefName}";
 
             Task namespaceTableTask = _namespaceStore.AddToNamespaceTable(record.Namespace, DynamoNamespaceStore.NamespaceUsage.Cache);
@@ -205,8 +205,8 @@ namespace Horde.Storage.Implementation
         {
             await Initialize();
 
-            using Scope scope = Tracer.Instance.StartActive("dynamo.update_last_access");
-            Span span = scope.Span;
+            using IScope scope = Tracer.Instance.StartActive("dynamo.update_last_access");
+            ISpan span = scope.Span;
             span.ResourceName = $"{record.Namespace}.{record.Bucket}.{record.RefName}";
 
             Table table = Table.LoadTable(Client, MainTableName);
diff --git a/Engine/Source/Programs/Horde/HordeStorage/Horde.Storage/Implementation/Refs/DynamoStore.cs b/Engine/Source/Programs/Horde/HordeStorage/Horde.Storage/Implementation/Refs/DynamoStore.cs
--- a/Engine/Source/Programs/Horde/HordeStorage/Horde.Storage/Implementation/Refs/DynamoStore.cs
+++ b/Engine/Source/Programs/Horde/HordeStorage/Horde.Storage/Implementation/Refs/DynamoStore.cs
@@ -69,8 +69,8 @@ namespace Horde.Storage.Implementation
 
             if (shouldCreateTable)
             {
-                using Scope scope = Tracer.Instance.StartActive("dynamo.create_tables");
-                Span span = scope.Span;
+                using IScope scope = Tracer.Instance.StartActive("dynamo.create_tables");
+                ISpan span = scope.Span;
                 span.ResourceName = $"CREATE {tableName}";
 
                 CreateTableRequest createTableRequest = new CreateTableRequest(tableName, keySchemaElements.ToList(), attributeDefinitions.ToList(), _settings.CurrentValue.UseOndemandCapacityProvisioning ? null : provisionedThroughput)
diff --git a/Engine/Source/Programs/Horde/HordeStorage/Horde.Storage/Implementation/Replication/RefsReplicator.cs b/Engine/Source/Programs/Horde/HordeStorage/Horde.Storage/Implementation/Replication/RefsReplicator.cs
--- a/Engine/Source/Programs/Horde/HordeStorage/Horde.Storage/Implementation/Replication/RefsReplicator.cs
+++ b/Engine/Source/Programs/Horde/HordeStorage/Horde.Storage/Implementation/Replication/RefsReplicator.cs
@@ -116,7 +116,7 @@ namespace Horde.Storage.Implementation
 
             try
             {
-                using Scope scope = Tracer.Instance.StartActive("replicator.run");
+                using IScope scope = Tracer.Instance.StartActive("replicator.run");
                 scope.Span.ResourceName =_name;
 
                 _replicationRunning = true;
@@ -262,7 +262,7 @@ namespace Horde.Storage.Implementation
             int maxParallelism = _replicatorSettings.MaxParallelReplications != -1 ? _replicatorSettings.MaxParallelReplications : 0;
             await snapshot.LiveObjects.ParallelForEachAsync(async snapshotLiveObject =>
             {
-                using Scope scope = Tracer.Instance.StartActive("replicator.replicate_op_snapshot");
+                using IScope scope = Tracer.Instance.StartActive("replicator.replicate_op_snapshot");
                 scope.Span.ResourceName = $"{ns}.{snapshotLiveObject.Blob}";
                 Interlocked.Increment(ref countOfObjectsCurrentlyReplicating);
                 LogReplicationHeartbeat(countOfObjectsCurrentlyReplicating);
@@ -307,7 +307,7 @@ namespace Horde.Storage.Implementation
 
             await GetRefEvents(ns, lastBucket, lastEvent, replicationToken).ParallelForEachAsync(async (ReplicationLogEvent @event) =>
             {
-                using Scope scope = Tracer.Instance.StartActive("replicator.replicate_op_incremental");
+                using IScope scope = Tracer.Instance.StartActive("replicator.replicate_op_incremental");
                 scope.Span.ResourceName = $"[REDACTED_EMAIL]}[REDACTED_EMAIL]}";
 
                 _logger.Information("{Name} New transaction to replicate found. New op: {@Op} . Count of running replications: {CurrentReplications}", _name, @event, replicationTasks.Count);
diff --git a/Engine/Source/Programs/Horde/HordeStorage/Horde.Storage/Implementation/Replication/Replicator.cs b/Engine/Source/Programs/Horde/HordeStorage/Horde.Storage/Implementation/Replication/Replicator.cs
--- a/Engine/Source/Programs/Horde/HordeStorage/Horde.Storage/Implementation/Replication/Replicator.cs
+++ b/Engine/Source/Programs/Horde/HordeStorage/Horde.Storage/Implementation/Replication/Replicator.cs
@@ -44,7 +44,7 @@ namespace Horde.Storage.Implementation
 
         protected override async Task ReplicateOp(IRestClient remoteClient, TransactionEvent op, CancellationToken replicationToken)
         {
-            using Scope scope = Tracer.Instance.StartActive("replicator.replicate_blobs");
+            using IScope scope = Tracer.Instance.StartActive("replicator.replicate_blobs");
             NamespaceId ns = _replicatorSettings.NamespaceToReplicate;
 
             // now we replicate the blobs
@@ -65,7 +65,7 @@ namespace Horde.Storage.Implementation
 
         protected override async Task<long?> ReplicateOpInline(IRestClient remoteClient, TransactionEvent op, CancellationToken replicationToken)
         {
-            using Scope scope = Tracer.Instance.StartActive("replicator.replicate_inline");
+            using IScope scope = Tracer.Instance.StartActive("replicator.replicate_inline");
             //scope.Span.ResourceName = op.transactionId;
 
             NamespaceId ns = _replicatorSettings.NamespaceToReplicate;
@@ -161,7 +161,7 @@ namespace Horde.Storage.Implementation
             long countOfReplicatedEvents = 0L;
             try
             {
-                using Scope scope = Tracer.Instance.StartActive("replicator.run");
+                using IScope scope = Tracer.Instance.StartActive("replicator.run");
                 scope.Span.ResourceName =_name;
 
                 _replicationRunning = true;
diff --git a/Engine/Source/Programs/Horde/HordeStorage/Horde.Storage/Implementation/TransactionLog/ScyllaReplicationLog.cs b/Engine/Source/Programs/Horde/HordeStorage/Horde.Storage/Implementation/TransactionLog/ScyllaReplicationLog.cs
--- a/Engine/Source/Programs/Horde/HordeStorage/Horde.Storage/Implementation/TransactionLog/ScyllaReplicationLog.cs
+++ b/Engine/Source/Programs/Horde/HordeStorage/Horde.Storage/Implementation/TransactionLog/ScyllaReplicationLog.cs
@@ -65,7 +65,7 @@ namespace Horde.Storage.Implementation
 
         public async Task<(string, Guid)> InsertAddEvent(NamespaceId ns, BucketId bucket, IoHashKey key, BlobIdentifier objectBlob, DateTime? timestamp)
         {
-            using Scope _ = Tracer.Instance.StartActive("scylla.insert_add_event");
+            using IScope _ = Tracer.Instance.StartActive("scylla.insert_add_event");
             Task addNamespaceTask = PotentiallyAddNamespace(ns);
             DateTime timeBucket = timestamp.GetValueOrDefault(DateTime.UtcNow);
             ScyllaReplicationLogEvent log = new ScyllaReplicationLogEvent(ns.ToString(), bucket.ToString(), key.ToString(), timeBucket, ScyllaReplicationLogEvent.OpType.Added, objectBlob);
@@ -77,7 +77,7 @@ namespace Horde.Storage.Implementation
 
         public async Task<(string, Guid)> InsertDeleteEvent(NamespaceId ns, BucketId bucket, IoHashKey key, DateTime? timestamp)
         {
-            using Scope _ = Tracer.Instance.StartActive("scylla.insert_delete_event");
+            using IScope _ = Tracer.Instance.StartActive("scylla.insert_delete_event");
             Task addNamespaceTask = PotentiallyAddNamespace(ns);
             DateTime timeBucket = timestamp.GetValueOrDefault(DateTime.UtcNow);
             ScyllaReplicationLogEvent log = new ScyllaReplicationLogEvent(ns.ToString(), bucket.ToString(), key.ToString(), timeBucket, ScyllaReplicationLogEvent.OpType.Deleted, null);
@@ -94,7 +94,7 @@ namespace Horde.Storage.Implementation
 
         public async IAsyncEnumerable<ReplicationLogEvent> Get(NamespaceId ns, string? lastBucket, Guid? lastEvent)
         {
-            using Scope getReplicationLogScope = Tracer.Instance.StartActive("scylla.get_replication_log");
+            using IScope getReplicationLogScope = Tracer.Instance.StartActive("scylla.get_replication_log");
 
             if (lastBucket == "now")
             {
@@ -108,7 +108,7 @@ namespace Horde.Storage.Implementation
             bool bucketFound = false;
             await foreach (long bucketField in buckets)
             {
-                using Scope readReplicationScope = Tracer.Instance.StartActive("scylla.read_replication_bucket");
+                using IScope readReplicationScope = Tracer.Instance.StartActive("scylla.read_replication_bucket");
                 DateTime t = DateTime.FromFileTimeUtc(bucketField);
                 string bucket = t.ToReplicationBucketIdentifier();
                 readReplicationScope.Span.ResourceName = bucket;
@@ -165,7 +165,7 @@ namespace Horde.Storage.Implementation
 
         private async IAsyncEnumerable<long> FindReplicationBuckets(NamespaceId ns, string? lastBucket)
         {
-            using Scope findReplicationBucketScope = Tracer.Instance.StartActive("scylla.find_replication_buckets");
+            using IScope findReplicationBucketScope = Tracer.Instance.StartActive("scylla.find_replication_buckets");
 
 
             DateTime startBucketTime;
@@ -178,7 +178,7 @@ namespace Horde.Storage.Implementation
             }
             else
             {
-                using Scope _ = Tracer.Instance.StartActive("scylla.determine_first_replication_bucket");
+                using IScope _ = Tracer.Instance.StartActive("scylla.determine_first_replication_bucket");
 
                 // we should have no data older then the ttl to lets just assume that the bucket to start searching from is now - time to live
                 DateTime oldestTimestamp = DateTime.UtcNow.AddSeconds(-1 * _settings.CurrentValue.ReplicationLogTimeToLive.TotalSeconds);
@@ -191,7 +191,7 @@ namespace Horde.Storage.Implementation
             DateTime bucketTime = startBucketTime.AddHours(1.0).ToHourlyBucket();
             while(bucketTime < DateTime.UtcNow && bucketTime > oldCutoff)
             {
-                using Scope _ = Tracer.Instance.StartActive("scylla.determine_replication_bucket_exists");
+                using IScope _ = Tracer.Instance.StartActive("scylla.determine_replication_bucket_exists");
                 // fetch all the buckets that exists and sort them based on time
                 IEnumerable<ScyllaReplicationLogEvent> logEvent = await _mapper.FetchAsync<ScyllaReplicationLogEvent>("WHERE namespace = ? AND replication_bucket = ? LIMIT 1", ns.ToString(), bucketTime.ToFileTimeUtc());
                 ScyllaReplicationLogEvent? e = logEvent.FirstOrDefault();
diff --git a/Engine/Source/Programs/Horde/HordeStorage/Jupiter.Common/BaseStartup.cs b/Engine/Source/Programs/Horde/HordeStorage/Jupiter.Common/BaseStartup.cs
--- a/Engine/Source/Programs/Horde/HordeStorage/Jupiter.Common/BaseStartup.cs
+++ b/Engine/Source/Programs/Horde/HordeStorage/Jupiter.Common/BaseStartup.cs
@@ -490,7 +490,7 @@ namespace Jupiter
 
         public async Task InvokeAsync(HttpContext httpContext)
         {
-            using Scope _ = Tracer.Instance.StartActive(_scopeName);
+            using IScope _ = Tracer.Instance.StartActive(_scopeName);
 
             //Move to next delegate/middleware in the pipleline
             await _next.Invoke(httpContext);
diff --git a/Engine/Source/Programs/Horde/HordeStorage/Jupiter.Common/Implementation/BatchWriter.cs b/Engine/Source/Programs/Horde/HordeStorage/Jupiter.Common/Implementation/BatchWriter.cs
--- a/Engine/Source/Programs/Horde/HordeStorage/Jupiter.Common/Implementation/BatchWriter.cs
+++ b/Engine/Source/Programs/Horde/HordeStorage/Jupiter.Common/Implementation/BatchWriter.cs
@@ -35,7 +35,7 @@ namespace Jupiter.Implementation
             foreach ((OpVerb opVerb, BucketId bucket, KeyId key) in ops)
             {
                 string opTypeName = opVerb.ToString();
-                using Scope scope = Tracer.Instance.StartActive($"batch.{opTypeName}");
+                using IScope scope = Tracer.Instance.StartActive($"batch.{opTypeName}");
                 scope.Span.ResourceName = $"{ns}.{bucket}.{key}";
 
                 (ContentHash? contentHash, BlobContents? blob, OpState opState) = await fetchAction(opVerb, ns, bucket, key, new [] {"blob", "contentHash"});
diff --git a/Engine/Source/Programs/Horde/HordeStorage/Jupiter.Common/Jupiter.Common.csproj b/Engine/Source/Programs/Horde/HordeStorage/Jupiter.Common/Jupiter.Common.csproj
--- a/Engine/Source/Programs/Horde/HordeStorage/Jupiter.Common/Jupiter.Common.csproj
+++ b/Engine/Source/Programs/Horde/HordeStorage/Jupiter.Common/Jupiter.Common.csproj
@@ -18,7 +18,7 @@
     <PackageReference Include="AWSSDK.SecretsManager" Version="3.7.1.39" />
     <PackageReference Include="AWSSDK.SecurityToken" Version="3.7.1.4" />
     <PackageReference Include="Blake3" Version="0.2.1" />
-    <PackageReference Include="Datadog.Trace" Version="1.23.0" />
+    <PackageReference Include="Datadog.Trace" Version="2.1.0" />
     <PackageReference Include="IdentityServer4.AccessTokenValidation" Version="3.0.1" />
     <PackageReference Include="JsonSubTypes" Version="1.8.0" />
     <PackageReference Include="Log4net.Appender.Serilog" Version="1.0.11" />
diff --git a/Engine/Source/Programs/Horde/HordeStorage/Jupiter.Common/Utils/RequestUtil.cs b/Engine/Source/Programs/Horde/HordeStorage/Jupiter.Common/Utils/RequestUtil.cs
--- a/Engine/Source/Programs/Horde/HordeStorage/Jupiter.Common/Utils/RequestUtil.cs
+++ b/Engine/Source/Programs/Horde/HordeStorage/Jupiter.Common/Utils/RequestUtil.cs
@@ -19,7 +19,7 @@ namespace Jupiter
                 throw new Exception("Expected content-length on all raw body requests");
             }
 
-            using Scope scope = Tracer.Instance.StartActive("readbody");
+            using IScope scope = Tracer.Instance.StartActive("readbody");
             scope.Span.SetTag("Content-Length", contentLength.ToString());
             await using MemoryStream ms = new MemoryStream((int)contentLength);
             DateTime readStart = DateTime.Now;
