Horde Storage - Removed usage of dd trace middlware to trace out timing for different middlewares. Just adds lots of noise to our APM logs and for things like auth taking time we generally get a scope anyway. The only thing that is important is the top level scope for the request.

diff --git a/Engine/Source/Programs/Horde/HordeStorage/Jupiter.Common/BaseStartup.cs b/Engine/Source/Programs/Horde/HordeStorage/Jupiter.Common/BaseStartup.cs
--- a/Engine/Source/Programs/Horde/HordeStorage/Jupiter.Common/BaseStartup.cs
+++ b/Engine/Source/Programs/Horde/HordeStorage/Jupiter.Common/BaseStartup.cs
@@ -270,12 +270,12 @@ namespace Jupiter
         private void ConfigureMiddlewares(bool isHighPerf, JupiterSettings jupiterSettings, IApplicationBuilder app, IWebHostEnvironment env)
         {
             // enable use of forwarding headers as we expect a reverse proxy to be running in front of us
-            app.UseMiddleware<DatadogTraceMiddleware>("ForwardedHeaders");
+            //app.UseMiddleware<DatadogTraceMiddleware>("ForwardedHeaders");
             app.UseForwardedHeaders();
 
             if (!isHighPerf && jupiterSettings.UseRequestLogging)
             {
-                app.UseMiddleware<DatadogTraceMiddleware>("RequestLogging");
+                //app.UseMiddleware<DatadogTraceMiddleware>("RequestLogging");
                 app.UseSerilogRequestLogging();
             }
 
@@ -290,18 +290,18 @@ namespace Jupiter
                 app.UseExceptionHandler("/error");
             }
                 
-            app.UseMiddleware<DatadogTraceMiddleware>("Routing");
+            //app.UseMiddleware<DatadogTraceMiddleware>("Routing");
             app.UseRouting();
 
             if (!isHighPerf)
             {
-                app.UseMiddleware<DatadogTraceMiddleware>("Authentication");
+                //app.UseMiddleware<DatadogTraceMiddleware>("Authentication");
                 app.UseAuthentication();
-                app.UseMiddleware<DatadogTraceMiddleware>("Authorization");
+                //app.UseMiddleware<DatadogTraceMiddleware>("Authorization");
                 app.UseAuthorization();
             }
             
-            app.UseMiddleware<DatadogTraceMiddleware>("Endpoints");
+            //app.UseMiddleware<DatadogTraceMiddleware>("Endpoints");
             app.UseEndpoints(endpoints =>
             {
                 bool PassAllChecks(HealthCheckRegistration check) => true;
@@ -325,7 +325,7 @@ namespace Jupiter
 
             if (!isHighPerf && jupiterSettings.HostSwaggerDocumentation)
             {
-                app.UseMiddleware<DatadogTraceMiddleware>("Swagger");
+                //app.UseMiddleware<DatadogTraceMiddleware>("Swagger");
                 app.UseSwagger();
                 app.UseReDoc(options => { options.SpecUrl = "/swagger/v1/swagger.json"; });
             }
