Horde Storage - Improve datadog scopes for consistency service runs

diff --git a/Engine/Source/Programs/Horde/HordeStorage/Horde.Storage/Implementation/Blob/ConsistencyCheckService.cs b/Engine/Source/Programs/Horde/HordeStorage/Horde.Storage/Implementation/Blob/ConsistencyCheckService.cs
--- a/Engine/Source/Programs/Horde/HordeStorage/Horde.Storage/Implementation/Blob/ConsistencyCheckService.cs
+++ b/Engine/Source/Programs/Horde/HordeStorage/Horde.Storage/Implementation/Blob/ConsistencyCheckService.cs
@@ -55,7 +55,7 @@ namespace Horde.Storage.Implementation
                 return false;
             }
 
-            using IScope scope = Tracer.Instance.StartActive("consistency_check.run");
+            using IScope scope = Tracer.Instance.StartActive("consistency_check.poll");
 
             await RunConsistencyCheck();
 
@@ -85,6 +85,9 @@ namespace Horde.Storage.Implementation
             // technically this does not need to be run per namespace but per s3 bucket
             await foreach (NamespaceId ns in _refsStore.GetNamespaces())
             {
+                using IScope scope = Tracer.Instance.StartActive("consistency_check.run");
+                scope.Span.ResourceName = ns.ToString();
+
                 await foreach ((BlobIdentifier blob, DateTime lastModified) in s3Store.ListObjects(ns))
                 {
                     if (countOfBlobsChecked % 100 == 0)
