Horde: Fix another exception when running in read-only mode.

#skipci

diff --git a/Engine/Source/Programs/Horde/Horde.Build/Projects/ScheduleService.cs b/Engine/Source/Programs/Horde/Horde.Build/Projects/ScheduleService.cs
--- a/Engine/Source/Programs/Horde/Horde.Build/Projects/ScheduleService.cs
+++ b/Engine/Source/Programs/Horde/Horde.Build/Projects/ScheduleService.cs
@@ -192,7 +192,6 @@ namespace HordeServer.Services
 		/// <inheritdoc/>
 		public async Task StartAsync(CancellationToken CancellationToken)
 		{
-			await UpdateQueueAsync(Clock.UtcNow);
 			await Ticker.StartAsync();
 		}
 
