Horde: Disable notifications for pools with no agents. Message is spamming slack and causing us to trigger their rate limiter.

diff --git a/Engine/Source/Programs/Horde/Horde.Build/Notifications/Impl/NotificationService.cs b/Engine/Source/Programs/Horde/Horde.Build/Notifications/Impl/NotificationService.cs
--- a/Engine/Source/Programs/Horde/Horde.Build/Notifications/Impl/NotificationService.cs
+++ b/Engine/Source/Programs/Horde/Horde.Build/Notifications/Impl/NotificationService.cs
@@ -205,7 +205,7 @@ namespace HordeServer.Notifications.Impl
 		{
 			if (Pool.EnableAutoscaling && !PoolHasAgentsOnline)
 			{
-				EnqueueTasks(Sink => Sink.NotifyJobScheduledAsync(Pool, PoolHasAgentsOnline, Job, Graph, BatchId));
+//				EnqueueTasks(Sink => Sink.NotifyJobScheduledAsync(Pool, PoolHasAgentsOnline, Job, Graph, BatchId));
 			}
 		}
 
