Horde: Add agent to version dialog
#rnx
#skipci

diff --git a/Engine/Source/Programs/Horde/HordeDashboard/src/backend/Api.ts b/Engine/Source/Programs/Horde/HordeDashboard/src/backend/Api.ts
--- a/Engine/Source/Programs/Horde/HordeDashboard/src/backend/Api.ts
+++ b/Engine/Source/Programs/Horde/HordeDashboard/src/backend/Api.ts
@@ -3821,6 +3821,12 @@ export type FindIssueResponse = {
 	spans: FindIssueSpanResponse[];
 }
 
+export type GetAgentSoftwareChannelResponse = {
+	name?: string;
+	modifiedBy?: string;
+	modifiedTime: string;
+	version?: string;
+}
 
 /// Jira issue information
 export type GetJiraIssueResponse = {
diff --git a/Engine/Source/Programs/Horde/HordeDashboard/src/backend/Backend.ts b/Engine/Source/Programs/Horde/HordeDashboard/src/backend/Backend.ts
--- a/Engine/Source/Programs/Horde/HordeDashboard/src/backend/Backend.ts
+++ b/Engine/Source/Programs/Horde/HordeDashboard/src/backend/Backend.ts
@@ -3,7 +3,7 @@
 import { configure } from "mobx";
 import { isNumber } from 'util';
 import templateCache from '../backend/TemplateCache';
-import { AgentData, ArtifactData, AuditLogEntry, AuditLogQuery, BatchUpdatePoolRequest, ChangeSummaryData, CreateDeviceRequest, CreateDeviceResponse, CreateJobRequest, CreateJobResponse, CreatePoolRequest, CreateSoftwareResponse, CreateSubscriptionRequest, CreateSubscriptionResponse, DashboardPreference, EventData, FindIssueResponse, GetArtifactZipRequest, GetDevicePlatformResponse, GetDevicePoolResponse, GetDeviceReservationResponse, GetDeviceResponse, GetGraphResponse, GetIssueStreamResponse, GetJiraIssueResponse, GetJobsTabResponse, GetJobStepRefResponse, GetJobStepTraceResponse, GetJobTimingResponse, GetLogEventResponse, GetNotificationResponse, GetPerforceServerStatusResponse, GetServerInfoResponse, GetServerSettingsResponse, GetSoftwareResponse, GetSubscriptionResponse, GetUserResponse, GetUtilizationTelemetryResponse, GlobalConfig, IssueData, IssueQuery, IssueQueryV2, JobData, JobQuery, JobsTabColumnType, JobStreamQuery, LeaseData, LogData, LogLineData, PoolData, ProjectData, ScheduleData, ScheduleQuery, SearchLogFileResponse, ServerUpdateResponse, SessionData, StreamData, TabType, TemplateData, TestData, UpdateAgentRequest, UpdateDeviceRequest, UpdateGlobalConfigRequest, UpdateIssueRequest, UpdateJobRequest, UpdateLeaseRequest, UpdateNotificationsRequest, UpdatePoolRequest, UpdateServerSettingsRequest, UpdateStepRequest, UpdateStepResponse, UpdateUserRequest, UsersQuery } from './Api';
+import { AgentData, ArtifactData, AuditLogEntry, AuditLogQuery, BatchUpdatePoolRequest, ChangeSummaryData, CreateDeviceRequest, CreateDeviceResponse, CreateJobRequest, CreateJobResponse, CreatePoolRequest, CreateSoftwareResponse, CreateSubscriptionRequest, CreateSubscriptionResponse, DashboardPreference, EventData, FindIssueResponse, GetAgentSoftwareChannelResponse, GetArtifactZipRequest, GetDevicePlatformResponse, GetDevicePoolResponse, GetDeviceReservationResponse, GetDeviceResponse, GetGraphResponse, GetIssueStreamResponse, GetJiraIssueResponse, GetJobsTabResponse, GetJobStepRefResponse, GetJobStepTraceResponse, GetJobTimingResponse, GetLogEventResponse, GetNotificationResponse, GetPerforceServerStatusResponse, GetServerInfoResponse, GetServerSettingsResponse, GetSoftwareResponse, GetSubscriptionResponse, GetUserResponse, GetUtilizationTelemetryResponse, GlobalConfig, IssueData, IssueQuery, IssueQueryV2, JobData, JobQuery, JobsTabColumnType, JobStreamQuery, LeaseData, LogData, LogLineData, PoolData, ProjectData, ScheduleData, ScheduleQuery, SearchLogFileResponse, ServerUpdateResponse, SessionData, StreamData, TabType, TemplateData, TestData, UpdateAgentRequest, UpdateDeviceRequest, UpdateGlobalConfigRequest, UpdateIssueRequest, UpdateJobRequest, UpdateLeaseRequest, UpdateNotificationsRequest, UpdatePoolRequest, UpdateServerSettingsRequest, UpdateStepRequest, UpdateStepResponse, UpdateUserRequest, UsersQuery } from './Api';
 import dashboard from './Dashboard';
 import { ChallengeStatus, Fetch } from './Fetch';
 import graphCache, { GraphQuery } from './GraphCache';
@@ -925,7 +925,7 @@ export class Backend {
 
 
     getIssuesV2(queryIn?: IssueQueryV2): Promise<FindIssueResponse[]> {
-        
+
         const query = queryIn ?? {};
 
         return new Promise<FindIssueResponse[]>((resolve, reject) => {
@@ -1014,12 +1014,12 @@ export class Backend {
     }
 
     getJiraIssues(streamId: string, jiraKeys: string[]): Promise<GetJiraIssueResponse[]> {
-    
+
         const params = {
             streamId: streamId,
             jiraKeys: jiraKeys
         };
-        
+
         return new Promise<GetJiraIssueResponse[]>((resolve, reject) => {
             this.backend.get(`/api/v1/jira`, { params: params }).then((value) => {
                 resolve(value.data as GetJiraIssueResponse[]);
@@ -1113,8 +1113,8 @@ export class Backend {
     }
 
     downloadAgentZip() {
-        try {            
-            const url = `${this.serverUrl}/api/v1/agentsoftware/default/zip`;            
+        try {
+            const url = `${this.serverUrl}/api/v1/agentsoftware/default/zip`;
             const link = document.createElement('a');
             link.href = url;
             document.body.appendChild(link);
@@ -1124,6 +1124,17 @@ export class Backend {
         }
     }
 
+    getAgentSoftwareChannel(name: string = "default"): Promise<GetAgentSoftwareChannelResponse> {
+        return new Promise<GetAgentSoftwareChannelResponse>((resolve, reject) => {
+            this.backend.get(`/api/v1/agentsoftware/${name}`).then((value) => {
+                resolve(value.data as GetAgentSoftwareChannelResponse);
+            }).catch(reason => {
+                reject(reason);
+            });
+        });
+    }
+
+
     updateJobStep(jobId: string, batchId: string, stepId: string, request: UpdateStepRequest): Promise<UpdateStepResponse> {
         return new Promise<UpdateStepResponse>((resolve, reject) => {
             this.backend.put(`api/v1/jobs/${jobId}/batches/${batchId}/steps/${stepId}`, request).then((value) => {
@@ -1252,14 +1263,14 @@ export class Backend {
         });
     }
 
-    checkoutDevice(deviceId: string, checkout:boolean): Promise<void> {
+    checkoutDevice(deviceId: string, checkout: boolean): Promise<void> {
         return new Promise<void>((resolve, reject) => {
-            this.backend.put(`/api/v2/devices/${deviceId}/checkout`, {checkout: checkout}).then(() => {
+            this.backend.put(`/api/v2/devices/${deviceId}/checkout`, { checkout: checkout }).then(() => {
                 resolve();
             }).catch(reason => {
                 reject(reason);
             });
-        });        
+        });
     }
 
 
@@ -1340,7 +1351,7 @@ export class Backend {
         });
     }
 
-    updateServerSettings(request: UpdateServerSettingsRequest): Promise< ServerUpdateResponse> {
+    updateServerSettings(request: UpdateServerSettingsRequest): Promise<ServerUpdateResponse> {
         return new Promise<ServerUpdateResponse>((resolve, reject) => {
             this.backend.put(`/api/v1/config/serversettings`, request).then((value) => {
                 resolve(value.data as ServerUpdateResponse);
@@ -1357,7 +1368,7 @@ export class Backend {
             }).catch(reason => {
                 reject(reason);
             });
-        });        
+        });
     }
 
     // updates global configuation
@@ -1381,6 +1392,7 @@ export class Backend {
         });
     }
 
+
     private async update() {
 
         if (this.updateID === "updating") {
@@ -1433,11 +1445,11 @@ export class Backend {
         return token;
 
     }
-        
+
 
     init() {
 
-    
+
         return new Promise<boolean>(async (resolve, reject) => {
 
             this.backend.setBaseUrl(this.serverUrl);
@@ -1450,7 +1462,12 @@ export class Backend {
                 return;
             }
 
-            this.serverInfo = await this.getServerInfo();
+            try {
+                this.serverInfo = await this.getServerInfo();
+                this.agentSoftwareInfo = await this.getAgentSoftwareChannel();
+            } catch (reason) {
+                console.error(`Error getting server/agent info from server, defaults used: ${reason}`);
+            }
 
             await dashboard.update();
 
@@ -1476,7 +1493,8 @@ export class Backend {
     updateID?: any;
     logout: boolean = false;
 
-    serverInfo: GetServerInfoResponse = {serverVersion: "0", osDescription: "Unknown", singleInstance: false};
+    agentSoftwareInfo: GetAgentSoftwareChannelResponse = { modifiedTime: "" };
+    serverInfo: GetServerInfoResponse = { serverVersion: "0", osDescription: "Unknown", singleInstance: false };
 
 
     private backend: Fetch;
