Horde: Fix non-constant message templates.

diff --git a/Engine/Source/Programs/Horde/Horde.Agent/Services/WorkerService.cs b/Engine/Source/Programs/Horde/Horde.Agent/Services/WorkerService.cs
--- a/Engine/Source/Programs/Horde/Horde.Agent/Services/WorkerService.cs
+++ b/Engine/Source/Programs/Horde/Horde.Agent/Services/WorkerService.cs
@@ -1026,7 +1026,7 @@ namespace HordeAgent.Services
 				}
 				else
 				{
-					Logger.LogError(Ex, $"Exception while executing batch: {Ex}");
+					Logger.LogError(Ex, "Exception while executing batch: {Ex}", Ex);
 				}
 			}
 
@@ -1219,7 +1219,7 @@ namespace HordeAgent.Services
 					return (JobStepOutcome.Failure, JobStepState.Aborted);
 				}
 				
-				StepLogger.LogError(Ex, $"Exception while executing step: {Ex}");
+				StepLogger.LogError(Ex, "Exception while executing step: {Ex}", Ex);
 				return (JobStepOutcome.Failure, JobStepState.Completed);
 			}
 		}
