Horde Storage - Added max amount of events replicated in one go of the replicator, this allows us to make sure the replicator has finished running once in a while, which helps with monitoring of it (it a region is very far behind it may take days for it to finish replicating the entire history of another region and its not easy to see its running in the mean time)

diff --git a/Engine/Source/Programs/Horde/HordeStorage/Horde.Storage/Implementation/Replication/RefsReplicator.cs b/Engine/Source/Programs/Horde/HordeStorage/Horde.Storage/Implementation/Replication/RefsReplicator.cs
--- a/Engine/Source/Programs/Horde/HordeStorage/Horde.Storage/Implementation/Replication/RefsReplicator.cs
+++ b/Engine/Source/Programs/Horde/HordeStorage/Horde.Storage/Implementation/Replication/RefsReplicator.cs
@@ -305,8 +305,16 @@ namespace Horde.Storage.Implementation
 
             _logger.Information("{Name} Starting incremental replication maxParallelism: {MaxParallelism} Last event: {LastEvent} Last Bucket {LastBucket}", _name, maxParallelism, lastEvent, lastBucket);
 
+            CancellationTokenSource cancellationTokenSource = new CancellationTokenSource();
+            CancellationTokenSource linkedTokenSource = CancellationTokenSource.CreateLinkedTokenSource(cancellationTokenSource.Token, replicationToken);
             await GetRefEvents(ns, lastBucket, lastEvent, replicationToken).ParallelForEachAsync(async (ReplicationLogEvent @event) =>
             {
+                // if we have done all the replication events we should do in a single run we abort
+                if (countOfReplicationsDone > _replicatorSettings.MaxReplicationsPerRun)
+                {
+                    linkedTokenSource.Cancel();
+                    return;
+                }
                 using IScope scope = Tracer.Instance.StartActive("replicator.replicate_op_incremental");
                 scope.Span.ResourceName = $"[REDACTED_EMAIL]}[REDACTED_EMAIL]}";
 
@@ -369,7 +377,7 @@ namespace Horde.Storage.Implementation
                         replicationTasks.Remove(currentOffset);
                     }
                 }
-            } , maxParallelism, breakLoopOnException: true, cancellationToken: replicationToken);
+            } , maxParallelism, breakLoopOnException: true, cancellationToken: linkedTokenSource.Token);
 
             return countOfReplicationsDone;
         }
diff --git a/Engine/Source/Programs/Horde/HordeStorage/Horde.Storage/Implementation/Replication/ReplicationSettings.cs b/Engine/Source/Programs/Horde/HordeStorage/Horde.Storage/Implementation/Replication/ReplicationSettings.cs
--- a/Engine/Source/Programs/Horde/HordeStorage/Horde.Storage/Implementation/Replication/ReplicationSettings.cs
+++ b/Engine/Source/Programs/Horde/HordeStorage/Horde.Storage/Implementation/Replication/ReplicationSettings.cs
@@ -64,6 +64,11 @@ namespace Horde.Storage
         /// Max number of replications that can run in parallel, set to -1 to disable the limit and go as wide as possible
         /// </summary>
         public int MaxParallelReplications { get; set; } = 64;
+
+        /// <summary>
+        /// The replication will stop its run after this many events have been processed, yielding its time until the next time a replication is triggered.
+        /// </summary>
+        public int MaxReplicationsPerRun { get; set; } = 1000;
     }
 
     public enum ReplicatorVersion
