Horde Storage - Fixed error when attempting to clean our internal namespace which does not need to be cleaned.

diff --git a/Engine/Source/Programs/Horde/HordeStorage/Horde.Storage/Implementation/GC/CacheCleanup.cs b/Engine/Source/Programs/Horde/HordeStorage/Horde.Storage/Implementation/GC/CacheCleanup.cs
--- a/Engine/Source/Programs/Horde/HordeStorage/Horde.Storage/Implementation/GC/CacheCleanup.cs
+++ b/Engine/Source/Programs/Horde/HordeStorage/Horde.Storage/Implementation/GC/CacheCleanup.cs
@@ -5,6 +5,7 @@ using System.Collections.Generic;
 using System.Threading;
 using System.Threading.Tasks;
 using Datadog.Trace;
+using Jupiter.Common;
 using Jupiter.Implementation;
 using Microsoft.Extensions.Options;
 using Serilog;
@@ -38,7 +39,12 @@ namespace Horde.Storage.Implementation
         {
             using IScope scope = Tracer.Instance.StartActive("gc.refs.namespace");
             scope.Span.ResourceName = ns.ToString();
-            if (_settings.CurrentValue.CleanNamespacesV1.Contains(ns.ToString()))
+            if (ns == INamespacePolicyResolver.JupiterInternalNamespace)
+            {
+                // do not apply our cleanup policies to the internal namespace
+                return Task.FromResult(new List<OldRecord>());
+            }
+            else if (_settings.CurrentValue.CleanNamespacesV1.Contains(ns.ToString()))
             {
                 return CleanNamespaceV1(ns, cancellationToken);
             }
