Horde Storage - Timestamp added to snapshot listing

diff --git a/Engine/Source/Programs/Horde/HordeStorage/Horde.Storage/Implementation/TransactionLog/IReplicationLog.cs b/Engine/Source/Programs/Horde/HordeStorage/Horde.Storage/Implementation/TransactionLog/IReplicationLog.cs
--- a/Engine/Source/Programs/Horde/HordeStorage/Horde.Storage/Implementation/TransactionLog/IReplicationLog.cs
+++ b/Engine/Source/Programs/Horde/HordeStorage/Horde.Storage/Implementation/TransactionLog/IReplicationLog.cs
@@ -26,16 +26,18 @@ namespace Horde.Storage.Implementation
     public class SnapshotInfo
     {
         [JsonConstructor]
-        public SnapshotInfo(NamespaceId snapshottedNamespace, NamespaceId blobNamespace, BlobIdentifier snapshotBlob)
+        public SnapshotInfo(NamespaceId snapshottedNamespace, NamespaceId blobNamespace, BlobIdentifier snapshotBlob, DateTime timestamp)
         {
             SnapshottedNamespace = snapshottedNamespace;
             BlobNamespace = blobNamespace;
             SnapshotBlob = snapshotBlob;
+            Timestamp = timestamp;
         }
 
         public NamespaceId SnapshottedNamespace { get; set; }
         public NamespaceId BlobNamespace { get; set; }
         public BlobIdentifier SnapshotBlob { get; set; }
+        public DateTime Timestamp { get; set; }
     }
 
     public class ReplicationLogEvent
diff --git a/Engine/Source/Programs/Horde/HordeStorage/Horde.Storage/Implementation/TransactionLog/ReplicationLogSnapshotBuilder.cs b/Engine/Source/Programs/Horde/HordeStorage/Horde.Storage/Implementation/TransactionLog/ReplicationLogSnapshotBuilder.cs
--- a/Engine/Source/Programs/Horde/HordeStorage/Horde.Storage/Implementation/TransactionLog/ReplicationLogSnapshotBuilder.cs
+++ b/Engine/Source/Programs/Horde/HordeStorage/Horde.Storage/Implementation/TransactionLog/ReplicationLogSnapshotBuilder.cs
@@ -84,7 +84,7 @@ namespace Horde.Storage.Implementation.TransactionLog
                     throw new TaskCanceledException();
 
                 // update the replication log with the new snapshot
-                await _replicationLog.AddSnapshot(new SnapshotInfo(ns, storeInNamespace, blobIdentifier));
+                await _replicationLog.AddSnapshot(new SnapshotInfo(ns, storeInNamespace, blobIdentifier, DateTime.Now));
 
                 return blobIdentifier;
 
diff --git a/Engine/Source/Programs/Horde/HordeStorage/Horde.Storage/Implementation/TransactionLog/ScyllaReplicationLog.cs b/Engine/Source/Programs/Horde/HordeStorage/Horde.Storage/Implementation/TransactionLog/ScyllaReplicationLog.cs
--- a/Engine/Source/Programs/Horde/HordeStorage/Horde.Storage/Implementation/TransactionLog/ScyllaReplicationLog.cs
+++ b/Engine/Source/Programs/Horde/HordeStorage/Horde.Storage/Implementation/TransactionLog/ScyllaReplicationLog.cs
@@ -253,7 +253,7 @@ namespace Horde.Storage.Implementation
 
             foreach (ScyllaSnapshot snapshot in snapshots)
             {
-                yield return new SnapshotInfo(ns, new NamespaceId(snapshot.BlobNamespace), snapshot.BlobSnapshot.AsBlobIdentifier());
+                yield return new SnapshotInfo(ns, new NamespaceId(snapshot.BlobNamespace), snapshot.BlobSnapshot.AsBlobIdentifier(), snapshot.Id.GetDate().DateTime);
             }
         }
     }
diff --git a/Engine/Source/Programs/Horde/HordeStorage/Horde.Storage/Tests/Functional/Objects/RefReplicatorTests.cs b/Engine/Source/Programs/Horde/HordeStorage/Horde.Storage/Tests/Functional/Objects/RefReplicatorTests.cs
--- a/Engine/Source/Programs/Horde/HordeStorage/Horde.Storage/Tests/Functional/Objects/RefReplicatorTests.cs
+++ b/Engine/Source/Programs/Horde/HordeStorage/Horde.Storage/Tests/Functional/Objects/RefReplicatorTests.cs
@@ -200,7 +200,7 @@ namespace Horde.Storage.FunctionalTests.Replication
             BlobIdentifier snapshotBlob = BlobIdentifier.FromBlob(snapshotContent);
 
             Mock<HttpMessageHandler> handler = new Mock<HttpMessageHandler>();
-            string s = JsonConvert.SerializeObject(new ReplicationLogSnapshots(new List<SnapshotInfo>{new SnapshotInfo(TestNamespace, SnapshotNamespace, snapshotBlob)}));
+            string s = JsonConvert.SerializeObject(new ReplicationLogSnapshots(new List<SnapshotInfo>{new SnapshotInfo(TestNamespace, SnapshotNamespace, snapshotBlob, DateTime.Now)}));
             handler.SetupRequest($"http://localhost/api/v1/replication-log/snapshots/{TestNamespace}").ReturnsResponse(s, "application/json");
 
             // after processing a snapshot it will attempt to incrementally replicate from there, which should be empty
@@ -285,7 +285,7 @@ namespace Horde.Storage.FunctionalTests.Replication
             BlobIdentifier snapshotBlob = BlobIdentifier.FromBlob(snapshotContent);
 
             Mock<HttpMessageHandler> handler = new Mock<HttpMessageHandler>();
-            string s = JsonConvert.SerializeObject(new ReplicationLogSnapshots(new List<SnapshotInfo>{new SnapshotInfo(TestNamespace, SnapshotNamespace, snapshotBlob)}));
+            string s = JsonConvert.SerializeObject(new ReplicationLogSnapshots(new List<SnapshotInfo>{new SnapshotInfo(TestNamespace, SnapshotNamespace, snapshotBlob, DateTime.Now)}));
             handler.SetupRequest($"http://localhost/api/v1/replication-log/snapshots/{TestNamespace}").ReturnsResponse(s, "application/json");
 
             // when the snapshot has been processed we have a set of incremental events as well
@@ -368,7 +368,7 @@ namespace Horde.Storage.FunctionalTests.Replication
             Guid missingId = Guid.NewGuid();
 
             Mock<HttpMessageHandler> handler = new Mock<HttpMessageHandler>();
-            string s = JsonConvert.SerializeObject(new ReplicationLogSnapshots(new List<SnapshotInfo>{new SnapshotInfo(TestNamespace, SnapshotNamespace, snapshotBlob)}));
+            string s = JsonConvert.SerializeObject(new ReplicationLogSnapshots(new List<SnapshotInfo>{new SnapshotInfo(TestNamespace, SnapshotNamespace, snapshotBlob, DateTime.Now)}));
             handler.SetupRequest($"http://localhost/api/v1/replication-log/snapshots/{TestNamespace}").ReturnsResponse(s, "application/json");
 
             // mock a error being generated due to the lastBucket/event being to old
