Horde Storage - Clarified if a reference has been deleted or not, from the monitoring of our ref gc it looks like deletes may not be functioning under all circumstances

diff --git a/Engine/Source/Programs/Horde/HordeStorage/Horde.Storage/Controllers/ReferencesController.cs b/Engine/Source/Programs/Horde/HordeStorage/Horde.Storage/Controllers/ReferencesController.cs
--- a/Engine/Source/Programs/Horde/HordeStorage/Horde.Storage/Controllers/ReferencesController.cs
+++ b/Engine/Source/Programs/Horde/HordeStorage/Horde.Storage/Controllers/ReferencesController.cs
@@ -507,8 +507,8 @@ namespace Horde.Storage.Controllers
 
             try
             {
-                long deleteCount = await _objectService.Delete(ns, bucket, key);
-                return Ok(new { DeletedCount = deleteCount } );
+                bool deleted = await _objectService.Delete(ns, bucket, key);
+                return Ok(new { DeletedCount = deleted ? 1: 0 } );
             }
             catch (NamespaceNotFoundException e)
             {
diff --git a/Engine/Source/Programs/Horde/HordeStorage/Horde.Storage/Implementation/GC/CacheCleanup.cs b/Engine/Source/Programs/Horde/HordeStorage/Horde.Storage/Implementation/GC/CacheCleanup.cs
--- a/Engine/Source/Programs/Horde/HordeStorage/Horde.Storage/Implementation/GC/CacheCleanup.cs
+++ b/Engine/Source/Programs/Horde/HordeStorage/Horde.Storage/Implementation/GC/CacheCleanup.cs
@@ -75,7 +75,7 @@ namespace Horde.Storage.Implementation
                 using IScope scope = Tracer.Instance.StartActive("refCleanup.delete_record");
                 scope.Span.ResourceName = $"{ns}:{bucket}.{name}";
                 // delete the old record from the ref refs
-                Task storeDelete = _referencesStore.Delete(ns, bucket, name);
+                Task<bool> storeDelete = _referencesStore.Delete(ns, bucket, name);
                 // insert a delete event into the transaction log
                 Task<(string, Guid)> transactionLogDelete = _replicationLog.InsertDeleteEvent(ns, bucket, name, null);
 
@@ -87,8 +87,16 @@ namespace Horde.Storage.Implementation
                 {
                     _logger.Warning(e, "Exception when attempting to delete record {Bucket} {Name} in {Namespace}", bucket, name, ns);
                 }
-                // we convert the ObjectRecords key (which is a iohash) to a keyid (which is a generic string) not the prettiest but this result is only used for debugging purposes anyway
-                deletedRecords.Add(new OldRecord(ns, bucket, new KeyId(name.ToString())));
+
+                if (await storeDelete)
+                {
+                    // we convert the ObjectRecords key (which is a iohash) to a keyid (which is a generic string) not the prettiest but this result is only used for debugging purposes anyway
+                    deletedRecords.Add(new OldRecord(ns, bucket, new KeyId(name.ToString())));
+                }
+                else
+                {
+                    _logger.Warning("Failed to delete record {Bucket} {Name} in {Namespace}", bucket, name, ns);            
+                }
             }
 
             _logger.Information("Finished cleaning {Namespace}. Refs considered: {ConsideredCount} Refs Deleted: {DeletedCount}", ns, consideredCount, deletedRecords.Count);
diff --git a/Engine/Source/Programs/Horde/HordeStorage/Horde.Storage/Implementation/ObjectService.cs b/Engine/Source/Programs/Horde/HordeStorage/Horde.Storage/Implementation/ObjectService.cs
--- a/Engine/Source/Programs/Horde/HordeStorage/Horde.Storage/Implementation/ObjectService.cs
+++ b/Engine/Source/Programs/Horde/HordeStorage/Horde.Storage/Implementation/ObjectService.cs
@@ -22,7 +22,7 @@ namespace Horde.Storage.Implementation
 
         IAsyncEnumerable<NamespaceId> GetNamespaces();
 
-        Task<long> Delete(NamespaceId ns, BucketId bucket, IoHashKey key);
+        Task<bool> Delete(NamespaceId ns, BucketId bucket, IoHashKey key);
         Task<long> DropNamespace(NamespaceId ns);
         Task<long> DeleteBucket(NamespaceId ns, BucketId bucket);
     }
@@ -174,7 +174,7 @@ namespace Horde.Storage.Implementation
             return _referencesStore.GetNamespaces();
         }
 
-        public Task<long> Delete(NamespaceId ns, BucketId bucket, IoHashKey key)
+        public Task<bool> Delete(NamespaceId ns, BucketId bucket, IoHashKey key)
         {
             return _referencesStore.Delete(ns, bucket, key);
         }
diff --git a/Engine/Source/Programs/Horde/HordeStorage/Horde.Storage/Implementation/Objects/IReferencesStore.cs b/Engine/Source/Programs/Horde/HordeStorage/Horde.Storage/Implementation/Objects/IReferencesStore.cs
--- a/Engine/Source/Programs/Horde/HordeStorage/Horde.Storage/Implementation/Objects/IReferencesStore.cs
+++ b/Engine/Source/Programs/Horde/HordeStorage/Horde.Storage/Implementation/Objects/IReferencesStore.cs
@@ -27,7 +27,7 @@ namespace Horde.Storage.Implementation
         IAsyncEnumerable<(BucketId, IoHashKey, DateTime)> GetRecords(NamespaceId ns);
 
         IAsyncEnumerable<NamespaceId> GetNamespaces();
-        Task<long> Delete(NamespaceId ns, BucketId bucket, IoHashKey key);
+        Task<bool> Delete(NamespaceId ns, BucketId bucket, IoHashKey key);
         Task<long> DropNamespace(NamespaceId ns);
         Task<long> DeleteBucket(NamespaceId ns, BucketId bucket);
     }
diff --git a/Engine/Source/Programs/Horde/HordeStorage/Horde.Storage/Implementation/Objects/MemoryReferencesStore.cs b/Engine/Source/Programs/Horde/HordeStorage/Horde.Storage/Implementation/Objects/MemoryReferencesStore.cs
--- a/Engine/Source/Programs/Horde/HordeStorage/Horde.Storage/Implementation/Objects/MemoryReferencesStore.cs
+++ b/Engine/Source/Programs/Horde/HordeStorage/Horde.Storage/Implementation/Objects/MemoryReferencesStore.cs
@@ -80,14 +80,14 @@ namespace Horde.Storage.Implementation
             return _namespaces.ToAsyncEnumerable();
         }
 
-        public Task<long> Delete(NamespaceId ns, BucketId bucket, IoHashKey key)
+        public Task<bool> Delete(NamespaceId ns, BucketId bucket, IoHashKey key)
         {
             if (!_objects.TryRemove(BuildKey(ns, bucket, key), out MemoryStoreObject? o))
             {
                 throw new ObjectNotFoundException(ns, bucket, key);
             }
 
-            return Task.FromResult(1L);
+            return Task.FromResult(true);
         }
 
         public Task<long> DropNamespace(NamespaceId ns)
diff --git a/Engine/Source/Programs/Horde/HordeStorage/Horde.Storage/Implementation/Objects/ScyllaReferencesStore.cs b/Engine/Source/Programs/Horde/HordeStorage/Horde.Storage/Implementation/Objects/ScyllaReferencesStore.cs
--- a/Engine/Source/Programs/Horde/HordeStorage/Horde.Storage/Implementation/Objects/ScyllaReferencesStore.cs
+++ b/Engine/Source/Programs/Horde/HordeStorage/Horde.Storage/Implementation/Objects/ScyllaReferencesStore.cs
@@ -253,7 +253,7 @@ namespace Horde.Storage.Implementation
             }
         }
 
-        public async Task<long> Delete(NamespaceId ns, BucketId bucket, IoHashKey key)
+        public async Task<bool> Delete(NamespaceId ns, BucketId bucket, IoHashKey key)
         {
             using IScope _ = Tracer.Instance.StartActive("scylla.delete_record");
             ObjectRecord record;
@@ -264,7 +264,7 @@ namespace Horde.Storage.Implementation
             catch (ObjectNotFoundException)
             {
                 // if the record does not exist we do not need to do anything
-                return 0L;
+                return false;
             }
 
             Task removeTTL = RemoveTTLRecord(record);
@@ -272,9 +272,9 @@ namespace Horde.Storage.Implementation
 
             await removeTTL;
             if (info.Applied)
-                return 1L;
+                return true;
 
-            return 0L;
+            return false;
         }
 
         public async Task<long> DropNamespace(NamespaceId ns)
