Horde Storage - Attemted to fix span mapping in datadog for replication, its currently losing the work happening inside the replicator (only showing the work on the remote instance)

diff --git a/Engine/Source/Programs/Horde/HordeStorage/Horde.Storage/Implementation/Replication/RefsReplicator.cs b/Engine/Source/Programs/Horde/HordeStorage/Horde.Storage/Implementation/Replication/RefsReplicator.cs
--- a/Engine/Source/Programs/Horde/HordeStorage/Horde.Storage/Implementation/Replication/RefsReplicator.cs
+++ b/Engine/Source/Programs/Horde/HordeStorage/Horde.Storage/Implementation/Replication/RefsReplicator.cs
@@ -307,6 +307,8 @@ namespace Horde.Storage.Implementation
 
             CancellationTokenSource cancellationTokenSource = new CancellationTokenSource();
             CancellationTokenSource linkedTokenSource = CancellationTokenSource.CreateLinkedTokenSource(cancellationTokenSource.Token, replicationToken);
+            ISpanContext parentSpanContext = Tracer.Instance.ActiveScope.Span.Context;
+
             await GetRefEvents(ns, lastBucket, lastEvent, replicationToken).ParallelForEachAsync(async (ReplicationLogEvent @event) =>
             {
                 // if we have done all the replication events we should do in a single run we abort
@@ -315,7 +317,7 @@ namespace Horde.Storage.Implementation
                     linkedTokenSource.Cancel();
                     return;
                 }
-                using IScope scope = Tracer.Instance.StartActive("replicator.replicate_op_incremental");
+                using IScope scope = Tracer.Instance.StartActive("replicator.replicate_op_incremental", new SpanCreationSettings {Parent = parentSpanContext});
                 scope.Span.ResourceName = $"[REDACTED_EMAIL]}[REDACTED_EMAIL]}";
 
                 _logger.Information("{Name} New transaction to replicate found. New op: {@Op} . Count of running replications: {CurrentReplications}", _name, @event, replicationTasks.Count);
