Horde Storage - More tracing around the ref replicator, this time as a parent for the actual shared web calls that is the same no matter if its incremental or snapshot.

diff --git a/Engine/Source/Programs/Horde/HordeStorage/Horde.Storage/Implementation/Replication/RefsReplicator.cs b/Engine/Source/Programs/Horde/HordeStorage/Horde.Storage/Implementation/Replication/RefsReplicator.cs
--- a/Engine/Source/Programs/Horde/HordeStorage/Horde.Storage/Implementation/Replication/RefsReplicator.cs
+++ b/Engine/Source/Programs/Horde/HordeStorage/Horde.Storage/Implementation/Replication/RefsReplicator.cs
@@ -391,6 +391,9 @@ namespace Horde.Storage.Implementation
 
         private async Task ReplicateOp(NamespaceId ns, BlobIdentifier blob, CancellationToken cancellationToken)
         {
+            using IScope scope = Tracer.Instance.StartActive("replicator.replicate_op");
+            scope.Span.ResourceName = $"{ns}.{blob}";
+
             // We could potentially do this, but that could be dangerous if missing child references
             // check if this blob exists locally before replicating, if it does we assume we have all of its references already
             //if (await _blobService.Exists(ns, blob))
