Horde Storage - Allow blob index usage to be enabled per namespace - and enable it for the Zen namespace

diff --git a/Engine/Source/Programs/Horde/HordeStorage/Horde.Storage/Implementation/Blob/AmazonS3Store.cs b/Engine/Source/Programs/Horde/HordeStorage/Horde.Storage/Implementation/Blob/AmazonS3Store.cs
--- a/Engine/Source/Programs/Horde/HordeStorage/Horde.Storage/Implementation/Blob/AmazonS3Store.cs
+++ b/Engine/Source/Programs/Horde/HordeStorage/Horde.Storage/Implementation/Blob/AmazonS3Store.cs
@@ -164,7 +164,7 @@ namespace Horde.Storage.Implementation
 
         public async Task<bool> Exists(NamespaceId ns, BlobIdentifier blobIdentifier)
         {
-            if (_settings.UseBlobIndexForExistsCheck)
+            if (_settings.UseBlobIndexForExistsCheck && _settings.NamespacesThatUseBlobIndexForExistsCheck.Contains(ns.ToString()))
             {
                 return await _blobIndex.BlobExistsInRegion(ns, blobIdentifier);
             }
diff --git a/Engine/Source/Programs/Horde/HordeStorage/Horde.Storage/Settings.cs b/Engine/Source/Programs/Horde/HordeStorage/Horde.Storage/Settings.cs
--- a/Engine/Source/Programs/Horde/HordeStorage/Horde.Storage/Settings.cs
+++ b/Engine/Source/Programs/Horde/HordeStorage/Horde.Storage/Settings.cs
@@ -234,6 +234,8 @@ namespace Horde.Storage
 
         public bool UseBlobIndexForExistsCheck { get; set; } = false;
 
+        public List<string> NamespacesThatUseBlobIndexForExistsCheck { get; set; } = new List<string>();
+
     }
 
     public class GCSettings
