Horde Storage - Reset the replication cancellation token on each run, so that the replicator doesnt continue aborting if this has ever been set.

diff --git a/Engine/Source/Programs/Horde/HordeStorage/Horde.Storage/Implementation/Replication/RefsReplicator.cs b/Engine/Source/Programs/Horde/HordeStorage/Horde.Storage/Implementation/Replication/RefsReplicator.cs
--- a/Engine/Source/Programs/Horde/HordeStorage/Horde.Storage/Implementation/Replication/RefsReplicator.cs
+++ b/Engine/Source/Programs/Horde/HordeStorage/Horde.Storage/Implementation/Replication/RefsReplicator.cs
@@ -119,6 +119,7 @@ namespace Horde.Storage.Implementation
                 using IScope scope = Tracer.Instance.StartActive("replicator.run");
                 scope.Span.ResourceName =_name;
 
+                _replicationTokenSource.TryReset();
                 _replicationRunning = true;
                 _replicationFinishedEvent.Reset();
                 CancellationToken replicationToken = _replicationTokenSource.Token;
