Horde Storage- Run consistency check on all namespaces

diff --git a/Engine/Source/Programs/Horde/HordeStorage/Horde.Storage/Implementation/Blob/ConsistencyCheckService.cs b/Engine/Source/Programs/Horde/HordeStorage/Horde.Storage/Implementation/Blob/ConsistencyCheckService.cs
--- a/Engine/Source/Programs/Horde/HordeStorage/Horde.Storage/Implementation/Blob/ConsistencyCheckService.cs
+++ b/Engine/Source/Programs/Horde/HordeStorage/Horde.Storage/Implementation/Blob/ConsistencyCheckService.cs
@@ -1,10 +1,12 @@
 ﻿// Copyright Epic Games, Inc. All Rights Reserved.
 
 using System;
+using System.Collections.Generic;
 using System.IO;
 using System.Threading;
 using System.Threading.Tasks;
 using Cassandra;
+using Dasync.Collections;
 using Datadog.Trace;
 using Jupiter;
 using Jupiter.Implementation;
@@ -21,6 +23,7 @@ namespace Horde.Storage.Implementation
         private readonly IServiceProvider _provider;
         private readonly ILeaderElection _leaderElection;
         private readonly IRefsStore _refsStore;
+        private readonly IReferencesStore _referencesStore;
         private readonly ILogger _logger = Log.ForContext<ConsistencyCheckService>();
 
         public class ConsistencyState
@@ -32,13 +35,14 @@ namespace Horde.Storage.Implementation
             return _settings.CurrentValue.Enabled;
         }
 
-        public ConsistencyCheckService(IOptionsMonitor<ConsistencyCheckSettings> settings, IServiceProvider provider, ILeaderElection leaderElection, IRefsStore refsStore) :
+        public ConsistencyCheckService(IOptionsMonitor<ConsistencyCheckSettings> settings, IServiceProvider provider, ILeaderElection leaderElection, IRefsStore refsStore, IReferencesStore referencesStore) :
             base(serviceName: nameof(ConsistencyCheckService), TimeSpan.FromSeconds(settings.CurrentValue.ConsistencyCheckPollFrequencySeconds), new ConsistencyState())
         {
             _settings = settings;
             _provider = provider;
             _leaderElection = leaderElection;
             _refsStore = refsStore;
+            _referencesStore = referencesStore;
         }
 
         public override async Task<bool> OnPoll(ConsistencyState state, CancellationToken cancellationToken)
@@ -81,9 +85,12 @@ namespace Horde.Storage.Implementation
                 return;
             }
 
+            List<NamespaceId> namespaces = await _refsStore.GetNamespaces().ToListAsync();
+            namespaces.AddRange(await _referencesStore.GetNamespaces().ToListAsync());
+
             long countOfBlobsChecked = 0;
             // technically this does not need to be run per namespace but per s3 bucket
-            await foreach (NamespaceId ns in _refsStore.GetNamespaces())
+            await foreach (NamespaceId ns in namespaces)
             {
                 using IScope scope = Tracer.Instance.StartActive("consistency_check.run");
                 scope.Span.ResourceName = ns.ToString();
