Horde Storage - Changed polling service to use the thread pool, should reduce the number of thread we run a bit hopefully resulting in less congestion during heavy load as the thread pool for requests can use all the threads we have.

diff --git a/Engine/Source/Programs/Horde/HordeStorage/Horde.Storage/Implementation/Blob/ConsistencyCheckService.cs b/Engine/Source/Programs/Horde/HordeStorage/Horde.Storage/Implementation/Blob/ConsistencyCheckService.cs
--- a/Engine/Source/Programs/Horde/HordeStorage/Horde.Storage/Implementation/Blob/ConsistencyCheckService.cs
+++ b/Engine/Source/Programs/Horde/HordeStorage/Horde.Storage/Implementation/Blob/ConsistencyCheckService.cs
@@ -30,7 +30,7 @@ namespace Horde.Storage.Implementation
         {
         }
 
-        public override bool ShouldStartPolling()
+        protected override bool ShouldStartPolling()
         {
             return _settings.CurrentValue.Enabled;
         }
diff --git a/Engine/Source/Programs/Horde/HordeStorage/Horde.Storage/Implementation/GC/BlobCleanupService.cs b/Engine/Source/Programs/Horde/HordeStorage/Horde.Storage/Implementation/GC/BlobCleanupService.cs
--- a/Engine/Source/Programs/Horde/HordeStorage/Horde.Storage/Implementation/GC/BlobCleanupService.cs
+++ b/Engine/Source/Programs/Horde/HordeStorage/Horde.Storage/Implementation/GC/BlobCleanupService.cs
@@ -19,7 +19,7 @@ namespace Horde.Storage.Implementation
         private volatile bool _alreadyPolling;
         private readonly ILogger _logger = Log.ForContext<BlobCleanupService>();
 
-        public override bool ShouldStartPolling()
+        protected override bool ShouldStartPolling()
         {
             return _settings.CurrentValue.BlobCleanupServiceEnabled;
         }
diff --git a/Engine/Source/Programs/Horde/HordeStorage/Horde.Storage/Implementation/GC/RefCleanupService.cs b/Engine/Source/Programs/Horde/HordeStorage/Horde.Storage/Implementation/GC/RefCleanupService.cs
--- a/Engine/Source/Programs/Horde/HordeStorage/Horde.Storage/Implementation/GC/RefCleanupService.cs
+++ b/Engine/Source/Programs/Horde/HordeStorage/Horde.Storage/Implementation/GC/RefCleanupService.cs
@@ -42,7 +42,7 @@ namespace Horde.Storage.Implementation
             _referencesStore = referencesStore;
         }
 
-        public override bool ShouldStartPolling()
+        protected override bool ShouldStartPolling()
         {
             return _settings.CurrentValue.CleanOldRefRecords;
         }
diff --git a/Engine/Source/Programs/Horde/HordeStorage/Horde.Storage/Implementation/Replication/ReplicationService.cs b/Engine/Source/Programs/Horde/HordeStorage/Horde.Storage/Implementation/Replication/ReplicationService.cs
--- a/Engine/Source/Programs/Horde/HordeStorage/Horde.Storage/Implementation/Replication/ReplicationService.cs
+++ b/Engine/Source/Programs/Horde/HordeStorage/Horde.Storage/Implementation/Replication/ReplicationService.cs
@@ -27,7 +27,7 @@ namespace Horde.Storage.Implementation
             public List<IReplicator> Replicators { get; } = new List<IReplicator>();
         }
 
-        public override bool ShouldStartPolling()
+        protected override bool ShouldStartPolling()
         {
             return _settings.CurrentValue.Enabled;
         }
diff --git a/Engine/Source/Programs/Horde/HordeStorage/Horde.Storage/Implementation/TransactionLog/ReplicationSnapshotService.cs b/Engine/Source/Programs/Horde/HordeStorage/Horde.Storage/Implementation/TransactionLog/ReplicationSnapshotService.cs
--- a/Engine/Source/Programs/Horde/HordeStorage/Horde.Storage/Implementation/TransactionLog/ReplicationSnapshotService.cs
+++ b/Engine/Source/Programs/Horde/HordeStorage/Horde.Storage/Implementation/TransactionLog/ReplicationSnapshotService.cs
@@ -31,7 +31,7 @@ namespace Horde.Storage.Implementation
         private readonly CancellationTokenSource _cancellationTokenSource = new CancellationTokenSource();
         private Task? _snapshotBuildTask = null;
 
-        public override bool ShouldStartPolling()
+        protected override bool ShouldStartPolling()
         {
             return _settings.CurrentValue.Enabled;
         }
diff --git a/Engine/Source/Programs/Horde/HordeStorage/Jupiter.Common/PollingService.cs b/Engine/Source/Programs/Horde/HordeStorage/Jupiter.Common/PollingService.cs
--- a/Engine/Source/Programs/Horde/HordeStorage/Jupiter.Common/PollingService.cs
+++ b/Engine/Source/Programs/Horde/HordeStorage/Jupiter.Common/PollingService.cs
@@ -24,19 +24,29 @@ namespace Jupiter
         private readonly string _serviceName;
         private readonly TimeSpan _pollFrequency;
         private readonly T _state;
-        private readonly Thread _pollingThread = new Thread(OnUpdate);
         private readonly CancellationTokenSource _stopPolling = new CancellationTokenSource();
+        private readonly ManualResetEvent _hasFinishedRunning = new ManualResetEvent(true);
+        private Timer? _timer;
 
-        public PollingService(string serviceName, TimeSpan pollFrequency, T state)
+        protected PollingService(string serviceName, TimeSpan pollFrequency, T state)
         {
             _serviceName = serviceName;
             _pollFrequency = pollFrequency;
             _state = state;
+
+            _timer = new Timer(OnUpdate, new ThreadState
+            {
+                ServiceName = _serviceName, 
+                PollFrequency = _pollFrequency, 
+                ServiceState = _state, 
+                StopPollingToken = _stopPolling.Token,
+                Instance = this,
+            }, -1, -1);
         }
 
         public bool Running
         {
-            get { return _pollingThread.IsAlive; }
+            get { return _timer != null; }
         }
 
         public T State
@@ -44,7 +54,7 @@ namespace Jupiter
             get { return _state; }
         }
 
-        public virtual bool ShouldStartPolling()
+        protected virtual bool ShouldStartPolling()
         {
             return true;
         }
@@ -56,14 +66,7 @@ namespace Jupiter
 
             if (shouldPoll)
             {
-                _pollingThread.Start(new ThreadState
-                {
-                    ServiceName = _serviceName, 
-                    PollFrequency = _pollFrequency, 
-                    ServiceState = _state, 
-                    StopPollingToken = _stopPolling.Token,
-                    Instance = this,
-                });
+                _timer?.Change(TimeSpan.Zero, _pollFrequency);
             }
 
             return Task.CompletedTask;
@@ -83,39 +86,26 @@ namespace Jupiter
             CancellationToken stopPollingToken = threadState.StopPollingToken;
             ILogger logger = Log.ForContext<PollingService<T>>();
 
+            instance._hasFinishedRunning.Reset();
+
             while (!stopPollingToken.IsCancellationRequested)
             {
-                var startTime = DateTime.Now;
                 try
                 {
                     bool _ = instance.OnPoll(threadState.ServiceState, stopPollingToken).Result;
                 }
                 catch (AggregateException e)
                 {
-                    logger.Error(e, "{Service} Aggregate exception in polling thread", serviceName);
+                    logger.Error(e, "{Service} Aggregate exception in polling service", serviceName);
                     foreach (Exception inner in e.InnerExceptions)
                     {
-                        logger.Error(inner, "{Service} inner exception in polling thread. Trace: {StackTrace}", serviceName, inner.StackTrace);
+                        logger.Error(inner, "{Service} inner exception in polling service. Trace: {StackTrace}", serviceName, inner.StackTrace);
                         
                     }
                 }
                 catch (Exception e)
                 {
-                    logger.Error(e, "{Service} Exception in polling thread", serviceName);
-                }
-                finally
-                {
-                    TimeSpan duration = DateTime.Now - startTime;
-                    TimeSpan pollFrequency = threadState.PollFrequency;
-
-                    // if we spent less then the poll frequency we sleep to avoid ping the remote server to much 
-                    if (duration < pollFrequency)
-                    {
-                        TimeSpan remainingDuration = pollFrequency - duration;
-                        logger.Information("{Service} ran polled for {Duration} Sleeping for {RemainingDuration} to stay within poll frequency", serviceName, duration, remainingDuration);
-
-                        Thread.Sleep(remainingDuration);
-                    }
+                    logger.Error(e, "{Service} Exception in polling service", serviceName);
                 }
             }
         }
@@ -134,10 +124,10 @@ namespace Jupiter
             await OnStopping(_state);
 
             _stopPolling.Cancel();
-            if (_pollingThread.IsAlive)
-            {
-                _pollingThread.Join();
-            }
+            _hasFinishedRunning.WaitOne();
+            if (_timer != null)
+                await _timer.DisposeAsync();
+            _timer = null;
         }
     }
 }
