Horde: Add proactive job queue based scaling of agent pools

- Extracts current autoscaling impl into different strategies
- Current autoscaling strategy renamed to LeaseUtilization (same impl)
- Tests added for LeaseUtilization
- New strategy using the job queue size added. Allows for more proactive resizing of pools
- Keeps the current AutoscaleService, new AutoscaleServiceV2 is only activated if a strategy is specified in ServerSettings. Good for feature toggling during rollout.

diff --git a/Engine/Source/Programs/Horde/Horde.Build.Tests/Fleet/JobQueueStrategyTest.cs b/Engine/Source/Programs/Horde/Horde.Build.Tests/Fleet/JobQueueStrategyTest.cs
new file mode 100644
--- /dev/null
+++ b/Engine/Source/Programs/Horde/Horde.Build.Tests/Fleet/JobQueueStrategyTest.cs
@@ -0,0 +1,123 @@
+using System;
+using System.Collections.Generic;
+using System.Threading.Tasks;
+using EpicGames.Core;
+using Horde.Build.Fleet.Autoscale;
+using HordeServer.Api;
+using HordeServer.Models;
+using HordeServer.Utilities;
+using HordeServerTests;
+using Microsoft.VisualStudio.TestTools.UnitTesting;
+
+namespace Horde.Build.Tests.Fleet
+{
+	using PoolId = StringId<IPool>;
+	using ProjectId = StringId<IProject>;
+	using StreamId = StringId<IStream>;
+	
+	[TestClass]
+	public class JobQueueTest : TestSetup
+	{
+		[TestMethod]
+		public async Task GetPoolQueueSizes()
+		{
+			(JobQueueStrategy Strategy, PoolSizeData PoolSizeData) = await SetUpJobsAsync(1, 5);
+			Dictionary<PoolId, int> PoolQueueSizes = await Strategy.GetPoolQueueSizesAsync(Clock.UtcNow - TimeSpan.FromHours(2));
+			Assert.AreEqual(1, PoolQueueSizes.Count);
+			Assert.AreEqual(5, PoolQueueSizes[PoolSizeData.Pool.Id]);
+		}
+		
+		[TestMethod]
+		public async Task JobQueueWith5Jobs()
+		{
+			(JobQueueStrategy Strategy, PoolSizeData PoolSizeData) = await SetUpJobsAsync(1, 5);
+			Dictionary<PoolId, PoolSizeData> PoolSizeDatas = new ();
+			PoolSizeDatas[PoolSizeData.Pool.Id] = PoolSizeData;
+			await Strategy.CalcDesiredPoolSizesAsync(PoolSizeDatas);
+			PoolSizeData = PoolSizeDatas[PoolSizeData.Pool.Id];
+			Assert.AreEqual(1, PoolSizeDatas.Count);
+			Assert.AreEqual(PoolSizeData.Agents.Count + 0, PoolSizeData.DesiredAgentCount);
+		}
+		
+		[TestMethod]
+		public async Task JobQueueWith15Jobs()
+		{
+			(JobQueueStrategy Strategy, PoolSizeData PoolSizeData) = await SetUpJobsAsync(1, 15);
+			Dictionary<PoolId, PoolSizeData> PoolSizeDatas = new ();
+			PoolSizeDatas[PoolSizeData.Pool.Id] = PoolSizeData;
+			await Strategy.CalcDesiredPoolSizesAsync(PoolSizeDatas);
+			PoolSizeData = PoolSizeDatas[PoolSizeData.Pool.Id];
+			Assert.AreEqual(1, PoolSizeDatas.Count);
+			Assert.AreEqual(PoolSizeData.Agents.Count + 1, PoolSizeData.DesiredAgentCount);
+		}
+		
+		[TestMethod]
+		public async Task JobQueueWith25Jobs()
+		{
+			(JobQueueStrategy Strategy, PoolSizeData PoolSizeData) = await SetUpJobsAsync(1, 25);
+			Dictionary<PoolId, PoolSizeData> PoolSizeDatas = new ();
+			PoolSizeDatas[PoolSizeData.Pool.Id] = PoolSizeData;
+			await Strategy.CalcDesiredPoolSizesAsync(PoolSizeDatas);
+			PoolSizeData = PoolSizeDatas[PoolSizeData.Pool.Id];
+			Assert.AreEqual(1, PoolSizeDatas.Count);
+			Assert.AreEqual(PoolSizeData.Agents.Count + 2, PoolSizeData.DesiredAgentCount);
+		}
+	
+		/// <summary>
+		/// Set up a fixture for job queue tests, ensuring a certain number of job batches are in running or waiting state
+		/// </summary>
+		/// <param name="NumBatchesRunning">Num of job batches that should be in state running</param>
+		/// <param name="NumBatchesWaiting">Num of job batches that should be in state waiting</param>
+		private async Task<(JobQueueStrategy, PoolSizeData)> SetUpJobsAsync(int NumBatchesRunning, int NumBatchesWaiting)
+		{
+			IPool Pool1 = await PoolService.CreatePoolAsync("bogusPool1", null, true, 0, 0);
+			IAgent Agent1 = await CreateAgentAsync(Pool1);
+			IAgent Agent2 = await CreateAgentAsync(Pool1);
+			PoolSizeData PoolSize = new (Pool1, new List<IAgent>() { Agent1, Agent2 }, null);
+			
+			string AgentTypeName1 = "bogusAgentType1";
+			Dictionary<string, CreateAgentTypeRequest> AgentTypes = new() { {AgentTypeName1, new() { Pool = Pool1.Name} }, };
+			IStream Stream = (await StreamCollection.TryCreateOrReplaceAsync(
+				new StreamId("ue5-main"),
+				null,
+				"",
+				"",
+				new ProjectId("does-not-exist"),
+				new StreamConfig { Name = "//UE5/Main", AgentTypes = AgentTypes}
+			))!;
+
+			string NodeForAgentType1 = "bogusNodeOnAgentType1";
+			IGraph Graph = await GraphCollection.AppendAsync(null, new()
+			{
+				new NewGroup(AgentTypeName1, new List<NewNode>
+				{
+					new (NodeForAgentType1),
+				})
+			});
+
+			for (int i = 0; i < NumBatchesRunning; i++)
+			{
+				IJob Job = await AddPlaceholderJob(Graph, Stream.Id, NodeForAgentType1);
+				await JobCollection.TryUpdateBatchAsync(Job, Graph, Job.Batches[0].Id, null, JobStepBatchState.Running, null);
+			}
+			
+			for (int i = 0; i < NumBatchesWaiting; i++)
+			{
+				IJob Job = await AddPlaceholderJob(Graph, Stream.Id, NodeForAgentType1);
+				await JobCollection.TryUpdateBatchAsync(Job, Graph, Job.Batches[0].Id, null, JobStepBatchState.Waiting, null);
+			}
+			
+			return (new (JobCollection, GraphCollection, StreamService, Clock), PoolSize);
+		}
+		
+		private async Task<IJob> AddPlaceholderJob(IGraph Graph, StreamId StreamId, string NodeNameToExecute)
+		{
+			IJob Job = await JobCollection.AddAsync(ObjectId<IJob>.GenerateNewId(), StreamId,
+				new StringId<TemplateRef>("bogusTemplateRefId"), ContentHash.Empty, Graph, "bogusJobName",
+				1000, 1000, null, null, null, null, null, null, null, false,
+				false, null, null, new List<string> { "-Target=" + NodeNameToExecute });
+
+			return Job;
+		}
+	}
+}
\ No newline at end of file
diff --git a/Engine/Source/Programs/Horde/Horde.Build.Tests/Fleet/LeaseUtilizationStrategyTest.cs b/Engine/Source/Programs/Horde/Horde.Build.Tests/Fleet/LeaseUtilizationStrategyTest.cs
new file mode 100644
--- /dev/null
+++ b/Engine/Source/Programs/Horde/Horde.Build.Tests/Fleet/LeaseUtilizationStrategyTest.cs
@@ -0,0 +1,118 @@
+using System;
+using System.Collections.Generic;
+using System.Threading.Tasks;
+using Google.Protobuf;
+using Google.Protobuf.WellKnownTypes;
+using Horde.Build.Fleet.Autoscale;
+using HordeCommon;
+using HordeCommon.Rpc.Tasks;
+using HordeServer.Models;
+using HordeServer.Utilities;
+using HordeServerTests;
+using Microsoft.VisualStudio.TestTools.UnitTesting;
+
+namespace Horde.Build.Tests.Fleet
+{
+	using PoolId = StringId<IPool>;
+	using StreamId = StringId<IStream>;
+	
+	[TestClass]
+	public class LeaseUtilizationStrategyTest : TestSetup
+	{
+		private IAgent Agent1 = null!;
+		private IAgent Agent2 = null!;
+		private IAgent Agent3 = null!;
+		private IAgent Agent4 = null!;
+		private List<IAgent> PoolAgents = null!;
+
+		public void CreateAgents(IPool Pool)
+		{
+			Agent1 = CreateAgentAsync(Pool).Result;
+			Agent2 = CreateAgentAsync(Pool).Result;
+			Agent3 = CreateAgentAsync(Pool).Result;
+			Agent4 = CreateAgentAsync(Pool).Result;
+			PoolAgents = new() { Agent1, Agent2, Agent3, Agent4 };
+		}
+
+		[TestMethod]
+		public async Task UtilizationFull()
+		{
+			IPool Pool = await PoolService.CreatePoolAsync("AutoscalePool1", null, true, 0, 0);
+			CreateAgents(Pool);
+			
+			await AddPlaceholderLease(Agent1, Pool, Clock.UtcNow - TimeSpan.FromMinutes(120), TimeSpan.FromMinutes(120));
+			await AddPlaceholderLease(Agent2, Pool, Clock.UtcNow - TimeSpan.FromMinutes(120), TimeSpan.FromMinutes(120));
+			await AddPlaceholderLease(Agent3, Pool, Clock.UtcNow - TimeSpan.FromMinutes(120), TimeSpan.FromMinutes(120));
+			await AddPlaceholderLease(Agent4, Pool, Clock.UtcNow - TimeSpan.FromMinutes(120), TimeSpan.FromMinutes(120));
+			await AssertPoolSizeAsync(Pool, 4);
+		}
+		
+		[TestMethod]
+		public async Task UtilizationHalf()
+		{
+			IPool Pool = await PoolService.CreatePoolAsync("AutoscalePool1", null, true, 0, 0);
+			CreateAgents(Pool);
+			
+			await AddPlaceholderLease(Agent1, Pool, Clock.UtcNow - TimeSpan.FromMinutes(120), TimeSpan.FromMinutes(120));
+			await AddPlaceholderLease(Agent2, Pool, Clock.UtcNow - TimeSpan.FromMinutes(120), TimeSpan.FromMinutes(120));
+			await AssertPoolSizeAsync(Pool, 2);
+		}
+		
+		[TestMethod]
+		public async Task UtilizationZero()
+		{
+			IPool Pool = await PoolService.CreatePoolAsync("AutoscalePool1", null, true, 0, 0);
+			CreateAgents(Pool);
+
+			await AssertPoolSizeAsync(Pool, 0);
+		}
+		
+		[TestMethod]
+		public async Task ReserveAgents()
+		{
+			IPool Pool = await PoolService.CreatePoolAsync("AutoscalePool1", null, true, 0, 5);
+			CreateAgents(Pool);
+			
+			await AddPlaceholderLease(Agent1, Pool, Clock.UtcNow - TimeSpan.FromMinutes(120), TimeSpan.FromMinutes(120));
+			await AddPlaceholderLease(Agent2, Pool, Clock.UtcNow - TimeSpan.FromMinutes(120), TimeSpan.FromMinutes(120));
+			await AddPlaceholderLease(Agent3, Pool, Clock.UtcNow - TimeSpan.FromMinutes(120), TimeSpan.FromMinutes(120));
+			await AddPlaceholderLease(Agent4, Pool, Clock.UtcNow - TimeSpan.FromMinutes(120), TimeSpan.FromMinutes(120));
+			
+			// Full utilization should mean all agents plus the reserve agents
+			await AssertPoolSizeAsync(Pool, 9);
+		}
+		
+		[TestMethod]
+		public async Task MinAgents()
+		{
+			IPool Pool = await PoolService.CreatePoolAsync("AutoscalePool1", null, true, 2, 0);
+			CreateAgents(Pool);
+			
+			// Even with no utilization, expect at least the min number of agents
+			await AssertPoolSizeAsync(Pool, 2);
+		}
+
+		private async Task AssertPoolSizeAsync(IPool Pool, int ExpectedNumAgents)
+		{
+			LeaseUtilizationStrategy Strategy = new (AgentCollection, PoolCollection, LeaseCollection, Clock);
+			Dictionary<PoolId, PoolSizeData> PoolSizeDatas = new ();
+			PoolSizeDatas[Pool.Id] = new(Pool, PoolAgents, null);
+			await Strategy.CalcDesiredPoolSizesAsync(PoolSizeDatas);
+			Assert.AreEqual(ExpectedNumAgents, PoolSizeDatas[Pool.Id].DesiredAgentCount);
+		}
+		
+		private async Task<ILease> AddPlaceholderLease(IAgent Agent, IPool Pool, DateTime StartTime, TimeSpan Duration)
+		{
+			Assert.IsNotNull(Agent.SessionId);
+			
+			ExecuteJobTask PlaceholderJobTask = new();
+			PlaceholderJobTask.JobName = "placeholderJobName";
+			byte[] Payload = Any.Pack(PlaceholderJobTask).ToByteArray();
+
+			ILease Lease = await LeaseCollection.AddAsync(ObjectId<ILease>.GenerateNewId(), "placeholderLease", Agent.Id, Agent.SessionId!.Value, new StreamId("placeholderStream"), Pool.Id, null, StartTime, Payload);
+			bool WasModified = await LeaseCollection.TrySetOutcomeAsync(Lease.Id, StartTime + Duration, LeaseOutcome.Success, null);
+			Assert.IsTrue(WasModified);
+			return Lease;
+		}
+	}
+}
\ No newline at end of file
diff --git a/Engine/Source/Programs/Horde/Horde.Build.Tests/TestSetup.cs b/Engine/Source/Programs/Horde/Horde.Build.Tests/TestSetup.cs
--- a/Engine/Source/Programs/Horde/Horde.Build.Tests/TestSetup.cs
+++ b/Engine/Source/Programs/Horde/Horde.Build.Tests/TestSetup.cs
@@ -246,6 +246,14 @@ namespace HordeServerTests
 				new List<Claim> {new Claim(ServerSettings.AdminClaimType, ServerSettings.AdminClaimValue)}, "TestAuthType"));
 			return ControllerContext;
 		}
+		
+		private static int AgentIdCounter = 1;
+		public async Task<IAgent> CreateAgentAsync(IPool Pool)
+		{
+			IAgent Agent = await AgentService.CreateAgentAsync("TestAgent" + AgentIdCounter++, true, null, new List<StringId<IPool>> { Pool.Id });
+			Agent = await AgentService.CreateSessionAsync(Agent, AgentStatus.Ok, new List<string>(), new Dictionary<string, int>(), null);
+			return Agent;
+		}
 
 		/// <summary>
 		/// Hack the Datadog tracing library to not block during shutdown of tests.
diff --git a/Engine/Source/Programs/Horde/Horde.Build/Fleet/Autoscale/AutoscaleService.cs b/Engine/Source/Programs/Horde/Horde.Build/Fleet/Autoscale/AutoscaleService.cs
new file mode 100644
--- /dev/null
+++ b/Engine/Source/Programs/Horde/Horde.Build/Fleet/Autoscale/AutoscaleService.cs
@@ -0,0 +1,288 @@
+// Copyright Epic Games, Inc. All Rights Reserved.
+
+using Amazon;
+using Amazon.EC2;
+using Amazon.EC2.Model;
+using Amazon.Runtime;
+using Amazon.Runtime.CredentialManagement;
+using Google.Protobuf.WellKnownTypes;
+using HordeCommon.Rpc.Tasks;
+using HordeServer.Collections;
+using HordeServer.Models;
+using HordeServer.Utilities;
+using Microsoft.Extensions.Logging;
+using MongoDB.Bson;
+using System;
+using System.Collections.Generic;
+using System.Diagnostics;
+using System.Linq;
+using System.Net;
+using System.Runtime.CompilerServices;
+using System.Threading;
+using System.Threading.Tasks;
+using HordeCommon;
+using Microsoft.Extensions.Options;
+using OpenTracing;
+using OpenTracing.Util;
+using StatsdClient;
+using Microsoft.Extensions.Hosting;
+
+namespace HordeServer.Services
+{
+	using PoolId = StringId<IPool>;
+
+	/// <summary>
+	/// Service for managing the autoscaling of agent pools
+	/// </summary>
+	public sealed class AutoscaleService : IHostedService, IDisposable
+	{
+		const int NumSamples = 10;
+		const int NumSamplesForResult = 9;
+		static readonly TimeSpan SampleTime = TimeSpan.FromMinutes(6.0);
+		static readonly TimeSpan ShrinkPoolCoolDown = TimeSpan.FromMinutes(20.0);
+
+		struct UtilizationSample
+		{
+			public double JobWork;
+			public double OtherWork;
+		}
+
+		class AgentData
+		{
+			public IAgent Agent { get; }
+			public UtilizationSample[] Samples { get; } = new UtilizationSample[NumSamples];
+			public List<ILease> Leases { get; } = new List<ILease>();
+
+			public AgentData(IAgent Agent)
+			{
+				this.Agent = Agent;
+			}
+
+			public void Add(double MinT, double MaxT, double JobWork, double OtherWork)
+			{
+				int MinIdx = Math.Clamp((int)MinT, 0, NumSamples - 1);
+				int MaxIdx = Math.Clamp((int)MaxT, MinIdx, NumSamples - 1);
+				for (int Idx = MinIdx; Idx <= MaxIdx; Idx++)
+				{
+					double Fraction = Math.Clamp(MaxT - Idx, 0.0, 1.0) - Math.Clamp(MinT - Idx, 0.0, 1.0);
+					Samples[Idx].JobWork += JobWork * Fraction;
+					Samples[Idx].OtherWork += OtherWork * Fraction;
+				}
+			}
+		}
+
+		class PoolData
+		{
+			public IPool Pool { get; }
+			public List<IAgent> Agents { get; } = new List<IAgent>();
+			public UtilizationSample[] Samples { get; } = new UtilizationSample[NumSamples];
+
+			public PoolData(IPool Pool)
+			{
+				this.Pool = Pool;
+			}
+
+			public void Add(AgentData AgentData)
+			{
+				Agents.Add(AgentData.Agent);
+				for (int Idx = 0; Idx < Samples.Length; Idx++)
+				{
+					Samples[Idx].JobWork += AgentData.Samples[Idx].JobWork;
+					Samples[Idx].OtherWork += AgentData.Samples[Idx].OtherWork;
+				}
+			}
+		}
+
+		IAgentCollection AgentCollection;
+		IPoolCollection PoolCollection;
+		ILeaseCollection LeaseCollection;
+		IFleetManager FleetManager;
+		IDogStatsd DogStatsd;
+		IClock Clock;
+		ITicker Ticker;
+		ILogger<AutoscaleService> Logger;
+
+		/// <summary>
+		/// Constructor
+		/// </summary>
+		public AutoscaleService(IAgentCollection AgentCollection, IPoolCollection PoolCollection, ILeaseCollection LeaseCollection, IFleetManager FleetManager, IDogStatsd DogStatsd, DatabaseService DatabaseService, IClock Clock, ILogger<AutoscaleService> Logger)
+		{
+			this.AgentCollection = AgentCollection;
+			this.PoolCollection = PoolCollection;
+			this.LeaseCollection = LeaseCollection;
+			this.FleetManager = FleetManager;
+			this.DogStatsd = DogStatsd;
+			this.Clock = Clock;
+			if (DatabaseService.ReadOnlyMode)
+			{
+				this.Ticker = new NullTicker();
+			}
+			else
+			{
+				this.Ticker = Clock.AddSharedTicker<AutoscaleService>(TimeSpan.FromMinutes(5.0), TickLeaderAsync, Logger);
+			}
+			this.Logger = Logger;
+		}
+
+		/// <inheritdoc/>
+		public Task StartAsync(CancellationToken CancellationToken) => Ticker.StartAsync();
+
+		/// <inheritdoc/>
+		public Task StopAsync(CancellationToken CancellationToken) => Ticker.StopAsync();
+
+		/// <inheritdoc/>
+		public void Dispose() => Ticker.Dispose();
+
+		/// <inheritdoc/>
+		async ValueTask TickLeaderAsync(CancellationToken StoppingToken)
+		{
+			DateTime UtcNow = Clock.UtcNow;
+
+			Logger.LogInformation("Autoscaling pools...");
+			Stopwatch Stopwatch = Stopwatch.StartNew();
+
+			Dictionary<PoolId, PoolData> PoolToData;
+			using (IScope Scope = GlobalTracer.Instance.BuildSpan("Compute utilization").StartActive())
+			{
+				// Find all the current agents
+				List<IAgent> Agents = await AgentCollection.FindAsync(Status: AgentStatus.Ok);
+
+				// Query leases in last interval
+				DateTime MaxTime = Clock.UtcNow;
+				DateTime MinTime = MaxTime - (SampleTime * NumSamples);
+				List<ILease> Leases = await LeaseCollection.FindLeasesAsync(MinTime, MaxTime);
+
+				// Add all the leases to a data object for each agent
+				Dictionary<AgentId, AgentData> AgentIdToData = Agents.ToDictionary(x => x.Id, x => new AgentData(x));
+				foreach (ILease Lease in Leases)
+				{
+					AgentData? AgentData;
+					if (AgentIdToData.TryGetValue(Lease.AgentId, out AgentData) &&
+					    AgentData.Agent.SessionId == Lease.SessionId)
+					{
+						AgentData.Leases.Add(Lease);
+					}
+				}
+
+				// Compute utilization for each agent
+				foreach (AgentData AgentData in AgentIdToData.Values)
+				{
+					foreach (ILease Lease in AgentData.Leases.OrderBy(x => x.StartTime))
+					{
+						double MinT = (Lease.StartTime - MinTime).TotalSeconds / SampleTime.TotalSeconds;
+						double MaxT = (Lease.FinishTime == null)
+							? NumSamples
+							: ((Lease.FinishTime.Value - MinTime).TotalSeconds / SampleTime.TotalSeconds);
+
+						Any Payload = Any.Parser.ParseFrom(Lease.Payload.ToArray());
+						if (Payload.Is(ExecuteJobTask.Descriptor))
+						{
+							AgentData.Add(MinT, MaxT, 1.0, 0.0);
+						}
+						else
+						{
+							AgentData.Add(MinT, MaxT, 0.0, 1.0);
+						}
+					}
+				}
+
+				// Get all the pools
+				List<IPool> Pools = await PoolCollection.GetAsync();
+				PoolToData = Pools.ToDictionary(x => x.Id, x => new PoolData(x));
+
+				// Find pool utilization over the query period
+				foreach (AgentData AgentData in AgentIdToData.Values)
+				{
+					foreach (PoolId PoolId in AgentData.Agent.GetPools())
+					{
+						PoolData? PoolData;
+						if (PoolToData.TryGetValue(PoolId, out PoolData))
+						{
+							PoolData.Add(AgentData);
+						}
+					}
+				}
+			}
+
+			// Output all the final stats
+			foreach (PoolData PoolData in PoolToData.Values.OrderByDescending(x => x.Agents.Count))
+			{
+				IPool Pool = PoolData.Pool;
+
+				int MinAgents = Pool.MinAgents ?? 1;
+				int NumReserveAgents = Pool.NumReserveAgents ?? 5;
+
+				double Utilization = PoolData.Samples.Select(x => x.JobWork).OrderByDescending(x => x).Skip(NumSamples - NumSamplesForResult).First();
+
+				int TargetAgents = Math.Max((int)Utilization + NumReserveAgents, MinAgents);
+				int Delta = TargetAgents - PoolData.Agents.Count;
+
+				Logger.LogInformation("{PoolName,-48} Jobs=[{JobUtilization}] Total=[{OtherUtilization}] Min={Min,5:0.0} Max={Max,5:0.0} Avg={Avg,5:0.0} Pct={Pct,5:0.0} Current={Current,4} Target={Target,4} Delta={Delta}",
+					Pool.Name,
+					GetDensityMap(PoolData.Samples.Select(x => x.JobWork / Math.Max(1, PoolData.Agents.Count))),
+					GetDensityMap(PoolData.Samples.Select(x => x.OtherWork / Math.Max(1, PoolData.Agents.Count))),
+					PoolData.Samples.Min(x => x.JobWork),
+					PoolData.Samples.Max(x => x.JobWork),
+					Utilization,
+					PoolData.Samples.Sum(x => x.JobWork) / NumSamples,
+					PoolData.Agents.Count,
+					TargetAgents,
+					Delta);
+
+				try
+				{
+					if (Pool.EnableAutoscaling)
+					{
+						using (IScope Scope = GlobalTracer.Instance.BuildSpan("Scaling pool").StartActive())
+						{
+							Scope.Span.SetTag("poolName", Pool.Name);
+							Scope.Span.SetTag("delta", Delta);
+
+							if (Delta > 0)
+							{
+								await FleetManager.ExpandPoolAsync(Pool, PoolData.Agents, Delta);
+								await PoolCollection.TryUpdateAsync(Pool, LastScaleUpTime: DateTime.UtcNow);
+							}
+
+							if (Delta < 0)
+							{
+								bool bShrinkIsOnCoolDown = Pool.LastScaleDownTime != null && Pool.LastScaleDownTime + ShrinkPoolCoolDown > DateTime.UtcNow;
+								if (!bShrinkIsOnCoolDown)
+								{
+									await FleetManager.ShrinkPoolAsync(Pool, PoolData.Agents, -Delta);
+									await PoolCollection.TryUpdateAsync(Pool, LastScaleDownTime: DateTime.UtcNow);
+								}
+								else
+								{
+									Logger.LogDebug("Cannot shrink {PoolName} right now, it's on cool-down until {CoolDownTimeEnds}", Pool.Name, Pool.LastScaleDownTime + ShrinkPoolCoolDown);
+								}
+							}
+						}
+					}
+				}
+				catch (Exception Ex)
+				{
+					Logger.LogError(Ex, "Failed to scale {PoolName}:\n{Exception}", Pool.Name, Ex);
+					continue;
+				}
+
+				DogStatsd.Gauge("agentpools.autoscale.target", TargetAgents, tags: new []{"pool:" + Pool.Name});
+				DogStatsd.Gauge("agentpools.autoscale.current", PoolData.Agents.Count, tags: new []{"pool:" + Pool.Name});
+			}
+			
+			Stopwatch.Stop();
+			Logger.LogInformation("Autoscaling pools took {ElapsedTime} ms", Stopwatch.ElapsedMilliseconds);
+		}
+
+		/// <summary>
+		/// Creates a string of characters indicating a sequence of 0-1 values over time
+		/// </summary>
+		/// <param name="Values">Sequence of values</param>
+		/// <returns>Density map</returns>
+		static string GetDensityMap(IEnumerable<double> Values)
+		{
+			const string Greyscale = " 123456789";
+			return new string(Values.Select(x => Greyscale[Math.Clamp((int)(x * Greyscale.Length), 0, Greyscale.Length - 1)]).ToArray());
+		}
+	}
+}
diff --git a/Engine/Source/Programs/Horde/Horde.Build/Fleet/Autoscale/AutoscaleServiceV2.cs b/Engine/Source/Programs/Horde/Horde.Build/Fleet/Autoscale/AutoscaleServiceV2.cs
new file mode 100644
--- /dev/null
+++ b/Engine/Source/Programs/Horde/Horde.Build/Fleet/Autoscale/AutoscaleServiceV2.cs
@@ -0,0 +1,138 @@
+// Copyright Epic Games, Inc. All Rights Reserved.
+
+using System;
+using System.Collections.Generic;
+using System.Diagnostics;
+using System.Linq;
+using System.Threading;
+using System.Threading.Tasks;
+using HordeCommon;
+using HordeServer.Collections;
+using HordeServer.Models;
+using HordeServer.Services;
+using HordeServer.Utilities;
+using Microsoft.Extensions.Hosting;
+using Microsoft.Extensions.Logging;
+using OpenTracing;
+using OpenTracing.Util;
+using StatsdClient;
+
+namespace Horde.Build.Fleet.Autoscale
+{
+	using PoolId = StringId<IPool>;
+
+	/// <summary>
+	/// Service for managing the autoscaling of agent pools
+	/// </summary>
+	public sealed class AutoscaleServiceV2 : IHostedService, IDisposable
+	{
+		static readonly TimeSpan ShrinkPoolCoolDown = TimeSpan.FromMinutes(20.0);
+
+		IPoolSizeStrategy PoolSizeStrategy;
+		IAgentCollection AgentCollection;
+		IPoolCollection PoolCollection;
+		IFleetManager FleetManager;
+		IDogStatsd DogStatsd;
+		IClock Clock;
+		ITicker Ticker;
+		ILogger<AutoscaleServiceV2> Logger;
+
+		/// <summary>
+		/// Constructor
+		/// </summary>
+		public AutoscaleServiceV2(IPoolSizeStrategy PoolSizeStrategy, IAgentCollection AgentCollection, IPoolCollection PoolCollection, IFleetManager FleetManager, IDogStatsd DogStatsd, IClock Clock, ILogger<AutoscaleServiceV2> Logger)
+		{
+			this.PoolSizeStrategy = PoolSizeStrategy;
+			this.AgentCollection = AgentCollection;
+			this.PoolCollection = PoolCollection;
+			this.FleetManager = FleetManager;
+			this.DogStatsd = DogStatsd;
+			this.Clock = Clock;
+			this.Ticker = Clock.AddSharedTicker<AutoscaleServiceV2>(TimeSpan.FromMinutes(5.0), TickLeaderAsync, Logger);
+			this.Logger = Logger;
+		}
+
+		/// <inheritdoc/>
+		public Task StartAsync(CancellationToken CancellationToken) => Ticker.StartAsync();
+
+		/// <inheritdoc/>
+		public Task StopAsync(CancellationToken CancellationToken) => Ticker.StopAsync();
+
+		/// <inheritdoc/>
+		public void Dispose() => Ticker.Dispose();
+
+		/// <inheritdoc/>
+		async ValueTask TickLeaderAsync(CancellationToken StoppingToken)
+		{
+			DateTime UtcNow = Clock.UtcNow;
+
+			Logger.LogInformation("Autoscaling pools using strategy {PoolSizeStrategyName}...", PoolSizeStrategy.Name);
+			Stopwatch Stopwatch = Stopwatch.StartNew();
+			using IScope _ = GlobalTracer.Instance.BuildSpan("AutoscaleService").StartActive();
+			
+			List<IAgent> Agents = await AgentCollection.FindAsync(Status: AgentStatus.Ok);
+			List<IAgent> GetAgentsInPool(PoolId PoolId)
+			{
+				return Agents.FindAll(a => a.GetPools().Any(p => p == PoolId));
+			}
+			
+			List<IPool> Pools = await PoolCollection.GetAsync();
+			Dictionary<PoolId, PoolSizeData> PoolSizeDatas = Pools.ToDictionary(x => x.Id, x => new PoolSizeData(x, GetAgentsInPool(x.Id), null));
+
+			await PoolSizeStrategy.CalcDesiredPoolSizesAsync(PoolSizeDatas);
+
+			foreach (PoolSizeData PoolSizeData in PoolSizeDatas.Values.OrderByDescending(x => x.Agents.Count))
+			{
+				IPool Pool = PoolSizeData.Pool;
+
+				if (!Pool.EnableAutoscaling || PoolSizeData.DesiredAgentCount == null) continue;
+
+				int CurrentAgentCount = PoolSizeData.Agents.Count;
+				int DesiredAgentCount = PoolSizeData.DesiredAgentCount.Value;
+				int DeltaAgentCount = DesiredAgentCount - CurrentAgentCount;
+
+				Logger.LogInformation("{PoolName,-48} Current={Current,4} Target={Target,4} Delta={Delta,4} Status={Status}", Pool.Name, CurrentAgentCount, DesiredAgentCount, DeltaAgentCount, PoolSizeData.StatusMessage);
+				
+				try
+				{
+					using IScope Scope = GlobalTracer.Instance.BuildSpan("Scaling pool").StartActive();
+					Scope.Span.SetTag("poolName", Pool.Name);
+					Scope.Span.SetTag("current", CurrentAgentCount);
+					Scope.Span.SetTag("desired", DesiredAgentCount);
+					Scope.Span.SetTag("delta", DeltaAgentCount);
+
+					if (DeltaAgentCount > 0)
+					{
+						await FleetManager.ExpandPoolAsync(Pool, PoolSizeData.Agents, DeltaAgentCount);
+						await PoolCollection.TryUpdateAsync(Pool, LastScaleUpTime: DateTime.UtcNow);
+					}
+
+					if (DeltaAgentCount < 0)
+					{
+						bool bShrinkIsOnCoolDown = Pool.LastScaleDownTime != null && Pool.LastScaleDownTime + ShrinkPoolCoolDown > DateTime.UtcNow;
+						if (!bShrinkIsOnCoolDown)
+						{
+							await FleetManager.ShrinkPoolAsync(Pool, PoolSizeData.Agents, -DeltaAgentCount);
+							await PoolCollection.TryUpdateAsync(Pool, LastScaleDownTime: DateTime.UtcNow);
+						}
+						else
+						{
+							Logger.LogDebug("Cannot shrink {PoolName} right now, it's on cool-down until {CoolDownTimeEnds}", Pool.Name, Pool.LastScaleDownTime + ShrinkPoolCoolDown);
+						}
+					}
+				}
+				catch (Exception Ex)
+				{
+					Logger.LogError(Ex, "Failed to scale {PoolName}:\n{Exception}", Pool.Name, Ex);
+					continue;
+				}
+
+				DogStatsd.Gauge("agentpools.autoscale.target", DesiredAgentCount, tags: new []{"pool:" + Pool.Name});
+				DogStatsd.Gauge("agentpools.autoscale.current", CurrentAgentCount, tags: new []{"pool:" + Pool.Name});
+			}
+			
+			Stopwatch.Stop();
+			Logger.LogInformation("Autoscaling pools took {ElapsedTime} ms", Stopwatch.ElapsedMilliseconds);
+		}
+	}
+}
diff --git a/Engine/Source/Programs/Horde/Horde.Build/Fleet/Autoscale/IPoolSizeStrategy.cs b/Engine/Source/Programs/Horde/Horde.Build/Fleet/Autoscale/IPoolSizeStrategy.cs
new file mode 100644
--- /dev/null
+++ b/Engine/Source/Programs/Horde/Horde.Build/Fleet/Autoscale/IPoolSizeStrategy.cs
@@ -0,0 +1,115 @@
+// Copyright Epic Games, Inc. All Rights Reserved.
+
+using System.Collections.Generic;
+using System.Threading.Tasks;
+using HordeServer.Models;
+using HordeServer.Utilities;
+
+namespace Horde.Build.Fleet.Autoscale
+{
+	using PoolId = StringId<IPool>;
+
+	/// <summary>
+	/// Available pool sizing strategies
+	/// </summary>
+	public enum PoolSizeStrategy
+	{
+		/// <summary>
+		/// Strategy based on lease utilization
+		/// <see cref="Horde.Build.Fleet.Autoscale.LeaseUtilizationStrategy"/>
+		/// </summary>
+		LeaseUtilization,
+		
+		/// <summary>
+		/// Strategy based on size of job build queue
+		/// <see cref="Horde.Build.Fleet.Autoscale.JobQueueStrategy"/> 
+		/// </summary>
+		JobQueue,
+		
+		/// <summary>
+		/// No-op strategy used as fallback/default behavior
+		/// <see cref="Horde.Build.Fleet.Autoscale.NoOpPoolSizeStrategy"/> 
+		/// </summary>
+		NoOp
+	}
+
+	/// <summary>
+	/// Class for specifying and grouping data together required for calculating pool size
+	/// </summary>
+	public class PoolSizeData
+	{
+		/// <summary>
+		/// Pool being resized
+		/// </summary>
+		public IPool Pool { get; }
+		
+		/// <summary>
+		/// All agents currently associated with the pool
+		/// </summary>
+		public List<IAgent> Agents { get; }
+
+		/// <summary>
+		/// The desired agent count (calculated and updated once the strategy has been run, null otherwise)
+		/// </summary>
+		public int? DesiredAgentCount  { get; }
+		
+		/// <summary>
+		/// Human-readable text describing the status of the pool (data the sizing is based on etc)
+		/// </summary>
+		public string StatusMessage { get; }
+
+		/// <summary>
+		/// Constructor
+		/// </summary>
+		/// <param name="Pool"></param>
+		/// <param name="Agents"></param>
+		/// <param name="DesiredAgentCount"></param>
+		/// <param name="StatusMessage"></param>
+		public PoolSizeData(IPool Pool, List<IAgent> Agents, int? DesiredAgentCount, string StatusMessage = "N/A")
+		{
+			this.Pool = Pool;
+			this.Agents = Agents;
+			this.DesiredAgentCount = DesiredAgentCount;
+			this.StatusMessage = StatusMessage;
+		}
+	}
+	
+	/// <summary>
+	/// Interface for different agent pool sizing strategies
+	/// </summary>
+	public interface IPoolSizeStrategy
+	{
+		/// <summary>
+		/// Calculate the adequate number of agents to be online for given pools
+		/// </summary>
+		/// <param name="Pools">Pools including attached agents</param>
+		/// <returns></returns>
+		Task CalcDesiredPoolSizesAsync(Dictionary<PoolId, PoolSizeData> Pools);
+		
+		/// <summary>
+		/// Name of the strategy
+		/// </summary>
+		string Name { get; }
+	}
+	
+	/// <summary>
+	/// No-operation strategy that won't resize pools, just return the existing count.
+	/// Used to ensure there's always a strategy available for dependency injection, even if it does nothing.
+	/// </summary>
+	public class NoOpPoolSizeStrategy : IPoolSizeStrategy
+	{
+		/// <inheritdoc/>
+		public Task CalcDesiredPoolSizesAsync(Dictionary<PoolId, PoolSizeData> Pools)
+		{
+			foreach (PoolSizeData Data in Pools.Values)
+			{
+				Pools[Data.Pool.Id] = new(Data.Pool, Data.Agents, Data.Agents.Count, "(no-op)");
+			}
+
+			return Task.CompletedTask;
+		}
+
+		/// <inheritdoc/>
+		public string Name { get; } = "NoOp";
+	}
+}
diff --git a/Engine/Source/Programs/Horde/Horde.Build/Fleet/Autoscale/JobQueueStrategy.cs b/Engine/Source/Programs/Horde/Horde.Build/Fleet/Autoscale/JobQueueStrategy.cs
new file mode 100644
--- /dev/null
+++ b/Engine/Source/Programs/Horde/Horde.Build/Fleet/Autoscale/JobQueueStrategy.cs
@@ -0,0 +1,130 @@
+// Copyright Epic Games, Inc. All Rights Reserved.
+
+using System;
+using System.Collections.Generic;
+using System.Linq;
+using System.Threading.Tasks;
+using HordeCommon;
+using HordeServer.Api;
+using HordeServer.Collections;
+using HordeServer.Models;
+using HordeServer.Services;
+using HordeServer.Utilities;
+
+namespace Horde.Build.Fleet.Autoscale
+{
+	using PoolId = StringId<IPool>;
+	using StreamId = StringId<IStream>;
+	
+	/// <summary>
+	/// Calculate pool size by observing the number of jobs in waiting state
+	///
+	/// Allows for more proactive scaling compared to LeaseUtilizationStrategy.
+	/// <see cref="Horde.Build.Fleet.Autoscale.LeaseUtilizationStrategy"/> 
+	/// </summary>
+	public class JobQueueStrategy : IPoolSizeStrategy
+	{
+		
+		private readonly IJobCollection Jobs;
+		private readonly IGraphCollection Graphs;
+		private readonly StreamService StreamService;
+		private readonly IClock Clock;
+		private readonly TimeSpan LookbackPeriod = TimeSpan.FromHours(2.0);
+
+		/// <summary>
+		/// Constructor
+		/// </summary>
+		/// <param name="Jobs"></param>
+		/// <param name="Graphs"></param>
+		/// <param name="StreamService"></param>
+		/// <param name="Clock"></param>
+		/// <param name="LookbackPeriod">Time period for each sample</param>
+		public JobQueueStrategy(IJobCollection Jobs, IGraphCollection Graphs, StreamService StreamService, IClock Clock, TimeSpan? LookbackPeriod = null)
+		{
+			this.Jobs = Jobs;
+			this.Graphs = Graphs;
+			this.StreamService = StreamService;
+			this.Clock = Clock;
+			this.LookbackPeriod = LookbackPeriod ?? this.LookbackPeriod;
+		}
+
+		/// <inheritdoc/>
+		public string Name { get; } = "JobQueue";
+
+		/// <summary>
+		/// Extract all job step batches from a job, with their associated pool 
+		/// </summary>
+		/// <param name="Job">Job to extract from</param>
+		/// <param name="Streams">Cached lookup table of streams</param>
+		/// <returns></returns>
+		private async Task<List<(IJob Job, IJobStepBatch Batch, PoolId PoolId)>> GetJobBatchesWithPools(IJob Job, Dictionary<StreamId, IStream> Streams)
+		{
+			IGraph Graph = await Graphs.GetAsync(Job.GraphHash);
+
+			List<(IJob Job, IJobStepBatch Batch, PoolId PoolId)> JobBatches = new();
+			foreach (IJobStepBatch Batch in Job.Batches)
+			{
+				if (Batch.State != JobStepBatchState.Waiting) continue;
+
+				if (!Streams.TryGetValue(Job.StreamId, out IStream? Stream))
+				{
+					continue;
+				}
+
+				string BatchAgentType = Graph.Groups[Batch.GroupIdx].AgentType;
+				if (!Stream.AgentTypes.TryGetValue(BatchAgentType, out AgentType? AgentType))
+				{
+					continue;
+				}
+				
+				JobBatches.Add((Job, Batch, AgentType.Pool));
+			}
+
+			return JobBatches;
+		}
+
+		internal async Task<Dictionary<PoolId, int>> GetPoolQueueSizesAsync(DateTimeOffset JobsCreatedAfter)
+		{
+			List<IStream> StreamsList = await StreamService.GetStreamsAsync();
+			Dictionary<StreamId, IStream> Streams = StreamsList.ToDictionary(x => x.Id, x => x);
+			List<IJob> RecentJobs = await Jobs.FindAsync(MinCreateTime: JobsCreatedAfter);
+
+			List<(IJob Job, IJobStepBatch Batch, PoolId PoolId)> JobBatches = new();
+			foreach (IJob Job in RecentJobs)
+			{
+				JobBatches.AddRange(await GetJobBatchesWithPools(Job, Streams));
+			}
+
+			List<(PoolId PoolId, int QueueSize)> PoolsWithQueueSize = JobBatches.GroupBy(t => t.PoolId).Select(t => (t.Key, t.Count())).ToList();
+			return PoolsWithQueueSize.ToDictionary(x => x.PoolId, x => x.QueueSize);
+		}
+
+		/// <inheritdoc/>
+		public async Task CalcDesiredPoolSizesAsync(Dictionary<PoolId, PoolSizeData> PoolSizes)
+		{
+			DateTimeOffset MinCreateTime = Clock.UtcNow - LookbackPeriod;
+			Dictionary<PoolId, int> PoolQueueSizes = await GetPoolQueueSizesAsync(MinCreateTime);
+			
+			foreach ((PoolId PoolId, int QueueSize) in PoolQueueSizes.OrderBy(x => x.Value))
+			{
+				if (PoolSizes.TryGetValue(PoolId, out PoolSizeData? PoolSize))
+				{
+					// TODO: Tweak these values once in production
+					int AdditionalAgentCount = QueueSize switch
+					{
+						< 10 => 0,
+						<= 20 => 1,
+						<= 30 => 2,
+						<= 40 => 3,
+						<= 50 => 4,
+						<= 60 => 5,
+						_ => 4
+					};
+
+					int DesiredAgentCount = PoolSize.Agents.Count + AdditionalAgentCount;
+					PoolSizes[PoolId] = new(PoolSize.Pool, PoolSize.Agents, DesiredAgentCount, $"QueueSize={QueueSize}");
+				}
+			}
+		}
+	}
+}
diff --git a/Engine/Source/Programs/Horde/Horde.Build/Fleet/Autoscale/LeaseUtilizationStrategy.cs b/Engine/Source/Programs/Horde/Horde.Build/Fleet/Autoscale/LeaseUtilizationStrategy.cs
new file mode 100644
--- /dev/null
+++ b/Engine/Source/Programs/Horde/Horde.Build/Fleet/Autoscale/LeaseUtilizationStrategy.cs
@@ -0,0 +1,233 @@
+// Copyright Epic Games, Inc. All Rights Reserved.
+
+using System;
+using System.Collections.Generic;
+using System.Linq;
+using System.Text;
+using System.Threading.Tasks;
+using Google.Protobuf.WellKnownTypes;
+using HordeCommon;
+using HordeCommon.Rpc.Tasks;
+using HordeServer.Collections;
+using HordeServer.Models;
+using HordeServer.Utilities;
+using OpenTracing;
+using OpenTracing.Util;
+
+namespace Horde.Build.Fleet.Autoscale
+{
+	using PoolId = StringId<IPool>;
+	
+	/// <summary>
+	/// Calculate pool size by looking at previously finished leases
+	/// </summary>
+	public class LeaseUtilizationStrategy : IPoolSizeStrategy
+	{
+		struct UtilizationSample
+		{
+			public double JobWork;
+			public double OtherWork;
+		}
+
+		class AgentData
+		{
+			public IAgent Agent { get; }
+			public UtilizationSample[] Samples { get; }
+			public List<ILease> Leases { get; } = new ();
+			private readonly int NumSamples;
+
+			public AgentData(IAgent Agent, int NumSamples)
+			{
+				this.Agent = Agent;
+				this.NumSamples = NumSamples;
+				Samples = new UtilizationSample[NumSamples];
+			}
+
+			public void Add(double MinT, double MaxT, double JobWork, double OtherWork)
+			{
+				int MinIdx = Math.Clamp((int)MinT, 0, NumSamples - 1);
+				int MaxIdx = Math.Clamp((int)MaxT, MinIdx, NumSamples - 1);
+				for (int Idx = MinIdx; Idx <= MaxIdx; Idx++)
+				{
+					double Fraction = Math.Clamp(MaxT - Idx, 0.0, 1.0) - Math.Clamp(MinT - Idx, 0.0, 1.0);
+					Samples[Idx].JobWork += JobWork * Fraction;
+					Samples[Idx].OtherWork += OtherWork * Fraction;
+				}
+			}
+		}
+
+		class PoolData
+		{
+			public IPool Pool { get; }
+			public List<IAgent> Agents { get; } = new ();
+			public UtilizationSample[] Samples { get; }
+
+			public PoolData(IPool Pool, int NumSamples)
+			{
+				this.Pool = Pool;
+				Samples = new UtilizationSample[NumSamples];
+			}
+
+			public void Add(AgentData AgentData)
+			{
+				Agents.Add(AgentData.Agent);
+				for (int Idx = 0; Idx < Samples.Length; Idx++)
+				{
+					Samples[Idx].JobWork += AgentData.Samples[Idx].JobWork;
+					Samples[Idx].OtherWork += AgentData.Samples[Idx].OtherWork;
+				}
+			}
+		}
+		
+		private readonly IAgentCollection AgentCollection;
+		private readonly IPoolCollection PoolCollection;
+		private readonly ILeaseCollection LeaseCollection;
+		private readonly IClock Clock;
+		private readonly TimeSpan SampleTime = TimeSpan.FromMinutes(6.0);
+		private readonly int NumSamples = 10;
+		private readonly int NumSamplesForResult = 9;
+
+		/// <summary>
+		/// Constructor
+		/// </summary>
+		/// <param name="AgentCollection"></param>
+		/// <param name="PoolCollection"></param>
+		/// <param name="LeaseCollection"></param>
+		/// <param name="Clock"></param>
+		/// <param name="NumSamples">Number of samples to collect for calculating lease utilization</param>
+		/// <param name="NumSamplesForResult">Min number of samples for a valid result</param>
+		/// <param name="SampleTime">Time period for each sample</param>
+		public LeaseUtilizationStrategy(IAgentCollection AgentCollection, IPoolCollection PoolCollection, ILeaseCollection LeaseCollection, IClock Clock, int NumSamples = 10, int NumSamplesForResult = 9, TimeSpan? SampleTime = null)
+		{
+			this.AgentCollection = AgentCollection;
+			this.PoolCollection = PoolCollection;
+			this.LeaseCollection = LeaseCollection;
+			this.Clock = Clock;
+			this.NumSamples = NumSamples;
+			this.NumSamplesForResult = NumSamplesForResult;
+			this.SampleTime = SampleTime ?? this.SampleTime;
+		}
+
+		private async Task<Dictionary<AgentId, AgentData>> GetAgentDataAsync()
+		{
+			using IScope _ = GlobalTracer.Instance.BuildSpan("GetAgentDataAsync").StartActive();
+			
+			// Find all the current agents
+			List<IAgent> Agents = await AgentCollection.FindAsync(Status: AgentStatus.Ok);
+
+			// Query leases in last interval
+			DateTime MaxTime = Clock.UtcNow;
+			DateTime MinTime = MaxTime - (SampleTime * NumSamples);
+			List<ILease> Leases = await LeaseCollection.FindLeasesAsync(MinTime, MaxTime);
+
+			// Add all the leases to a data object for each agent
+			Dictionary<AgentId, AgentData> AgentIdToData = Agents.ToDictionary(x => x.Id, x => new AgentData(x, NumSamples));
+			foreach (ILease Lease in Leases)
+			{
+				AgentData? AgentData;
+				if (AgentIdToData.TryGetValue(Lease.AgentId, out AgentData) && AgentData.Agent.SessionId == Lease.SessionId)
+				{
+					AgentData.Leases.Add(Lease);
+				}
+			}
+			
+			// Compute utilization for each agent
+			foreach (AgentData AgentData in AgentIdToData.Values)
+			{
+				foreach (ILease Lease in AgentData.Leases.OrderBy(x => x.StartTime))
+				{
+					double MinT = (Lease.StartTime - MinTime).TotalSeconds / SampleTime.TotalSeconds;
+					double MaxT = (Lease.FinishTime == null)
+						? NumSamples
+						: ((Lease.FinishTime.Value - MinTime).TotalSeconds / SampleTime.TotalSeconds);
+
+					Any Payload = Any.Parser.ParseFrom(Lease.Payload.ToArray());
+					if (Payload.Is(ExecuteJobTask.Descriptor))
+					{
+						AgentData.Add(MinT, MaxT, 1.0, 0.0);
+					}
+					else
+					{
+						AgentData.Add(MinT, MaxT, 0.0, 1.0);
+					}
+				}
+			}
+
+			return AgentIdToData;
+		}
+		
+		private async Task<Dictionary<PoolId, PoolData>> GetPoolDataAsync()
+		{
+			using IScope _ = GlobalTracer.Instance.BuildSpan("GetPoolDataAsync").StartActive();
+
+			Dictionary<AgentId, AgentData> AgentIdToData = await GetAgentDataAsync();
+			
+			// Get all the pools
+			List<IPool> Pools = await PoolCollection.GetAsync();
+			Dictionary<PoolId, PoolData> PoolToData = Pools.ToDictionary(x => x.Id, x => new PoolData(x, NumSamples));
+
+			// Find pool utilization over the query period
+			foreach (AgentData AgentData in AgentIdToData.Values)
+			{
+				foreach (PoolId PoolId in AgentData.Agent.GetPools())
+				{
+					PoolData? PoolData;
+					if (PoolToData.TryGetValue(PoolId, out PoolData))
+					{
+						PoolData.Add(AgentData);
+					}
+				}
+			}
+
+			return PoolToData;
+		}
+
+		/// <inheritdoc/>
+		public string Name { get; } = "LeaseUtilization";
+
+		/// <inheritdoc/>
+		public async Task CalcDesiredPoolSizesAsync(Dictionary<PoolId, PoolSizeData> PoolSizes)
+		{
+			Dictionary<PoolId, PoolData> PoolToData = await GetPoolDataAsync();
+
+			foreach (PoolData PoolData in PoolToData.Values.OrderByDescending(x => x.Agents.Count))
+			{
+				IPool Pool = PoolData.Pool;
+
+				if (PoolSizes.TryGetValue(Pool.Id, out PoolSizeData? PoolSize))
+				{
+					int MinAgents = Pool.MinAgents ?? 1;
+					int NumReserveAgents = Pool.NumReserveAgents ?? 5;
+					double Utilization = PoolData.Samples.Select(x => x.JobWork).OrderByDescending(x => x).Skip(NumSamples - NumSamplesForResult).First();
+				
+					// Number of agents in use over the sampling period. Can never be greater than number of agents available in pool.
+					int NumAgentsUtilized = (int)Utilization;
+					
+					// Include reserve agent count to ensure pool always can grow
+					int DesiredAgentCount = Math.Max(NumAgentsUtilized + NumReserveAgents, MinAgents);;
+					
+					StringBuilder Sb = new();
+					Sb.AppendFormat("Jobs=[{0}] ", GetDensityMap(PoolData.Samples.Select(x => x.JobWork / Math.Max(1, PoolData.Agents.Count))));
+					Sb.AppendFormat("Total=[{0}] ", GetDensityMap(PoolData.Samples.Select(x => x.OtherWork / Math.Max(1, PoolData.Agents.Count))));
+					Sb.AppendFormat("Min=[{0,5:0.0}] ", PoolData.Samples.Min(x => x.JobWork));
+					Sb.AppendFormat("Max=[{0,5:0.0}] ", PoolData.Samples.Max(x => x.JobWork));
+					Sb.AppendFormat("Avg=[{0,5:0.0}] ", Utilization);
+					Sb.AppendFormat("Pct=[{0,5:0.0}] ", PoolData.Samples.Sum(x => x.JobWork) / NumSamples);
+
+					PoolSizes[Pool.Id] = new(Pool, PoolSize.Agents, DesiredAgentCount, Sb.ToString());
+				}
+			}
+		}
+		
+		/// <summary>
+		/// Creates a string of characters indicating a sequence of 0-1 values over time
+		/// </summary>
+		/// <param name="Values">Sequence of values</param>
+		/// <returns>Density map</returns>
+		private static string GetDensityMap(IEnumerable<double> Values)
+		{
+			const string Greyscale = " 123456789";
+			return new string(Values.Select(x => Greyscale[Math.Clamp((int)(x * Greyscale.Length), 0, Greyscale.Length - 1)]).ToArray());
+		}
+	}
+}
diff --git a/Engine/Source/Programs/Horde/Horde.Build/Fleet/AutoscaleService.cs b/Engine/Source/Programs/Horde/Horde.Build/Fleet/AutoscaleService.cs
deleted file mode 100644
diff --git a/Engine/Source/Programs/Horde/Horde.Build/ServerSettings.cs b/Engine/Source/Programs/Horde/Horde.Build/ServerSettings.cs
--- a/Engine/Source/Programs/Horde/Horde.Build/ServerSettings.cs
+++ b/Engine/Source/Programs/Horde/Horde.Build/ServerSettings.cs
@@ -504,6 +504,11 @@ namespace HordeServer
 		/// The Uri for the Jira installation
 		/// </summary>
 		public Uri? JiraUrl { get; set; }
+		
+		/// <summary>
+		/// Type of auto-scaling strategy to be used for sizing agent pools (if left empty, legacy auto scaling impl is used)
+		/// </summary>
+		public string? AgentPoolSizeStrategy { get; set; }
 
 		/// <summary>
 		/// Set the minimum size of the global thread pool
diff --git a/Engine/Source/Programs/Horde/Horde.Build/Startup.cs b/Engine/Source/Programs/Horde/Horde.Build/Startup.cs
--- a/Engine/Source/Programs/Horde/Horde.Build/Startup.cs
+++ b/Engine/Source/Programs/Horde/Horde.Build/Startup.cs
@@ -82,6 +82,7 @@ using HordeServer.Compute.Impl;
 using HordeServer.Compute;
 using System.Net.Http.Headers;
 using Amazon.SecurityToken.Model;
+using Horde.Build.Fleet.Autoscale;
 using Horde.Build.Utilities;
 using Serilog.Events;
 using HordeServer.Jobs;
@@ -334,6 +335,21 @@ namespace HordeServer
 			Services.AddSingleton(typeof(ISingletonDocument<>), typeof(SingletonDocument<>));
 
 			Services.AddSingleton<AutoscaleService>();
+
+			if (Enum.TryParse(Settings.AgentPoolSizeStrategy, out PoolSizeStrategy PoolSizeStrategy))
+			{
+				switch (PoolSizeStrategy)
+				{
+					case PoolSizeStrategy.LeaseUtilization: Services.AddSingleton<IPoolSizeStrategy, LeaseUtilizationStrategy>(); break;
+					case PoolSizeStrategy.JobQueue: Services.AddSingleton<IPoolSizeStrategy, JobQueueStrategy>(); break;
+					default: Services.AddSingleton<IPoolSizeStrategy, NoOpPoolSizeStrategy>(); break;
+				}
+			}
+			else
+			{
+				Services.AddSingleton<IPoolSizeStrategy, NoOpPoolSizeStrategy>();
+			}
+			
 			switch (Settings.FleetManager)
 			{
 				case FleetManagerType.Aws:
@@ -557,8 +573,17 @@ namespace HordeServer
 
 			if (Settings.EnableBackgroundServices)
 			{
+				// Allow both legacy and v2 impl to co-exist during rollout phase
+				if (Settings.AgentPoolSizeStrategy != null)
+				{
+					Services.AddHostedService(Provider => Provider.GetRequiredService<AutoscaleServiceV2>());
+				}
+				else
+				{
+					Services.AddHostedService(Provider => Provider.GetRequiredService<AutoscaleService>());	
+				}
+				
 				Services.AddHostedService(Provider => Provider.GetRequiredService<AgentService>());
-				Services.AddHostedService(Provider => Provider.GetRequiredService<AutoscaleService>());
 				Services.AddHostedService(Provider => Provider.GetRequiredService<CommitService>());
 				Services.AddHostedService(Provider => Provider.GetRequiredService<ConsistencyService>());
 				Services.AddHostedService(Provider => (DowntimeService)Provider.GetRequiredService<IDowntimeService>());
