Horde Storage - fixes to the polling service avoiding loops inside the timer as well as a re-entrance protection

diff --git a/Engine/Source/Programs/Horde/HordeStorage/Jupiter.Common/PollingService.cs b/Engine/Source/Programs/Horde/HordeStorage/Jupiter.Common/PollingService.cs
--- a/Engine/Source/Programs/Horde/HordeStorage/Jupiter.Common/PollingService.cs
+++ b/Engine/Source/Programs/Horde/HordeStorage/Jupiter.Common/PollingService.cs
@@ -26,6 +26,7 @@ namespace Jupiter
         private readonly T _state;
         private readonly CancellationTokenSource _stopPolling = new CancellationTokenSource();
         private readonly ManualResetEvent _hasFinishedRunning = new ManualResetEvent(true);
+        private volatile bool _alreadyRunning = false;
         private Timer? _timer;
 
         protected PollingService(string serviceName, TimeSpan pollFrequency, T state)
@@ -86,28 +87,37 @@ namespace Jupiter
             CancellationToken stopPollingToken = threadState.StopPollingToken;
             ILogger logger = Log.ForContext<PollingService<T>>();
 
-            instance._hasFinishedRunning.Reset();
+            if (instance._alreadyRunning)
+                return;
 
-            while (!stopPollingToken.IsCancellationRequested)
+            try
             {
-                try
-                {
-                    bool _ = instance.OnPoll(threadState.ServiceState, stopPollingToken).Result;
-                }
-                catch (AggregateException e)
-                {
-                    logger.Error(e, "{Service} Aggregate exception in polling service", serviceName);
-                    foreach (Exception inner in e.InnerExceptions)
-                    {
-                        logger.Error(inner, "{Service} inner exception in polling service. Trace: {StackTrace}", serviceName, inner.StackTrace);
-                        
-                    }
-                }
-                catch (Exception e)
+                instance._alreadyRunning = true;
+                instance._hasFinishedRunning.Reset();
+
+                if (stopPollingToken.IsCancellationRequested)
+                    return;
+
+                bool _ = instance.OnPoll(threadState.ServiceState, stopPollingToken).Result;
+            }
+            catch (AggregateException e)
+            {
+                logger.Error(e, "{Service} Aggregate exception in polling service", serviceName);
+                foreach (Exception inner in e.InnerExceptions)
                 {
-                    logger.Error(e, "{Service} Exception in polling service", serviceName);
+                    logger.Error(inner, "{Service} inner exception in polling service. Trace: {StackTrace}",
+                        serviceName, inner.StackTrace);
+
                 }
             }
+            catch (Exception e)
+            {
+                logger.Error(e, "{Service} Exception in polling service", serviceName);
+            }
+            finally
+            {
+                instance._alreadyRunning = false;
+            }
         }
 
         public abstract Task<bool> OnPoll(T state, CancellationToken cancellationToken);
