Horde: Add template revision information to advanced job view
#skipci
#rnx

diff --git a/Engine/Source/Programs/Horde/HordeDashboard/src/backend/Api.ts b/Engine/Source/Programs/Horde/HordeDashboard/src/backend/Api.ts
--- a/Engine/Source/Programs/Horde/HordeDashboard/src/backend/Api.ts
+++ b/Engine/Source/Programs/Horde/HordeDashboard/src/backend/Api.ts
@@ -1788,6 +1788,12 @@ export type GetStreamResponse = {
 	/**Name of the stream */
 	name: string;
 
+	/**The config file path on the server*/
+	configPath?:string;
+
+	/**Revision of the config file */
+	configRevision?:string;
+
 	/**List of tabs to display for this stream*/
 	tabs: GetStreamTabResponse[];
 
