Horde: Fix build errors.

diff --git a/Engine/Source/Programs/Horde/Horde.Build/Utilities/ObjectSet.cs b/Engine/Source/Programs/Horde/Horde.Build/Utilities/ObjectSet.cs
--- a/Engine/Source/Programs/Horde/Horde.Build/Utilities/ObjectSet.cs
+++ b/Engine/Source/Programs/Horde/Horde.Build/Utilities/ObjectSet.cs
@@ -64,7 +64,7 @@ namespace HordeServer.Utilities
 	/// <summary>
 	/// Converter for BlobPackIndex objects
 	/// </summary>
-	class ObjectPackIndexConverter : CbConverter<ObjectPackIndex>
+	class ObjectPackIndexConverter : CbConverterBase<ObjectPackIndex>
 	{
 		class EncodeFormat
 		{
diff --git a/Engine/Source/Programs/Horde/Horde.Build/Utilities/StringId.cs b/Engine/Source/Programs/Horde/Horde.Build/Utilities/StringId.cs
--- a/Engine/Source/Programs/Horde/Horde.Build/Utilities/StringId.cs
+++ b/Engine/Source/Programs/Horde/Horde.Build/Utilities/StringId.cs
@@ -285,7 +285,7 @@ namespace HordeServer.Utilities
 		}
 	}
 
-	sealed class CbStringIdConverter<T> : CbConverter<StringId<T>>
+	sealed class CbStringIdConverter<T> : CbConverterBase<StringId<T>>
 	{
 		public override StringId<T> Read(CbField Field)
 		{
