Horde Storage - Adding a concept of default namespace policy, allowing us to accept any namespace when running locally with auth disabled. Could be used in production to accept any namespace as well but I do not recommend it.

diff --git a/Engine/Source/Programs/Horde/HordeStorage/Horde.Storage/appsettings.Development.json b/Engine/Source/Programs/Horde/HordeStorage/Horde.Storage/appsettings.Development.json
--- a/Engine/Source/Programs/Horde/HordeStorage/Horde.Storage/appsettings.Development.json
+++ b/Engine/Source/Programs/Horde/HordeStorage/Horde.Storage/appsettings.Development.json
@@ -102,7 +102,10 @@
       },
       "ue.ddc": {
         "Claims": [ "UE4" ],
-        "StoragePool": "ddc2" 
+        "StoragePool": "ddc2"
+      },
+      "*": {
+        "Claims": [ "*" ]
       }
     }
   },
diff --git a/Engine/Source/Programs/Horde/HordeStorage/Jupiter.Common/Implementation/NamespacePolicyResolver.cs b/Engine/Source/Programs/Horde/HordeStorage/Jupiter.Common/Implementation/NamespacePolicyResolver.cs
--- a/Engine/Source/Programs/Horde/HordeStorage/Jupiter.Common/Implementation/NamespacePolicyResolver.cs
+++ b/Engine/Source/Programs/Horde/HordeStorage/Jupiter.Common/Implementation/NamespacePolicyResolver.cs
@@ -1,5 +1,6 @@
 // Copyright Epic Games, Inc. All Rights Reserved.
 
+using System;
 using Jupiter.Implementation;
 using Microsoft.Extensions.Options;
 
@@ -35,7 +36,20 @@ namespace Jupiter.Common
                 };
             }
 
-            return _namespaceSettings.CurrentValue.Policies[ns.ToString()];
+            if (_namespaceSettings.CurrentValue.Policies.TryGetValue(ns.ToString(),
+                    out NamespaceSettings.PerNamespaceSettings? settings))
+            {
+                return settings;
+            }
+
+            // attempt to find the default mapping
+            if (_namespaceSettings.CurrentValue.Policies.TryGetValue("*",
+                    out NamespaceSettings.PerNamespaceSettings? defaultSettings))
+            {
+                return defaultSettings;
+            }
+
+            throw new ArgumentException($"Unable to find a valid policy for namespace {ns}");
         }
     }
 }
