Horde Storage - Added memory backed ref store to allow Horde Storage to run locally without any external dependencies

diff --git a/Engine/Source/Programs/Horde/HordeStorage/Horde.Storage/Horde.Storage.Startup.cs b/Engine/Source/Programs/Horde/HordeStorage/Horde.Storage/Horde.Storage.Startup.cs
--- a/Engine/Source/Programs/Horde/HordeStorage/Horde.Storage/Horde.Storage.Startup.cs
+++ b/Engine/Source/Programs/Horde/HordeStorage/Horde.Storage/Horde.Storage.Startup.cs
@@ -538,6 +538,7 @@ namespace Horde.Storage
 
             IRefsStore refsStore = settings.RefDbImplementation switch
             {
+                HordeStorageSettings.RefDbImplementations.Memory => ActivatorUtilities.CreateInstance<MemoryRefsStore>(provider),
                 HordeStorageSettings.RefDbImplementations.DynamoDb => ActivatorUtilities.CreateInstance<DynamoDbRefsStore>(provider),
                 HordeStorageSettings.RefDbImplementations.Mongo => ActivatorUtilities.CreateInstance<MongoRefsStore>(provider),
                 HordeStorageSettings.RefDbImplementations.Cosmos => ActivatorUtilities.CreateInstance<CosmosRefsStore>(provider),
@@ -561,6 +562,8 @@ namespace Horde.Storage
 
             switch (settings.RefDbImplementation)
             {
+                case HordeStorageSettings.RefDbImplementations.Memory:
+                    break;
                 case HordeStorageSettings.RefDbImplementations.Mongo:
                 case HordeStorageSettings.RefDbImplementations.Cosmos:
                     MongoSettings mongoSettings = provider.GetService<IOptionsMonitor<MongoSettings>>()!.CurrentValue;
diff --git a/Engine/Source/Programs/Horde/HordeStorage/Horde.Storage/Implementation/Refs/IRefsStore.cs b/Engine/Source/Programs/Horde/HordeStorage/Horde.Storage/Implementation/Refs/IRefsStore.cs
--- a/Engine/Source/Programs/Horde/HordeStorage/Horde.Storage/Implementation/Refs/IRefsStore.cs
+++ b/Engine/Source/Programs/Horde/HordeStorage/Horde.Storage/Implementation/Refs/IRefsStore.cs
@@ -29,7 +29,7 @@ namespace Horde.Storage.Implementation
 
         public BlobIdentifier[] Blobs { get; }
         public Dictionary<string, object>? Metadata { get; }
-        public DateTime? LastAccessTime { get; }
+        public DateTime? LastAccessTime { get; internal set; }
         public ContentHash ContentHash { get; }
 
         public TransactionEvent ToAddTransactionEvent()
diff --git a/Engine/Source/Programs/Horde/HordeStorage/Horde.Storage/Implementation/Refs/MemoryRefsStore.cs b/Engine/Source/Programs/Horde/HordeStorage/Horde.Storage/Implementation/Refs/MemoryRefsStore.cs
new file mode 100644
--- /dev/null
+++ b/Engine/Source/Programs/Horde/HordeStorage/Horde.Storage/Implementation/Refs/MemoryRefsStore.cs
@@ -0,0 +1,113 @@
+// Copyright Epic Games, Inc. All Rights Reserved.
+
+using System;
+using System.Collections.Generic;
+using System.Threading.Tasks;
+using Dasync.Collections;
+using Jupiter.Implementation;
+using Serilog;
+
+using RefContainer = System.Collections.Concurrent.ConcurrentDictionary<Jupiter.Implementation.KeyId, Horde.Storage.Implementation.RefRecord>;
+using BucketContainer = System.Collections.Concurrent.ConcurrentDictionary<Jupiter.Implementation.BucketId, System.Collections.Concurrent.ConcurrentDictionary<Jupiter.Implementation.KeyId, Horde.Storage.Implementation.RefRecord>>;
+using NamespaceContainer = System.Collections.Concurrent.ConcurrentDictionary<Jupiter.Implementation.NamespaceId, System.Collections.Concurrent.ConcurrentDictionary<Jupiter.Implementation.BucketId, System.Collections.Concurrent.ConcurrentDictionary<Jupiter.Implementation.KeyId, Horde.Storage.Implementation.RefRecord>>>;
+
+namespace Horde.Storage.Implementation
+{
+    internal class MemoryRefsStore : IRefsStore
+    {
+        private readonly ILogger _logger = Log.ForContext<MemoryRefsStore>();
+        private readonly NamespaceContainer _refs = new NamespaceContainer();
+
+        public Task<RefRecord?> Get(NamespaceId ns, BucketId bucket, KeyId key, IRefsStore.ExtraFieldsFlag fields)
+        {
+            BucketContainer bucketContainer = _refs.GetOrAdd(ns, new BucketContainer());
+            RefContainer refContainer = bucketContainer.GetOrAdd(bucket, new RefContainer());
+            if (refContainer.TryGetValue(key, out RefRecord? refObject))
+            {
+                return Task.FromResult<RefRecord?>(refObject);
+            }
+
+            return Task.FromResult<RefRecord?>(null);
+        }
+
+        public Task Add(RefRecord record)
+        {
+            BucketContainer bucketContainer = _refs.GetOrAdd(record.Namespace, new BucketContainer());
+            RefContainer refContainer = bucketContainer.GetOrAdd(record.Bucket, new RefContainer());
+            refContainer.AddOrUpdate(record.RefName, id => record, (id, o) => record);
+
+            return Task.CompletedTask;
+        }
+
+        public Task<long> DeleteBucket(NamespaceId ns, BucketId bucket)
+        {
+            BucketContainer bucketContainer = _refs.GetOrAdd(ns, new BucketContainer());
+            if (bucketContainer.TryRemove(bucket, out RefContainer? refContainer))
+            {
+                return Task.FromResult<long>(refContainer.Keys.Count);
+            }
+
+            return Task.FromResult(0L);
+        }
+
+        public Task<long> Delete(NamespaceId ns, BucketId bucket, KeyId key)
+        {
+            BucketContainer bucketContainer = _refs.GetOrAdd(ns, new BucketContainer());
+            RefContainer refContainer = bucketContainer.GetOrAdd(bucket, new RefContainer());
+
+            if (refContainer.TryRemove(key, out RefRecord? refRecord))
+            {
+                return Task.FromResult(1L);
+            }
+
+            return Task.FromResult(0L);
+        }
+
+        public Task UpdateLastAccessTime(RefRecord record, DateTime lastAccessTime)
+        {
+            BucketContainer bucketContainer = _refs.GetOrAdd(record.Namespace, new BucketContainer());
+            RefContainer refContainer = bucketContainer.GetOrAdd(record.Bucket, new RefContainer());
+            if (refContainer.TryGetValue(record.RefName, out RefRecord? refRecord))
+            {
+                refRecord.LastAccessTime = lastAccessTime;
+            }
+
+            return Task.CompletedTask;
+        }
+
+        public IAsyncEnumerable<NamespaceId> GetNamespaces()
+        {
+            return _refs.Keys.ToAsyncEnumerable();
+        }
+
+        public async IAsyncEnumerable<OldRecord> GetOldRecords(NamespaceId ns, TimeSpan oldRecordCutoff)
+        {
+            DateTime cutoff = DateTime.Now.AddSeconds(-1 * oldRecordCutoff.TotalSeconds);
+            BucketContainer bucketContainer = _refs.GetOrAdd(ns, new BucketContainer());
+
+            await Task.CompletedTask;
+
+            foreach (BucketId bucket in bucketContainer.Keys)
+            {
+                RefContainer refContainer = bucketContainer[bucket];
+                foreach (KeyId keyId in refContainer.Keys)
+                {
+                    RefRecord refRecord = refContainer[keyId];
+                    if (refRecord.LastAccessTime < cutoff)
+                        yield return new OldRecord(ns, bucket, keyId);
+                }
+            }
+        }
+
+        public Task DropNamespace(NamespaceId ns)
+        {
+            _refs.TryRemove(ns, out BucketContainer? _);
+
+            return Task.CompletedTask;
+        }
+    }
+
+    class RefObject
+    {
+    }
+}
diff --git a/Engine/Source/Programs/Horde/HordeStorage/Horde.Storage/Settings.cs b/Engine/Source/Programs/Horde/HordeStorage/Horde.Storage/Settings.cs
--- a/Engine/Source/Programs/Horde/HordeStorage/Horde.Storage/Settings.cs
+++ b/Engine/Source/Programs/Horde/HordeStorage/Horde.Storage/Settings.cs
@@ -14,10 +14,10 @@ namespace Horde.Storage
     {
         public enum RefDbImplementations
         {
+            Memory,
             Mongo,
             Cosmos,
             DynamoDb
-
         }
 
         public enum TransactionLogWriterImplementations
@@ -109,7 +109,7 @@ namespace Horde.Storage
 
         [Required] public ReplicationLogWriterImplementations ReplicationLogWriterImplementation { get; set; } = ReplicationLogWriterImplementations.Memory;
 
-        [Required] public RefDbImplementations RefDbImplementation { get; set; } = RefDbImplementations.Mongo;
+        [Required] public RefDbImplementations RefDbImplementation { get; set; } = RefDbImplementations.Memory;
 
         [Required]
         public TreeStoreImplementations TreeStoreImplementation { get; set; } = TreeStoreImplementations.Memory;
diff --git a/Engine/Source/Programs/Horde/HordeStorage/Horde.Storage/Tests/Functional/RefTests.cs b/Engine/Source/Programs/Horde/HordeStorage/Horde.Storage/Tests/Functional/RefTests.cs
--- a/Engine/Source/Programs/Horde/HordeStorage/Horde.Storage/Tests/Functional/RefTests.cs
+++ b/Engine/Source/Programs/Horde/HordeStorage/Horde.Storage/Tests/Functional/RefTests.cs
@@ -88,6 +88,49 @@ namespace Horde.Storage.FunctionalTests.Ref
         }
     }
 
+    [TestClass]
+    public class MemoryRefTests : RefTests
+    {
+        private MemoryRefsStore? _refsStore;
+
+        protected override async Task SeedDb(IServiceProvider provider)
+        {
+            //verify we are using the expected refs store
+            IRefsStore? refStore = provider.GetService<IRefsStore>();
+            Assert.IsTrue(RefStoreIs(refStore, typeof(MemoryRefsStore)));
+            if (refStore is CachedRefStore cachedRefStore)
+            {
+                _refsStore = (MemoryRefsStore)cachedRefStore.BackingStore;
+            }
+            else
+            {
+                _refsStore = (MemoryRefsStore)refStore!;
+            }
+
+            IBlobService blobStore = (IBlobService)provider.GetService(typeof(IBlobService))!;
+
+            await SeedRefTestData(_refsStore, blobStore);
+
+        }
+
+        protected override Task TeardownDb(IServiceProvider provider)
+        {
+            return Task.CompletedTask;
+        }
+
+        protected override async Task<RefRecord> GetTestRecord(BucketId bucket, KeyId name, IRefsStore.ExtraFieldsFlag fields)
+        {
+            var record = await _refsStore!.Get(TestNamespace, bucket, name, fields);
+            if (record == null)
+                throw new Exception("Unable to find record");
+
+            return record;
+        }
+        protected override IEnumerable<KeyValuePair<string, string>> GetSettings()
+        {
+            return new[] {new KeyValuePair<string, string>("Horde_Storage:RefDbImplementation", HordeStorageSettings.RefDbImplementations.Memory.ToString())};
+        }
+    }
 
     [TestClass]
     public class DynamoRefTests : RefTests
diff --git a/Engine/Source/Programs/Horde/HordeStorage/Horde.Storage/appsettings.Development.json b/Engine/Source/Programs/Horde/HordeStorage/Horde.Storage/appsettings.Development.json
--- a/Engine/Source/Programs/Horde/HordeStorage/Horde.Storage/appsettings.Development.json
+++ b/Engine/Source/Programs/Horde/HordeStorage/Horde.Storage/appsettings.Development.json
@@ -18,7 +18,7 @@
   },
 
   "Horde_Storage": {
-    "refDBImplementation": "mongo",
+    "refDBImplementation": "memory",
 
     "TransactionLogWriterImplementation": "memory",
 
