Horde Storage - Simplified the health checks to not rely on IHealthCheck - as they will abort if someone is requesting it while another request is in flight.

The old health checks are now moved to a seperate set of endpoints, and I switched our kubernetes configuration over to use that, as noone else should be calling these endpoints now.

The expectation is that this fixes issues we see as Horde Storage scales up and pods starts getting shut down due to failed health checks even though the load doesnt look high enough for that to really happen.

diff --git a/Engine/Source/Programs/Horde/HordeStorage/Jupiter.Common/BaseStartup.cs b/Engine/Source/Programs/Horde/HordeStorage/Jupiter.Common/BaseStartup.cs
--- a/Engine/Source/Programs/Horde/HordeStorage/Jupiter.Common/BaseStartup.cs
+++ b/Engine/Source/Programs/Horde/HordeStorage/Jupiter.Common/BaseStartup.cs
@@ -307,17 +307,30 @@ namespace Jupiter
                 bool PassAllChecks(HealthCheckRegistration check) => true;
 
                 // Ready checks in Kubernetes is to verify that the service is working, if this returns false the app will not get any traffic (load balancer ignores it)
-                endpoints.MapHealthChecks("/health/ready", options: new HealthCheckOptions()
+                endpoints.MapHealthChecks("/health/readiness", options: new HealthCheckOptions()
                 {
                     Predicate = jupiterSettings.DisableHealthChecks ? PassAllChecks : (check) => check.Tags.Contains("self"),
                 });
 
                 // Live checks in Kubernetes to see if the pod is working as it should, if this returns false the entire pod is killed
-                endpoints.MapHealthChecks("/health/live", options: new HealthCheckOptions()
+                endpoints.MapHealthChecks("/health/liveness", options: new HealthCheckOptions()
                 {
                     Predicate = jupiterSettings.DisableHealthChecks ? PassAllChecks : (check) => check.Tags.Contains("services"),
                 });
 
+                endpoints.MapGet("/health/ready", async context =>
+                {
+                    context.Response.StatusCode = 200;
+                    context.Response.Headers.ContentType = "text/plain";
+                    await context.Response.Body.WriteAsync(Encoding.ASCII.GetBytes("Healthy"));
+                });
+
+                endpoints.MapGet("/health/live", async context =>
+                {
+                    context.Response.StatusCode = 200;
+                    context.Response.Headers.ContentType = "text/plain";
+                    await context.Response.Body.WriteAsync(Encoding.ASCII.GetBytes("Healthy"));
+                });
                 OnUseEndpoints(env, endpoints);
                 
                 endpoints.MapControllers();
