Fix bug where lease finish time being null, therefore still running, fails the conditional check to be counted into the current bucket

diff --git a/Engine/Source/Programs/Horde/Horde.Build/Fleet/TelemetryService.cs b/Engine/Source/Programs/Horde/Horde.Build/Fleet/TelemetryService.cs
--- a/Engine/Source/Programs/Horde/Horde.Build/Fleet/TelemetryService.cs
+++ b/Engine/Source/Programs/Horde/Horde.Build/Fleet/TelemetryService.cs
@@ -109,7 +109,7 @@ namespace HordeServer.Services
 				}
 				foreach (ILease Lease in Leases)
 				{
-					if (Lease.StartTime < BucketMaxTime && Lease.FinishTime >= BucketMinTime)
+					if (Lease.StartTime < BucketMaxTime && (!Lease.FinishTime.HasValue || Lease.FinishTime >= BucketMinTime))
 					{
 						List<PoolId>? LeasePools;
 						if (AgentToPoolIds.TryGetValue(Lease.AgentId, out LeasePools))
