Horde: Add ability to queue notifications for batch sending later

- Send any queued notifications at a hard-coded 12 hour interval
- Only JobScheduled notifications makes use of this they were too frequent before

diff --git a/Engine/Source/Programs/Horde/Horde.Build.Tests/NotificationServiceTests.cs b/Engine/Source/Programs/Horde/Horde.Build.Tests/NotificationServiceTests.cs
--- a/Engine/Source/Programs/Horde/Horde.Build.Tests/NotificationServiceTests.cs
+++ b/Engine/Source/Programs/Horde/Horde.Build.Tests/NotificationServiceTests.cs
@@ -1,22 +1,49 @@
 // Copyright Epic Games, Inc. All Rights Reserved.
 
+using System;
 using System.Collections.Generic;
 using System.Threading.Tasks;
+using HordeCommon;
 using HordeServer.Api;
 using HordeServer.Models;
+using HordeServer.Notifications;
+using HordeServer.Notifications.Impl;
 using HordeServer.Utilities;
+using Microsoft.Extensions.DependencyInjection;
+using Microsoft.Extensions.Logging;
 using Microsoft.VisualStudio.TestTools.UnitTesting;
-using MongoDB.Bson;
 using Moq;
-
+using StackExchange.Redis;
 using StreamId = HordeServer.Utilities.StringId<HordeServer.Models.IStream>;
 
 namespace HordeServerTests
 {
 	using JobId = ObjectId<IJob>;
 
+	public class FakeNotificationSink : INotificationSink
+	{
+		public List<JobScheduledNotification> JobScheduledNotifications = new();
+		public int JobScheduledCallCount;
+		
+		public Task NotifyJobScheduledAsync(List<JobScheduledNotification> Notifications)
+		{
+			JobScheduledNotifications.AddRange(Notifications);
+			JobScheduledCallCount++;
+			return Task.CompletedTask;
+		}
+
+		public Task NotifyJobCompleteAsync(IStream JobStream, IJob Job, IGraph Graph, LabelOutcome Outcome) { throw new NotImplementedException(); }
+		public Task NotifyJobCompleteAsync(IUser User, IStream JobStream, IJob Job, IGraph Graph, LabelOutcome Outcome) { throw new NotImplementedException(); }
+		public Task NotifyJobStepCompleteAsync(IUser User, IStream JobStream, IJob Job, IJobStepBatch Batch, IJobStep Step, INode Node, List<ILogEventData> JobStepEventData) { throw new NotImplementedException(); }
+		public Task NotifyLabelCompleteAsync(IUser User, IJob Job, IStream Stream, ILabel Label, int LabelIdx, LabelOutcome Outcome, List<(string, JobStepOutcome, Uri)> StepData) { throw new NotImplementedException(); }
+		public Task NotifyIssueUpdatedAsync(IIssue Issue) { throw new NotImplementedException(); }
+		public Task NotifyConfigUpdateFailureAsync(string ErrorMessage, string FileName, int? Change = null, IUser? Author = null, string? Description = null) { throw new NotImplementedException(); }
+		public Task NotifyDeviceServiceAsync(string Message, IDevice? Device = null, IDevicePool? Pool = null, IStream? Stream = null, IJob? Job = null, IJobStep? Step = null, INode? Node = null) { throw new NotImplementedException(); }
+	}
+	
+
 	[TestClass]
-	public class NotificationServiceTests
+	public class NotificationServiceTests : TestSetup
 	{
 		public IJob CreateJob(StreamId StreamId, int Change, string Name, IGraph Graph)
 		{
@@ -56,13 +83,40 @@ namespace HordeServerTests
 			Job.SetupGet(x => x.Batches).Returns(Batches);
 			return Job.Object;
 		}
-
+		
 		[TestMethod]
-		[Ignore]
-		public Task NotifyLabelUpdateTests()
+		public async Task NotifyJobScheduled()
 		{
-			// #1 
-			return Task.CompletedTask;
+			FakeNotificationSink FakeSink = new();
+			NotificationService Service = GetNotificationService(FakeSink);
+			Fixture Fixture = await CreateFixtureAsync();
+			IPool Pool = await PoolService.CreatePoolAsync("BogusPool", Properties: new Dictionary<string, string>());
+
+			Assert.AreEqual(0, FakeSink.JobScheduledNotifications.Count);
+			Service.NotifyJobScheduled(Pool, false, Fixture.Job1, Fixture.Graph, SubResourceId.Random());
+			Service.NotifyJobScheduled(Pool, false, Fixture.Job2, Fixture.Graph, SubResourceId.Random());
+			
+			// Currently no good way to wait for NotifyJobScheduled() to complete as the execution is completely async in background task (see ExecuteAsync)
+			await Task.Delay(1000);
+			await Clock.AdvanceAsync(Service.NotificationBatchInterval + TimeSpan.FromMinutes(5));
+			Assert.AreEqual(2, FakeSink.JobScheduledNotifications.Count);
+			Assert.AreEqual(1, FakeSink.JobScheduledCallCount);
+		}
+
+		private NotificationService GetNotificationService(INotificationSink? Sink)
+		{
+			ILogger<NotificationService> Logger = ServiceProvider.GetRequiredService<ILogger<NotificationService>>();
+			IDatabase Database = GetRedisConnectionPool().GetDatabase();
+			List<INotificationSink> Sinks = new ();
+			if (Sink != null) Sinks.Add(Sink);
+			
+			NotificationService Service = new NotificationService(
+				Sinks, ServerSettingsMon, Logger, GraphCollection, SubscriptionCollection,
+				NotificationTriggerCollection, UserCollection, JobService, StreamService, IssueService, LogFileService,
+				new NoOpDogStatsd(), Database, Clock);
+
+			Service.Ticker.StartAsync().Wait(5000);
+			return Service;
 		}
 
 		//public void NotifyLabelUpdate(IJob Job, IReadOnlyList<(LabelState, LabelOutcome)> OldLabelStates, IReadOnlyList<(LabelState, LabelOutcome)> NewLabelStates)
diff --git a/Engine/Source/Programs/Horde/Horde.Build/Notifications/INotificationService.cs b/Engine/Source/Programs/Horde/Horde.Build/Notifications/INotificationService.cs
--- a/Engine/Source/Programs/Horde/Horde.Build/Notifications/INotificationService.cs
+++ b/Engine/Source/Programs/Horde/Horde.Build/Notifications/INotificationService.cs
@@ -3,14 +3,47 @@
 using HordeServer.Api;
 using HordeServer.Models;
 using MongoDB.Bson;
-using System;
 using System.Collections.Generic;
-using System.Linq;
+using System.Diagnostics.CodeAnalysis;
 using System.Security.Claims;
 using System.Threading.Tasks;
 
 namespace HordeServer.Notifications
 {
+	/// <summary>
+	/// Marker interface for (serializable) notifications
+	/// </summary>
+	[SuppressMessage("Design", "CA1040:Avoid empty interfaces", Justification = "Marker interface for generic type handling")]
+	public interface INotification { }
+
+	/// <summary>
+	/// Notification for job scheduled events
+	/// </summary>
+	public class JobScheduledNotification : INotification
+	{
+		/// <summary>Job ID</summary>
+		public string JobId { get; }
+		
+		/// <summary>Job name</summary>
+		public string JobName { get; }
+		
+		/// <summary>Pool name job got scheduled in</summary>
+		public string PoolName { get; }
+
+		/// <summary>
+		/// Constructor
+		/// </summary>
+		/// <param name="JobId"></param>
+		/// <param name="JobName"></param>
+		/// <param name="PoolName"></param>
+		public JobScheduledNotification(string JobId, string JobName, string PoolName)
+		{
+			this.JobId = JobId;
+			this.JobName = JobName;
+			this.PoolName = PoolName;
+		}
+	}
+	
 	/// <summary>
 	/// Interface for the notification service
 	/// </summary>
@@ -44,16 +77,6 @@ namespace HordeServer.Notifications
 		/// <returns>Async task</returns>
 		void NotifyJobStepComplete(IJob Job, IGraph Graph, SubResourceId BatchId, SubResourceId StepId);
 		
-		/// <summary>
-		/// Notify all subscribers that a job batch has been scheduled
-		/// </summary>
-		/// <param name="Pool">Pool for which job scheduled in</param>
-		/// <param name="PoolHasAgentsOnline">True if there are agents online in the pool</param>
-		/// <param name="Job">The job containing the batch that has been scheduled</param>
-		/// <param name="Graph">Graph used for the job</param>
-		/// <param name="BatchId">The batch id</param>
-		void NotifyJobScheduled(IPool Pool, bool PoolHasAgentsOnline, IJob Job, IGraph Graph, SubResourceId BatchId);
-
 		/// <summary>
 		/// Notify all subscribers that a job step's outcome has changed
 		/// </summary>
diff --git a/Engine/Source/Programs/Horde/Horde.Build/Notifications/INotificationSink.cs b/Engine/Source/Programs/Horde/Horde.Build/Notifications/INotificationSink.cs
--- a/Engine/Source/Programs/Horde/Horde.Build/Notifications/INotificationSink.cs
+++ b/Engine/Source/Programs/Horde/Horde.Build/Notifications/INotificationSink.cs
@@ -3,10 +3,8 @@
 using HordeCommon;
 using HordeServer.Api;
 using HordeServer.Models;
-using MongoDB.Bson;
 using System;
 using System.Collections.Generic;
-using System.Linq;
 using System.Threading.Tasks;
 
 namespace HordeServer.Notifications
@@ -19,13 +17,9 @@ namespace HordeServer.Notifications
 		/// <summary>
 		/// Send notifications that a job has been scheduled
 		/// </summary>
-		/// <param name="Pool">Pool for which job scheduled in</param>
-		/// <param name="PoolHasAgentsOnline">True if there are agents online in the pool</param>
-		/// <param name="Job">The job containing the batch that has been scheduled</param>
-		/// <param name="Graph">Graph used for the job</param>
-		/// <param name="BatchId">The batch id</param>
+		/// <param name="Notifications">List of notifications to send</param>
 		/// <returns>Async task</returns>
-		Task NotifyJobScheduledAsync(IPool Pool, bool PoolHasAgentsOnline, IJob Job, IGraph Graph, SubResourceId BatchId);
+		Task NotifyJobScheduledAsync(List<JobScheduledNotification> Notifications);
 		
 		/// <summary>
 		/// Send notifications that a job has completed
diff --git a/Engine/Source/Programs/Horde/Horde.Build/Notifications/Impl/NotificationService.cs b/Engine/Source/Programs/Horde/Horde.Build/Notifications/Impl/NotificationService.cs
--- a/Engine/Source/Programs/Horde/Horde.Build/Notifications/Impl/NotificationService.cs
+++ b/Engine/Source/Programs/Horde/Horde.Build/Notifications/Impl/NotificationService.cs
@@ -31,13 +31,14 @@ using StatsdClient;
 using JsonSerializer = System.Text.Json.JsonSerializer;
 using HordeServer.Notifications;
 using HordeServer.Services;
+using StackExchange.Redis;
 
 namespace HordeServer.Notifications.Impl
 {
 	using UserId = ObjectId<IUser>;
 
 	/// <summary>
-	/// Wraps funtionality for delivering notifications.
+	/// Wraps functionality for delivering notifications.
 	/// </summary>
 	public class NotificationService : BackgroundService, INotificationService
 	{
@@ -90,6 +91,11 @@ namespace HordeServer.Notifications.Impl
 		/// Instance of the <see cref="IDogStatsd"/>.
 		/// </summary>
 		private readonly IDogStatsd DogStatsd;
+
+		/// <summary>
+		/// Redis database
+		/// </summary>
+		private readonly IDatabase RedisDatabase;
 		
 		/// <summary>
 		/// Settings for the application.
@@ -110,11 +116,23 @@ namespace HordeServer.Notifications.Impl
 		/// Settings for the application.
 		/// </summary>
 		private readonly ILogger<NotificationService> Logger;
+		
+		/// <summary>
+		/// Ticker for running batch sender method
+		/// </summary>
+		internal ITicker Ticker;
+
+		/// <summary>
+		/// Interval at which queued notifications should be sent as a batch 
+		/// </summary>
+		internal TimeSpan NotificationBatchInterval = TimeSpan.FromHours(12);
+		
+		static string RedisQueueListKey(string NotificationType) => "NotificationService.queued." + NotificationType;
 
 		/// <summary>
 		/// Constructor
 		/// </summary>
-		public NotificationService(IEnumerable<INotificationSink> Sinks, IOptionsMonitor<ServerSettings> Settings, ILogger<NotificationService> Logger, IGraphCollection GraphCollection, ISubscriptionCollection SubscriptionCollection, INotificationTriggerCollection TriggerCollection, IUserCollection UserCollection, JobService JobService, StreamService StreamService, IIssueService IssueService, ILogFileService LogFileService, IDogStatsd DogStatsd)
+		public NotificationService(IEnumerable<INotificationSink> Sinks, IOptionsMonitor<ServerSettings> Settings, ILogger<NotificationService> Logger, IGraphCollection GraphCollection, ISubscriptionCollection SubscriptionCollection, INotificationTriggerCollection TriggerCollection, IUserCollection UserCollection, JobService JobService, StreamService StreamService, IIssueService IssueService, ILogFileService LogFileService, IDogStatsd DogStatsd, IDatabase RedisDatabase, IClock Clock)
 		{
 			this.Sinks = Sinks.ToList();
 			this.Settings = Settings;
@@ -128,11 +146,14 @@ namespace HordeServer.Notifications.Impl
 			this.IssueService = IssueService;
 			this.LogFileService = LogFileService;
 			this.DogStatsd = DogStatsd;
+			this.RedisDatabase = RedisDatabase;
 
 			IssueService.OnIssueUpdated += NotifyIssueUpdated;
 			JobService.OnJobStepComplete += NotifyJobStepComplete;
 			JobService.OnJobScheduled += NotifyJobScheduled;
 			JobService.OnLabelUpdate += NotifyLabelUpdate;
+			
+			this.Ticker = Clock.AddSharedTicker<NotificationService>(NotificationBatchInterval, TickEveryTwelveHoursAsync, Logger);
 		}
 
 		/// <inheritdoc/>
@@ -146,6 +167,7 @@ namespace HordeServer.Notifications.Impl
 			JobService.OnLabelUpdate -= NotifyLabelUpdate;
 
 			GC.SuppressFinalize(this);
+			Ticker.Dispose();
 		}
 
 		/// <inheritdoc/>
@@ -205,7 +227,7 @@ namespace HordeServer.Notifications.Impl
 		{
 			if (Pool.EnableAutoscaling && !PoolHasAgentsOnline)
 			{
-//				EnqueueTasks(Sink => Sink.NotifyJobScheduledAsync(Pool, PoolHasAgentsOnline, Job, Graph, BatchId));
+				EnqueueTasks(Sink => EnqueueNotificationForBatchSending(new JobScheduledNotification(Job.Id.ToString(), Job.Name, Pool.Name)));
 			}
 		}
 
@@ -263,10 +285,72 @@ namespace HordeServer.Notifications.Impl
 			}
 		}
 
+		/// <summary>
+		/// Enqueue a notification in Redis for batch sending later on 
+		/// </summary>
+		/// <param name="Notification"></param>
+		/// <typeparam name="T">Any INotification type</typeparam>
+		private async Task EnqueueNotificationForBatchSending<T>(T Notification) where T : INotification
+		{
+			try
+			{
+				byte[] Data = JsonSerializer.SerializeToUtf8Bytes(Notification);
+				await RedisDatabase.ListRightPushAsync(RedisQueueListKey(typeof(T).ToString()), Data);
+			}
+			catch (Exception e)
+			{
+				Logger.LogError(e, "Unable to serialize and queue notification {Type} for batch sending in Redis", Notification.GetType());
+			}
+		}
+
+		private async ValueTask TickEveryTwelveHoursAsync(CancellationToken StoppingToken)
+		{
+			List<JobScheduledNotification> JobScheduledNotifications = await GetAllQueuedNotificationsAsync<JobScheduledNotification>();
+
+			foreach (INotificationSink Sink in Sinks)
+			{
+				await Sink.NotifyJobScheduledAsync(JobScheduledNotifications);
+			}
+		}
+		
+		/// <summary>
+		/// Get and clear all queued notifications of type T from Redis
+		/// </summary>
+		/// <typeparam name="T"></typeparam>
+		/// <returns>Deserialized notifications</returns>
+		/// <exception cref="Exception"></exception>
+		private async Task<List<T>> GetAllQueuedNotificationsAsync<T>() where T : INotification
+		{
+			// Reading and deleting the entire list from Redis in this way is not thread-safe.
+			// But likely not a big deal considering the alternative of distributed locks.
+			string RedisKey = RedisQueueListKey(typeof(T).ToString());
+			RedisValue[] RawNotifications = await RedisDatabase.ListRangeAsync(RedisKey);
+			await RedisDatabase.KeyDeleteAsync(RedisKey);
+
+			List<T> Notifications = new List<T>();
+			foreach (byte[] Data in RawNotifications)
+			{
+				try
+				{
+					T? Notification = JsonSerializer.Deserialize<T>(Data);
+					if (Notification == null) throw new Exception("Unable to deserialize");
+					Notifications.Add(Notification);
+				}
+				catch (Exception e)
+				{
+					Logger.LogError(e, "Unable to deserialize notification data of {Type}", typeof(T));
+				}
+			}
+
+			return Notifications;
+		}
+
 		/// <inheritdoc/>
 		[System.Diagnostics.CodeAnalysis.SuppressMessage("Design", "CA1031:Do not catch general exception types", Justification = "<Pending>")]
 		protected async override Task ExecuteAsync(CancellationToken StoppingToken)
 		{
+			await Ticker.StartAsync();
+			
 			// This background service just waits for tasks to finish and prints any exception info. The only reason to do this is to
 			// ensure we finish processing everything before shutdown.
 			using (CancellationTask StoppingTask = new CancellationTask(StoppingToken))
@@ -316,6 +400,8 @@ namespace HordeServer.Notifications.Impl
 					}
 				}
 			}
+			
+			await Ticker.StopAsync();
 		}
 
 		internal Task ExecuteBackgroundForTest(CancellationToken StoppingToken)
diff --git a/Engine/Source/Programs/Horde/Horde.Build/Notifications/Impl/SlackNotificationSink.cs b/Engine/Source/Programs/Horde/Horde.Build/Notifications/Impl/SlackNotificationSink.cs
--- a/Engine/Source/Programs/Horde/Horde.Build/Notifications/Impl/SlackNotificationSink.cs
+++ b/Engine/Source/Programs/Horde/Horde.Build/Notifications/Impl/SlackNotificationSink.cs
@@ -280,30 +280,37 @@ namespace HordeServer.Notifications.Impl
 		#endregion
 		
 		/// <inheritdoc/>
-		public async Task NotifyJobScheduledAsync(IPool Pool, bool PoolHasAgentsOnline, IJob Job, IGraph Graph, SubResourceId BatchId)
+		public async Task NotifyJobScheduledAsync(List<JobScheduledNotification> Notifications)
 		{
 			if (Settings.JobNotificationChannel != null)
 			{
-				Logger.LogInformation("Sending Slack notification for scheduled job {JobId}, batch {BatchId} to channel {SlackChannel}", Job.Id, BatchId, Settings.JobNotificationChannel );
-				await SendJobScheduledOnEmptyAutoScaledPoolMessageAsync($"#{Settings.JobNotificationChannel}", Pool, PoolHasAgentsOnline, Job, Graph, BatchId);
+				string JobIds = string.Join(", ", Notifications.Select(x => x.JobId));
+				Logger.LogInformation("Sending Slack notification for scheduled job IDs {JobIds} to channel {SlackChannel}", JobIds, Settings.JobNotificationChannel);
+				await SendJobScheduledOnEmptyAutoScaledPoolMessageAsync($"#{Settings.JobNotificationChannel}", Notifications);
 			}
 		}
 		
-		private Task SendJobScheduledOnEmptyAutoScaledPoolMessageAsync(string Recipient, IPool Pool, bool PoolHasAgentsOnline, IJob Job, IGraph Graph, SubResourceId BatchId)
+		private async Task SendJobScheduledOnEmptyAutoScaledPoolMessageAsync(string Recipient, List<JobScheduledNotification> Notifications)
 		{
-			string JobUrl = Settings.DashboardUrl + "/job/" + Job.Id;
+			string JobIds = string.Join(", ", Notifications.Select(x => x.JobId));
+				
+			StringBuilder Sb = new();
+			foreach (JobScheduledNotification Notification in Notifications)
+			{
+				string JobUrl = Settings.DashboardUrl + "/job/" + Notification.JobId;
+				Sb.AppendLine($"Job `{Notification.JobName}` with ID <{JobUrl}|{Notification.JobId}> in pool `{Notification.PoolName}`");
+			}
 			
 			Color OutcomeColor = BlockKitAttachmentColors.Warning;
 			BlockKitAttachment Attachment = new BlockKitAttachment();
 			Attachment.Color = OutcomeColor;
-			Attachment.FallbackText = $"Job scheduled in an auto-scaled pool with no agents online. Job ID {Job.Id}";
-
-			Attachment.Blocks.Add(new HeaderBlock($"Job scheduled in empty pool", false, true));
+			Attachment.FallbackText = $"Job(s) scheduled in an auto-scaled pool with no agents online. Job IDs {JobIds}";
 
-			Attachment.Blocks.Add(new SectionBlock($"A <{JobUrl}|job> was scheduled in an auto-scaled pool but with no current agents online."));
-			Attachment.Blocks.Add(new SectionBlock($"Job {Job.Name}\nJob ID <{JobUrl}|{Job.Id}>\nPool {Pool.Name}"));
-
-			return SendMessageAsync(Recipient, Attachments: new[] { Attachment });
+			Attachment.Blocks.Add(new HeaderBlock($"Jobs scheduled in empty pool", false, true));
+			Attachment.Blocks.Add(new SectionBlock($"One or more jobs were scheduled in an auto-scaled pool but with no current agents online."));
+			Attachment.Blocks.Add(new SectionBlock(Sb.ToString()));
+			
+			await SendMessageAsync(Recipient, Attachments: new[] { Attachment });
 		}
 
 		#region Job Complete
