Handle purge filetypes when enumerating changes in a repo.

diff --git a/Engine/Source/Programs/Horde/Horde.Build/Perforce/Impl/PerforceService.cs b/Engine/Source/Programs/Horde/Horde.Build/Perforce/Impl/PerforceService.cs
--- a/Engine/Source/Programs/Horde/Horde.Build/Perforce/Impl/PerforceService.cs
+++ b/Engine/Source/Programs/Horde/Horde.Build/Perforce/Impl/PerforceService.cs
@@ -697,6 +697,7 @@ namespace HordeServer.Services
 					return HeadRev;
 				case P4.FileAction.Delete:
 				case P4.FileAction.MoveDelete:
+				case P4.FileAction.Purge:
 					return -1;
 				default:
 					throw new Exception($"Unrecognized P4 file change type '{HeadAction}' for file {Path}#{HeadRev}");
