Horde: Register missing AutoscaleServiceV2 for DI

Will still only activate as a *hosted* service if configured in settings.

diff --git a/Engine/Source/Programs/Horde/Horde.Build/Startup.cs b/Engine/Source/Programs/Horde/Horde.Build/Startup.cs
--- a/Engine/Source/Programs/Horde/Horde.Build/Startup.cs
+++ b/Engine/Source/Programs/Horde/Horde.Build/Startup.cs
@@ -335,6 +335,7 @@ namespace HordeServer
 			Services.AddSingleton(typeof(ISingletonDocument<>), typeof(SingletonDocument<>));
 
 			Services.AddSingleton<AutoscaleService>();
+			Services.AddSingleton<AutoscaleServiceV2>();
 
 			if (Enum.TryParse(Settings.AgentPoolSizeStrategy, out PoolSizeStrategy PoolSizeStrategy))
 			{
