Horde: Fix incorrect claim mapping when creating users.

diff --git a/Engine/Source/Programs/Horde/Horde.Build/Authentication/OktaHandler.cs b/Engine/Source/Programs/Horde/Horde.Build/Authentication/OktaHandler.cs
--- a/Engine/Source/Programs/Horde/Horde.Build/Authentication/OktaHandler.cs
+++ b/Engine/Source/Programs/Horde/Horde.Build/Authentication/OktaHandler.cs
@@ -86,8 +86,8 @@ namespace HordeServer.Authentication
 				return HandleRequestResult.Fail("No identity specified");
 			}
 
-			string Login = Identity.FindFirst(ClaimTypes.Name)!.Value;
-			string? Name = Identity.FindFirst("name")?.Value;
+			string Login = Identity.FindFirst(HordeClaimTypes.User)!.Value;
+			string? Name = Identity.FindFirst(ClaimTypes.Name)?.Value;
 			string? Email = Identity.FindFirst(ClaimTypes.Email)?.Value;
 
 			IUser User = await UserCollection.FindOrAddUserByLoginAsync(Login, Name, Email);
