Horde: Move global ACL into the globals singleton, and ensure it's updated when the config file changes. Also remove serialization path that supports integers for actions.

diff --git a/Engine/Source/Programs/Horde/Horde.Build/Acls/Acl.cs b/Engine/Source/Programs/Horde/Horde.Build/Acls/Acl.cs
--- a/Engine/Source/Programs/Horde/Horde.Build/Acls/Acl.cs
+++ b/Engine/Source/Programs/Horde/Horde.Build/Acls/Acl.cs
@@ -769,10 +769,6 @@ namespace HordeServer.Models
 				{
 					break;
 				}
-				else if (Type == BsonType.Int32)
-				{
-					Values.Add((AclAction)Context.Reader.ReadInt32());
-				}
 				else
 				{
 					Values.Add((AclAction)Enum.Parse(typeof(AclAction), Context.Reader.ReadString()));
@@ -794,16 +790,4 @@ namespace HordeServer.Models
 			Context.Writer.WriteEndArray();
 		}
 	}
-
-	/// <summary>
-	/// Represents an access control list for an object in the database
-	/// </summary>
-	[SingletonDocument("5e39b51941b875626e600193")]
-	public class GlobalPermissions : SingletonBase
-	{
-		/// <summary>
-		/// The global ACL object
-		/// </summary>
-		public Acl Acl { get; set; } = new Acl();
-	}
 }
diff --git a/Engine/Source/Programs/Horde/Horde.Build/Acls/AclService.cs b/Engine/Source/Programs/Horde/Horde.Build/Acls/AclService.cs
--- a/Engine/Source/Programs/Horde/Horde.Build/Acls/AclService.cs
+++ b/Engine/Source/Programs/Horde/Horde.Build/Acls/AclService.cs
@@ -29,9 +29,9 @@ namespace HordeServer.Services
 	public class GlobalPermissionsCache
 	{
 		/// <summary>
-		/// The global permissions list
+		/// The root acl
 		/// </summary>
-		public GlobalPermissions? GlobalPermissions { get; set; }
+		public Acl? RootAcl { get; set; }
 	}
 
 	/// <summary>
@@ -79,11 +79,6 @@ namespace HordeServer.Services
 		/// </summary>
 		DatabaseService DatabaseService;
 
-		/// <summary>
-		/// The global permissions singleton
-		/// </summary>
-		ISingletonDocument<GlobalPermissions> GlobalPermissions;
-
 		/// <summary>
 		/// The settings object
 		/// </summary>
@@ -93,12 +88,10 @@ namespace HordeServer.Services
 		/// Constructor
 		/// </summary>
 		/// <param name="DatabaseService">The database service instance</param>
-		/// <param name="GlobalPermissions">The global permissions instance</param>
 		/// <param name="Settings">The settings object</param>
-		public AclService(DatabaseService DatabaseService, ISingletonDocument<GlobalPermissions> GlobalPermissions, IOptionsMonitor<ServerSettings> Settings)
+		public AclService(DatabaseService DatabaseService, IOptionsMonitor<ServerSettings> Settings)
 		{
 			this.DatabaseService = DatabaseService;
-			this.GlobalPermissions = GlobalPermissions;
 			this.Settings = Settings.CurrentValue;
 
 			List<AclAction> AdminActions = new List<AclAction>();
@@ -127,19 +120,10 @@ namespace HordeServer.Services
 		/// Gets the root ACL scope table
 		/// </summary>
 		/// <returns>Scopes instance</returns>
-		public Task<GlobalPermissions> GetGlobalPermissionsAsync()
-		{
-			return GlobalPermissions.GetAsync();
-		}
-
-		/// <summary>
-		/// Updates the root ACL scopes
-		/// </summary>
-		/// <param name="NewPermissions">New permissions</param>
-		/// <returns>Async task</returns>
-		public Task<bool> TryUpdateGlobalPermissionsAsync(GlobalPermissions NewPermissions)
+		public async Task<Acl> GetRootAcl()
 		{
-			return GlobalPermissions.TryUpdateAsync(NewPermissions);
+			Globals Globals = await DatabaseService.GetGlobalsAsync();
+			return Globals.RootAcl ?? new Acl();
 		}
 
 		/// <summary>
@@ -151,20 +135,20 @@ namespace HordeServer.Services
 		/// <returns>Async task</returns>
 		public async Task<bool> AuthorizeAsync(AclAction Action, ClaimsPrincipal User, GlobalPermissionsCache? Cache = null)
 		{
-			GlobalPermissions? GlobalPermissions;
+			Acl? RootAcl;
 			if(Cache == null)
 			{
-				GlobalPermissions = await GetGlobalPermissionsAsync();
+				RootAcl = await GetRootAcl();
 			}
-			else if(Cache.GlobalPermissions == null)
+			else if(Cache.RootAcl == null)
 			{
-				GlobalPermissions = Cache.GlobalPermissions = await GetGlobalPermissionsAsync();
+				RootAcl = Cache.RootAcl = await GetRootAcl();
 			}
 			else
 			{
-				GlobalPermissions = Cache.GlobalPermissions;
+				RootAcl = Cache.RootAcl;
 			}
-			return GlobalPermissions.Acl.Authorize(Action, User) ?? DefaultAcl.Authorize(Action, User) ?? false;
+			return RootAcl.Authorize(Action, User) ?? DefaultAcl.Authorize(Action, User) ?? false;
 		}
 
 		/// <summary>
diff --git a/Engine/Source/Programs/Horde/Horde.Build/Controllers/PermissionsController.cs b/Engine/Source/Programs/Horde/Horde.Build/Controllers/PermissionsController.cs
deleted file mode 100644
diff --git a/Engine/Source/Programs/Horde/Horde.Build/Models/Globals.cs b/Engine/Source/Programs/Horde/Horde.Build/Models/Globals.cs
--- a/Engine/Source/Programs/Horde/Horde.Build/Models/Globals.cs
+++ b/Engine/Source/Programs/Horde/Horde.Build/Models/Globals.cs
@@ -301,6 +301,11 @@ namespace HordeServer.Models
 		/// </summary>
 		public int MaxConformCount { get; set; }
 
+		/// <summary>
+		/// The root ACL
+		/// </summary>
+		public Acl? RootAcl { get; set; }
+
 		/// <summary>
 		/// Randomly generated JWT signing key
 		/// </summary>
diff --git a/Engine/Source/Programs/Horde/Horde.Build/Server/ConfigService.cs b/Engine/Source/Programs/Horde/Horde.Build/Server/ConfigService.cs
--- a/Engine/Source/Programs/Horde/Horde.Build/Server/ConfigService.cs
+++ b/Engine/Source/Programs/Horde/Horde.Build/Server/ConfigService.cs
@@ -49,7 +49,7 @@ namespace HordeServer.Services
 		/// <summary>
 		/// Config file version number
 		/// </summary>
-		const int Version = 5;
+		const int Version = 6;
 
 		DatabaseService DatabaseService;
 		ProjectService ProjectService;
@@ -150,6 +150,7 @@ namespace HordeServer.Services
 				Globals.PerforceClusters = CachedGlobalConfig.PerforceClusters;
 				Globals.ScheduledDowntime = CachedGlobalConfig.Downtime;
 				Globals.MaxConformCount = CachedGlobalConfig.MaxConformCount;
+				Globals.RootAcl = Acl.Merge(null, CachedGlobalConfig.Acl);
 
 				await UpdateStorageConfigAsync(GlobalConfig.Storage);
 
