Remove references to deleted global permissions class.

diff --git a/Engine/Source/Programs/Horde/Horde.Build.Tests/RpcServiceTest.cs b/Engine/Source/Programs/Horde/Horde.Build.Tests/RpcServiceTest.cs
--- a/Engine/Source/Programs/Horde/Horde.Build.Tests/RpcServiceTest.cs
+++ b/Engine/Source/Programs/Horde/Horde.Build.Tests/RpcServiceTest.cs
@@ -26,7 +26,6 @@ using Microsoft.Extensions.Logging;
 using Microsoft.VisualStudio.TestTools.UnitTesting;
 using MongoDB.Bson;
 using AgentCapabilities = HordeCommon.Rpc.Messages.AgentCapabilities;
-using GlobalPermissions = HordeServer.Models.GlobalPermissions;
 using ISession = Microsoft.AspNetCore.Http.ISession;
 using PoolId = HordeServer.Utilities.StringId<HordeServer.Models.IPool>;
 
diff --git a/Engine/Source/Programs/Horde/Horde.Build.Tests/TestSetup.cs b/Engine/Source/Programs/Horde/Horde.Build.Tests/TestSetup.cs
--- a/Engine/Source/Programs/Horde/Horde.Build.Tests/TestSetup.cs
+++ b/Engine/Source/Programs/Horde/Horde.Build.Tests/TestSetup.cs
@@ -213,7 +213,6 @@ namespace HordeServerTests
 
 			Services.AddSingleton<IStorageService, SimpleStorageService>();
 
-			Services.AddSingleton<ISingletonDocument<GlobalPermissions>>(new SingletonDocumentStub<GlobalPermissions>());
 			Services.AddSingleton<ISingletonDocument<AgentSoftwareChannels>>(new SingletonDocumentStub<AgentSoftwareChannels>());
 		}
 
