Horde: Add endpoint in job controller for finding job timings

Batch getting these on the server-side leads to less back-and-forth when resolving timings for many jobs.

diff --git a/Engine/Source/Programs/Horde/Horde.Build.Tests/JobsControllerDbTest.cs b/Engine/Source/Programs/Horde/Horde.Build.Tests/JobsControllerDbTest.cs
--- a/Engine/Source/Programs/Horde/Horde.Build.Tests/JobsControllerDbTest.cs
+++ b/Engine/Source/Programs/Horde/Horde.Build.Tests/JobsControllerDbTest.cs
@@ -59,5 +59,16 @@ namespace HordeServerTests
 	        Assert.IsTrue(StepRes.AbortRequested);
 //	        Assert.AreEqual("Anonymous", StepRes.AbortByUser);
         }
+        
+        [TestMethod]
+        public async Task FindJobTimingsTest()
+        {
+	        Fixture Fixture = await CreateFixtureAsync();
+	        IJob Job = Fixture.Job1;
+	        string[] Templates = { Job.TemplateId.ToString() };
+	        object Obj = (await JobsController.FindJobTimingsAsync(Fixture.Stream!.Id.ToString(), Templates)).Value!;
+	        FindJobTimingsResponse Res = (Obj as FindJobTimingsResponse)!;
+	        Assert.AreEqual(1, Res.Timings.Count);
+        }
     }
 }
\ No newline at end of file
diff --git a/Engine/Source/Programs/Horde/Horde.Build/Api/Jobs.cs b/Engine/Source/Programs/Horde/Horde.Build/Api/Jobs.cs
--- a/Engine/Source/Programs/Horde/Horde.Build/Api/Jobs.cs
+++ b/Engine/Source/Programs/Horde/Horde.Build/Api/Jobs.cs
@@ -404,6 +404,26 @@ namespace HordeServer.Api
 			this.Labels = Labels;
 		}
 	}
+	
+	/// <summary>
+	/// The timing info for 
+	/// </summary>
+	public class FindJobTimingsResponse
+	{
+		/// <summary>
+		/// Timing info for each job
+		/// </summary>
+		public Dictionary<string, GetJobTimingResponse> Timings { get; set; }
+
+		/// <summary>
+		/// Constructor
+		/// </summary>
+		/// <param name="Timings">Timing info for each job</param>
+		public FindJobTimingsResponse(Dictionary<string, GetJobTimingResponse> Timings)
+		{
+			this.Timings = Timings;
+		}
+	}
 
 	/// <summary>
 	/// Request used to update a jobstep
diff --git a/Engine/Source/Programs/Horde/Horde.Build/Controllers/JobsController.cs b/Engine/Source/Programs/Horde/Horde.Build/Controllers/JobsController.cs
--- a/Engine/Source/Programs/Horde/Horde.Build/Controllers/JobsController.cs
+++ b/Engine/Source/Programs/Horde/Horde.Build/Controllers/JobsController.cs
@@ -563,8 +563,12 @@ namespace HordeServer.Controllers
 			}
 
 			IJobTiming JobTiming = await JobService.GetJobTimingAsync(Job);
-
 			IGraph Graph = await JobService.GetGraphAsync(Job);
+			return PropertyFilter.Apply(CreateJobTimingResponse(Job, Graph, JobTiming), Filter);
+		}
+
+		private static GetJobTimingResponse CreateJobTimingResponse(IJob Job, IGraph Graph, IJobTiming JobTiming)
+		{
 			Dictionary<INode, TimingInfo> NodeToTimingInfo = Job.GetTimingInfo(Graph, JobTiming);
 
 			Dictionary<string, GetStepTimingInfoResponse> Steps = new Dictionary<string, GetStepTimingInfoResponse>();
@@ -584,7 +588,49 @@ namespace HordeServer.Controllers
 				Labels.Add(new GetLabelTimingInfoResponse(Label, TimingInfo));
 			}
 
-			return PropertyFilter.Apply(new GetJobTimingResponse(Steps, Labels), Filter);
+			return new GetJobTimingResponse(Steps, Labels);
+		}
+		
+		/// <summary>
+		/// Find timing information about the graph for multiple jobs
+		/// </summary>
+		/// <param name="StreamId">The stream to search in</param>
+		/// <param name="Templates">List of templates to find</param>
+		/// <param name="Filter">Filter for the fields to return</param>
+		/// <param name="Count">Number of results to return</param>
+		/// <returns>Job timings for each job ID</returns>
+		[HttpGet]
+		[Route("/api/v1/jobs/timing")]
+		[ProducesResponseType(typeof(FindJobTimingsResponse), 200)]
+		public async Task<ActionResult<object>> FindJobTimingsAsync(
+			[FromQuery] string? StreamId = null,
+			[FromQuery(Name = "template")] string[]? Templates = null,
+			[FromQuery] PropertyFilter? Filter = null,
+			[FromQuery] int Count = 100)
+		{
+			if (StreamId == null)
+			{
+				return BadRequest("Missing/invalid query parameter streamId");
+			}
+
+			TemplateRefId[] TemplateRefIds = Templates switch
+			{
+				{ Length: > 0 } => Templates.Select(x => new TemplateRefId(x)).ToArray(),
+				_ => Array.Empty<TemplateRefId>()
+			};
+
+			List<IJob> Jobs = await JobService.FindJobsByStreamWithTemplatesAsync(new StreamId(StreamId), TemplateRefIds, Count: Count, ConsistentRead: false);
+
+			Dictionary<string, GetJobTimingResponse> JobTimings = await Jobs.ToAsyncEnumerable()
+				.WhereAwait(async Job => await JobService.AuthorizeAsync(Job, AclAction.ViewJob, User, null))
+				.ToDictionaryAwaitAsync(x => ValueTask.FromResult(x.Id.ToString()), async Job =>
+				{
+					IJobTiming JobTiming = await JobService.GetJobTimingAsync(Job);
+					IGraph Graph = await JobService.GetGraphAsync(Job);
+					return CreateJobTimingResponse(Job, Graph, JobTiming);
+				});
+			
+			return PropertyFilter.Apply(new FindJobTimingsResponse(JobTimings), Filter);
 		}
 
 		/// <summary>
