Horde: Bump the user id claim version number to force all users to re-authenticate and get the correct user claims.

diff --git a/Engine/Source/Programs/Horde/Horde.Build/Utilities/HordeClaimTypes.cs b/Engine/Source/Programs/Horde/Horde.Build/Utilities/HordeClaimTypes.cs
--- a/Engine/Source/Programs/Horde/Horde.Build/Utilities/HordeClaimTypes.cs
+++ b/Engine/Source/Programs/Horde/Horde.Build/Utilities/HordeClaimTypes.cs
@@ -51,7 +51,7 @@ namespace HordeServer.Utilities
 		/// <summary>
 		/// Unique id of the horde user (see <see cref="IUserCollection"/>)
 		/// </summary>
-		public const string UserId = Prefix + "user-id-v2";
+		public const string UserId = Prefix + "user-id-v3";
 
 		/// <summary>
 		/// Claim for downloading artifacts from a particular job
