Horde Storage - Fixed issue with the replication log when using a bucket identifier that is to old and doesnt exist anymore

Added test to cover this critical behavior

diff --git a/Engine/Source/Programs/Horde/HordeStorage/Horde.Storage/Implementation/TransactionLog/ScyllaReplicationLog.cs b/Engine/Source/Programs/Horde/HordeStorage/Horde.Storage/Implementation/TransactionLog/ScyllaReplicationLog.cs
--- a/Engine/Source/Programs/Horde/HordeStorage/Horde.Storage/Implementation/TransactionLog/ScyllaReplicationLog.cs
+++ b/Engine/Source/Programs/Horde/HordeStorage/Horde.Storage/Implementation/TransactionLog/ScyllaReplicationLog.cs
@@ -167,6 +167,8 @@ namespace Horde.Storage.Implementation
         {
             using IScope findReplicationBucketScope = Tracer.Instance.StartActive("scylla.find_replication_buckets");
 
+            // ignore any bucket that is older then a cutoff, as that can cause us to end up scanning thru a lot of hours that will never exist (incremental logs are deleted after 7 days)
+            DateTime oldCutoff = DateTime.UtcNow.AddDays(-14);
 
             DateTime startBucketTime;
 
@@ -174,6 +176,12 @@ namespace Horde.Storage.Implementation
             {
                 long bucket = FromReplicationBucketIdentifier(lastBucket);
                 startBucketTime = DateTime.FromFileTimeUtc(bucket);
+                if (startBucketTime < oldCutoff)
+                {
+                    // attempting to use a old bucket, this will not exist anymore so will break here
+                    yield break;
+                }
+
                 yield return bucket;
             }
             else
@@ -185,8 +193,6 @@ namespace Horde.Storage.Implementation
                 startBucketTime = oldestTimestamp;
             }
 
-            // ignore any bucket that is older then a cutoff, as that can cause us to end up scanning thru a lot of hours that will never exist (incremental logs are deleted after 7 days)
-            DateTime oldCutoff = DateTime.UtcNow.AddDays(-14);
             // we returned the start bucket earlier so now we start with the next one
             DateTime bucketTime = startBucketTime.AddHours(1.0).ToHourlyBucket();
             while(bucketTime < DateTime.UtcNow && bucketTime > oldCutoff)
@@ -354,7 +360,6 @@ namespace Horde.Storage.Implementation
 
         [Cassandra.Mapping.Attributes.Column("blob_namespace")]
         public string BlobNamespace { get; set; }
-
     }
 
     [Cassandra.Mapping.Attributes.Table("replication_namespace")]
diff --git a/Engine/Source/Programs/Horde/HordeStorage/Horde.Storage/Tests/Functional/Objects/ReplicationTests.cs b/Engine/Source/Programs/Horde/HordeStorage/Horde.Storage/Tests/Functional/Objects/ReplicationTests.cs
--- a/Engine/Source/Programs/Horde/HordeStorage/Horde.Storage/Tests/Functional/Objects/ReplicationTests.cs
+++ b/Engine/Source/Programs/Horde/HordeStorage/Horde.Storage/Tests/Functional/Objects/ReplicationTests.cs
@@ -437,6 +437,36 @@ namespace Horde.Storage.FunctionalTests.References
             CollectionAssert.AreEqual(new [] {TestNamespace}, await _replicationLog.GetNamespaces().ToArrayAsync());
         }
 
+        [TestMethod]
+        public async Task ReplicationLogOldBucket()
+        {
+            // the namespace exists but the bucket id is from a old bucket that does not exist anymore
+            CompactBinaryWriter writer = new CompactBinaryWriter();
+            writer.BeginObject();
+            writer.AddString("thisIsAField", "stringField");
+            writer.EndObject();
+
+            byte[] objectData = writer.Save();
+            BlobIdentifier objectHash = BlobIdentifier.FromBlob(objectData);
+
+            await _replicationLog.InsertAddEvent(TestNamespace, TestBucket, IoHashKey.FromName("firstObject"), objectHash, DateTime.Now.AddDays(-1));
+
+            string eventBucket = DateTime.Now.AddDays(-60).ToReplicationBucketIdentifier();
+            Guid eventId = Guid.NewGuid();
+
+            {
+                HttpResponseMessage result = await _httpClient!.GetAsync(requestUri: $"api/v1/replication-log/incremental/{TestNamespace}?lastBucket={eventBucket}&lastEvent={eventId}");
+                Assert.AreEqual(HttpStatusCode.BadRequest, result.StatusCode);
+                
+                string s = await result.Content.ReadAsStringAsync();
+                
+                ProblemDetails? problem = await result.Content.ReadFromJsonAsync<ProblemDetails?>();
+                Assert.IsNotNull(problem);
+            }
+
+            CollectionAssert.AreEqual(new [] {TestNamespace}, await _replicationLog.GetNamespaces().ToArrayAsync());
+        }
+
         [TestMethod]
         public async Task ReplicationLogEmptyLog()
         {
