Horde Storage - Handle when a polling task has been cancelled

diff --git a/Engine/Source/Programs/Horde/HordeStorage/Jupiter.Common/PollingService.cs b/Engine/Source/Programs/Horde/HordeStorage/Jupiter.Common/PollingService.cs
--- a/Engine/Source/Programs/Horde/HordeStorage/Jupiter.Common/PollingService.cs
+++ b/Engine/Source/Programs/Horde/HordeStorage/Jupiter.Common/PollingService.cs
@@ -1,6 +1,7 @@
 ﻿// Copyright Epic Games, Inc. All Rights Reserved.
 
 using System;
+using System.Linq;
 using System.Threading;
 using System.Threading.Tasks;
 using Microsoft.Extensions.Hosting;
@@ -102,13 +103,26 @@ namespace Jupiter
             }
             catch (AggregateException e)
             {
-                logger.Error(e, "{Service} Aggregate exception in polling service", serviceName);
-                foreach (Exception inner in e.InnerExceptions)
+                bool taskCancelled =
+                    e.InnerExceptions.Any(exception => exception.GetType() == typeof(TaskCanceledException));
+                if (!taskCancelled)
                 {
-                    logger.Error(inner, "{Service} inner exception in polling service. Trace: {StackTrace}",
-                        serviceName, inner.StackTrace);
+                    logger.Error(e, "{Service} Aggregate exception in polling service", serviceName);
+                    foreach (Exception inner in e.InnerExceptions)
+                    {
+                        logger.Error(inner, "{Service} inner exception in polling service. Trace: {StackTrace}",
+                            serviceName, inner.StackTrace);
 
+                    }
                 }
+                else
+                {
+                    logger.Warning("{Service} poll cancelled in polling service", serviceName);
+                }
+            }
+            catch (TaskCanceledException)
+            {
+                logger.Warning("{Service} poll cancelled in polling service", serviceName);
             }
             catch (Exception e)
             {
