Horde Storage - Fixed bug were replication of delete records would not actually update the state and wouldnt look like the replicator was working (as count of replications was always 0)

diff --git a/Engine/Source/Programs/Horde/HordeStorage/Horde.Storage/Implementation/Replication/RefsReplicator.cs b/Engine/Source/Programs/Horde/HordeStorage/Horde.Storage/Implementation/Replication/RefsReplicator.cs
--- a/Engine/Source/Programs/Horde/HordeStorage/Horde.Storage/Implementation/Replication/RefsReplicator.cs
+++ b/Engine/Source/Programs/Horde/HordeStorage/Horde.Storage/Implementation/Replication/RefsReplicator.cs
@@ -338,12 +338,14 @@ namespace Horde.Storage.Implementation
                     }
 
                     // we do not need to replicate delete events
-                    if (@event.Op == ReplicationLogEvent.OpType.Deleted)
-                        return;
-                    if (@event.Blob == null)
-                        throw new Exception($"Event: [REDACTED_EMAIL]} [REDACTED_EMAIL]} in namespace [REDACTED_EMAIL]} was missing a blob, unable to replicate it");
+                    if (@event.Op != ReplicationLogEvent.OpType.Deleted)
+                    {
+                        if (@event.Blob == null)
+                            throw new Exception($"Event: [REDACTED_EMAIL]} [REDACTED_EMAIL]} in namespace [REDACTED_EMAIL]} was missing a blob, unable to replicate it");
+
+                        await ReplicateOp(@event.Namespace, @event.Blob, replicationToken);
+                    }
 
-                    await ReplicateOp(@event.Namespace, @event.Blob, replicationToken);
                     // TODO: Avoid adding to the replication log as we will get into infinite recursion if we do with the remote site
                     // we could add to the replication log only if the op was missing, but that could cause issue for a 3rd site replicating from us
                     // and the 3rd site is the whole reason we even add this to replication log in the first place
