Horde: Make notification EventRecord constructors public due to NET6 breakage, [REDACTED_PROFILE_URL]/docs/issues/20887
#rnx
#skipci

diff --git a/Engine/Source/Programs/Horde/Horde.Build/Api/Events.cs b/Engine/Source/Programs/Horde/Horde.Build/Api/Events.cs
--- a/Engine/Source/Programs/Horde/Horde.Build/Api/Events.cs
+++ b/Engine/Source/Programs/Horde/Horde.Build/Api/Events.cs
@@ -50,7 +50,7 @@ namespace HordeServer.Api
 		/// <summary>
 		/// Default constructor
 		/// </summary>
-		private JobCompleteEventRecord()
+		public JobCompleteEventRecord()
 		{
 			StreamId = String.Empty;
 			TemplateId = String.Empty;
@@ -107,7 +107,7 @@ namespace HordeServer.Api
 		/// <summary>
 		/// Default constructor
 		/// </summary>
-		private LabelCompleteEventRecord()
+		public LabelCompleteEventRecord()
 		{
 			this.StreamId = String.Empty;
 			this.TemplateId = String.Empty;
@@ -164,7 +164,7 @@ namespace HordeServer.Api
 		/// <summary>
 		/// Default constructor
 		/// </summary>
-		private StepCompleteEventRecord()
+		public StepCompleteEventRecord()
 		{
 			StreamId = String.Empty;
 			TemplateId = String.Empty;
