Horde: Fixing Schedules with Multiple Patterns only triggering the earliest pattern

Added a test for multiple pattern schedules as well

#rnx

diff --git a/Engine/Source/Programs/Horde/Horde.Build.Tests/SchedulerTests.cs b/Engine/Source/Programs/Horde/Horde.Build.Tests/SchedulerTests.cs
--- a/Engine/Source/Programs/Horde/Horde.Build.Tests/SchedulerTests.cs
+++ b/Engine/Source/Programs/Horde/Horde.Build.Tests/SchedulerTests.cs
@@ -175,6 +175,23 @@ namespace HordeServerTests
 			Assert.AreEqual(StartTime + TimeSpan.FromHours(1.0 + 24.0), NextTime!.Value);
 		}
 
+		[TestMethod]
+		public void MultiPatternTest()
+		{
+			DateTime StartTime = new DateTime(2021, 1, 1, 0, 0, 0, DateTimeKind.Utc); // Friday Jan 1, 2021 
+			Clock.UtcNow = StartTime;
+
+			Schedule Schedule = new Schedule(Clock.UtcNow, RequireSubmittedChange: false);
+			Schedule.Patterns.Add(new SchedulePattern(null, 11 * 60, 0, 0));
+			Schedule.Patterns.Add(new SchedulePattern(null, 19 * 60, 0, 0));
+
+			DateTime? NextTime = Schedule.GetNextTriggerTimeUtc(StartTime, TimeZoneInfo.Utc);
+			Assert.AreEqual(StartTime + TimeSpan.FromHours(11), NextTime!.Value);
+
+			NextTime = Schedule.GetNextTriggerTimeUtc(NextTime.Value, TimeZoneInfo.Utc);
+			Assert.AreEqual(StartTime + TimeSpan.FromHours(19), NextTime!.Value);
+		}
+
 		[TestMethod]
 		public async Task NoSubmittedChangeTestAsync()
 		{
diff --git a/Engine/Source/Programs/Horde/Horde.Build/Projects/ISchedule.cs b/Engine/Source/Programs/Horde/Horde.Build/Projects/ISchedule.cs
--- a/Engine/Source/Programs/Horde/Horde.Build/Projects/ISchedule.cs
+++ b/Engine/Source/Programs/Horde/Horde.Build/Projects/ISchedule.cs
@@ -324,7 +324,7 @@ namespace HordeServer.Models
 				foreach (SchedulePattern Pattern in Patterns)
 				{
 					DateTime PatternTriggerTime = Pattern.GetNextTriggerTimeUtc(LastTimeUtc, TimeZone);
-					if (NextTriggerTimeUtc == null || NextTriggerTimeUtc < PatternTriggerTime)
+					if (NextTriggerTimeUtc == null || PatternTriggerTime < NextTriggerTimeUtc)
 					{
 						NextTriggerTimeUtc = PatternTriggerTime;
 					}
