Horde Storage - Further debugging information when encountering a snapshot response that is not correctly formatted

diff --git a/Engine/Source/Programs/Horde/HordeStorage/Horde.Storage/Implementation/Replication/RefsReplicator.cs b/Engine/Source/Programs/Horde/HordeStorage/Horde.Storage/Implementation/Replication/RefsReplicator.cs
--- a/Engine/Source/Programs/Horde/HordeStorage/Horde.Storage/Implementation/Replication/RefsReplicator.cs
+++ b/Engine/Source/Programs/Horde/HordeStorage/Horde.Storage/Implementation/Replication/RefsReplicator.cs
@@ -465,6 +465,11 @@ namespace Horde.Storage.Implementation
 
                     if (problemDetails.Type == ProblemTypes.UseSnapshot)
                     {
+                        if (!problemDetails.Extensions.ContainsKey("SnapshotId"))
+                        {
+                            throw new Exception("Unknown response, looks like a problem details but is lacking snapshot id " + body);
+                        }
+
                         BlobIdentifier snapshotBlob = new BlobIdentifier(problemDetails.Extensions["SnapshotId"]!.ToString()!);
                         throw new UseSnapshotException(snapshotBlob);
                     }
