Horde Storage - Fixes to deserialize ProblemDetails with extensions, as newtonsoft doesnt populate the missing members in Extensions.

diff --git a/Engine/Source/Programs/Horde/HordeStorage/Horde.Storage/Implementation/Replication/RefsReplicator.cs b/Engine/Source/Programs/Horde/HordeStorage/Horde.Storage/Implementation/Replication/RefsReplicator.cs
--- a/Engine/Source/Programs/Horde/HordeStorage/Horde.Storage/Implementation/Replication/RefsReplicator.cs
+++ b/Engine/Source/Programs/Horde/HordeStorage/Horde.Storage/Implementation/Replication/RefsReplicator.cs
@@ -8,6 +8,7 @@ using System.Net;
 using System.Net.Http;
 using System.Runtime.CompilerServices;
 using System.Text;
+using System.Text.Json;
 using System.Threading;
 using System.Threading.Tasks;
 using Dasync.Collections;
@@ -21,6 +22,7 @@ using Microsoft.Extensions.Options;
 using Newtonsoft.Json;
 using Nito.AsyncEx;
 using Serilog;
+using JsonSerializer = Newtonsoft.Json.JsonSerializer;
 
 namespace Horde.Storage.Implementation
 {
@@ -461,16 +463,16 @@ namespace Horde.Storage.Implementation
                 {
                     ProblemDetails? problemDetails = JsonConvert.DeserializeObject<ProblemDetails>(body);
                     if (problemDetails == null)
-                        throw new Exception("Unknown bad request body when reading incremental replication log. Body: {body}");
+                        throw new Exception($"Unknown bad request body when reading incremental replication log. Body: {body}");
 
                     if (problemDetails.Type == ProblemTypes.UseSnapshot)
                     {
-                        if (!problemDetails.Extensions.ContainsKey("SnapshotId"))
-                        {
-                            throw new Exception("Unknown response, looks like a problem details but is lacking snapshot id " + body);
-                        }
+                        ProblemDetailsWithSnapshots? problemDetailsWithSnapshots = JsonConvert.DeserializeObject<ProblemDetailsWithSnapshots>(body);
+
+                        if (problemDetailsWithSnapshots == null)
+                            throw new Exception($"Unable to cast the problem details to a snapshot version. Body: {body}");
 
-                        BlobIdentifier snapshotBlob = new BlobIdentifier(problemDetails.Extensions["SnapshotId"]!.ToString()!);
+                        BlobIdentifier snapshotBlob = problemDetailsWithSnapshots.SnapshotId;
                         throw new UseSnapshotException(snapshotBlob);
                     }
                 }
@@ -565,4 +567,9 @@ namespace Horde.Storage.Implementation
             LastBucket = null;
         }
     }
+
+    public class ProblemDetailsWithSnapshots : ProblemDetails
+    {
+        public BlobIdentifier SnapshotId { get; set; } = null!;
+    }
 }
