Horde: Fix some (new?) warnings being treated as errors.

diff --git a/Engine/Source/Programs/Horde/Horde.Agent/Commands/Execution/ExecuteCommandV2.cs b/Engine/Source/Programs/Horde/Horde.Agent/Commands/Execution/ExecuteCommandV2.cs
--- a/Engine/Source/Programs/Horde/Horde.Agent/Commands/Execution/ExecuteCommandV2.cs
+++ b/Engine/Source/Programs/Horde/Horde.Agent/Commands/Execution/ExecuteCommandV2.cs
@@ -358,7 +358,7 @@ namespace HordeAgent.Commands
 				using IComputeClient ComputeClient = CreateComputeClient(Host.Services);
 				
 				byte[] Json = await FileReference.ReadAllBytesAsync(Input);
-				JsonComputeTask JsonComputeTask = JsonSerializer.Deserialize<JsonComputeTask>(Json, new JsonSerializerOptions { PropertyNameCaseInsensitive = true });
+				JsonComputeTask JsonComputeTask = JsonSerializer.Deserialize<JsonComputeTask>(Json, new JsonSerializerOptions { PropertyNameCaseInsensitive = true })!;
 
 				ByteString SaltBytes = ByteString.Empty;
 				if (Salt != null)
diff --git a/Engine/Source/Programs/Horde/Horde.Agent/Execution/BuildGraphExecutor.cs b/Engine/Source/Programs/Horde/Horde.Agent/Execution/BuildGraphExecutor.cs
--- a/Engine/Source/Programs/Horde/Horde.Agent/Execution/BuildGraphExecutor.cs
+++ b/Engine/Source/Programs/Horde/Horde.Agent/Execution/BuildGraphExecutor.cs
@@ -389,7 +389,7 @@ namespace HordeAgent.Execution
 			Options.PropertyNameCaseInsensitive = true;
 			Options.Converters.Add(new JsonStringEnumConverter());
 
-			ExportedGraph Graph = JsonSerializer.Deserialize<ExportedGraph>(await FileReference.ReadAllBytesAsync(DefinitionFile), Options);
+			ExportedGraph Graph = JsonSerializer.Deserialize<ExportedGraph>(await FileReference.ReadAllBytesAsync(DefinitionFile), Options)!;
 
 			List<string> MissingAgentTypes = new List<string>();
 
@@ -834,7 +834,7 @@ namespace HordeAgent.Execution
 					Logger.LogInformation("Reading telemetry from {File}", TelemetryFile);
 					byte[] Data = await FileReference.ReadAllBytesAsync(TelemetryFile);
 
-					TraceEventList Telemetry = JsonSerializer.Deserialize<TraceEventList>(Data.AsSpan());
+					TraceEventList Telemetry = JsonSerializer.Deserialize<TraceEventList>(Data.AsSpan())!;
 					if (Telemetry.Spans.Count > 0)
 					{
 						string DefaultServiceName = TelemetryFile.GetFileNameWithoutAnyExtensions();
@@ -911,7 +911,7 @@ namespace HordeAgent.Execution
 					FileReference TraceFile = FileReference.Combine(TelemetryDir, "Trace.json");
 					using (FileStream Stream = FileReference.Open(TraceFile, FileMode.Create))
 					{
-						JsonSerializerOptions Options = new JsonSerializerOptions { IgnoreNullValues = true };
+						JsonSerializerOptions Options = new JsonSerializerOptions { DefaultIgnoreCondition = JsonIgnoreCondition.WhenWritingNull };
 						await JsonSerializer.SerializeAsync(Stream, RootSpan, Options);
 					}
 					await ArtifactUploader.UploadAsync(RpcConnection, JobId, BatchId, Step.StepId, "Trace.json", TraceFile, Logger, CancellationToken.None);
@@ -932,7 +932,7 @@ namespace HordeAgent.Execution
 					using (FileStream Stream = FileReference.Open(TestDataFile, FileMode.Open))
 					{
 						JsonSerializerOptions Options = new JsonSerializerOptions { PropertyNameCaseInsensitive = true };
-						TestData = await JsonSerializer.DeserializeAsync<TestData>(Stream, Options);
+						TestData = (await JsonSerializer.DeserializeAsync<TestData>(Stream, Options))!;
 					}
 
 					foreach (TestDataItem Item in TestData.Items)
@@ -990,7 +990,7 @@ namespace HordeAgent.Execution
 			JsonSerializerOptions Options = new JsonSerializerOptions();
 			Options.PropertyNameCaseInsensitive = true;
 			Options.Converters.Add(new JsonStringEnumConverter());
-			ReportData Report = JsonSerializer.Deserialize<ReportData>(Data, Options);
+			ReportData Report = JsonSerializer.Deserialize<ReportData>(Data, Options)!;
 
 			if (String.IsNullOrEmpty(Report.Name))
 			{
