Horde Storage - Fixed issue were the content id resolve would return hashes which was unknown to the caller (when a content id mapped to blobs that were missing in the blob store).
That fix is relativley simple and contained to the content id resolvers themselves (ScyllaContentIdStore).

This bug exposed the confusing use of BlobIdentifier as both a identifier into the blob store (its original intent) and as a content id (which needs to be resolved before it can be used to locate a blob in the blob store).
So I added a new type for ContentId so we can keep track of this seperate, this required quite a lot more refactoring to handle properly (indicating how bad of a state this was in).
This also fixed a issue were ContentHash was not being serialized into compact binary representations in the api, which is now fixed.

This is a potentially dangerous change with lots of things impacted, added a unit test to cover exactly the case that prompted the bug in the first place.

diff --git a/Engine/Source/Programs/Horde/HordeStorage/Horde.Storage/Controllers/CompressedBlobController.cs b/Engine/Source/Programs/Horde/HordeStorage/Horde.Storage/Controllers/CompressedBlobController.cs
--- a/Engine/Source/Programs/Horde/HordeStorage/Horde.Storage/Controllers/CompressedBlobController.cs
+++ b/Engine/Source/Programs/Horde/HordeStorage/Horde.Storage/Controllers/CompressedBlobController.cs
@@ -52,7 +52,7 @@ namespace Horde.Storage.Controllers
 
         public async Task<IActionResult> Get(
             [Required] NamespaceId ns,
-            [Required] BlobIdentifier id)
+            [Required] ContentId id)
         {
             AuthorizationResult authorizationResult = await _authorizationService.AuthorizeAsync(User, ns, NamespaceAccessRequirement.Name);
 
@@ -82,7 +82,7 @@ namespace Horde.Storage.Controllers
         [ProducesDefaultResponseType]
         public async Task<IActionResult> Head(
             [Required] NamespaceId ns,
-            [Required] BlobIdentifier id)
+            [Required] ContentId id)
         {
             AuthorizationResult authorizationResult = await _authorizationService.AuthorizeAsync(User, ns, NamespaceAccessRequirement.Name);
 
@@ -119,7 +119,7 @@ namespace Horde.Storage.Controllers
         [ProducesDefaultResponseType]
         public async Task<IActionResult> ExistsMultiple(
             [Required] NamespaceId ns,
-            [Required] [FromQuery] List<BlobIdentifier> id)
+            [Required] [FromQuery] List<ContentId> id)
         {
             AuthorizationResult authorizationResult = await _authorizationService.AuthorizeAsync(User, ns, NamespaceAccessRequirement.Name);
 
@@ -128,8 +128,8 @@ namespace Horde.Storage.Controllers
                 return Forbid();
             }
 
-            ConcurrentBag<BlobIdentifier> missingBlobs = new ConcurrentBag<BlobIdentifier>();
-            ConcurrentBag<BlobIdentifier> invalidContentIds = new ConcurrentBag<BlobIdentifier>();
+            ConcurrentBag<ContentId> partialContentIds = new ConcurrentBag<ContentId>();
+            ConcurrentBag<ContentId> invalidContentIds = new ConcurrentBag<ContentId>();
 
             IEnumerable<Task> tasks = id.Select(async blob =>
             {
@@ -145,7 +145,8 @@ namespace Horde.Storage.Controllers
                 {
                     if (!await _storage.Exists(ns, chunk))
                     {
-                        missingBlobs.Add(blob);
+                        partialContentIds.Add(blob);
+                        break;
                     }
                 }
             });
@@ -154,10 +155,10 @@ namespace Horde.Storage.Controllers
             if (invalidContentIds.Count != 0)
                 return NotFound(new ValidationProblemDetails { Title = $"Missing content ids {string.Join(" ,", invalidContentIds)}"});
 
-            return Ok(new HeadMultipleResponse { Needs = missingBlobs.ToArray()});
+            return Ok(new ExistCheckMultipleContentIdResponse { Needs = partialContentIds.ToArray()});
         }
 
-        private async Task<BlobContents> GetImpl(NamespaceId ns, BlobIdentifier contentId)
+        private async Task<BlobContents> GetImpl(NamespaceId ns, ContentId contentId)
         {
             BlobIdentifier[]? chunks = await _contentIdStore.Resolve(ns, contentId);
             if (chunks == null || chunks.Length == 0)
@@ -198,7 +199,7 @@ namespace Horde.Storage.Controllers
         [RequiredContentType(CustomMediaTypeNames.UnrealCompressedBuffer)]
         public async Task<IActionResult> Put(
             [Required] NamespaceId ns,
-            [Required] BlobIdentifier id)
+            [Required] ContentId id)
         {
             AuthorizationResult authorizationResult = await _authorizationService.AuthorizeAsync(User, ns, NamespaceAccessRequirement.Name);
 
@@ -213,7 +214,7 @@ namespace Horde.Storage.Controllers
 
             try
             {
-                BlobIdentifier identifier = await PutImpl(ns, id, payload);
+                ContentId identifier = await PutImpl(ns, id, payload);
 
                 return Ok(new
                 {
@@ -229,7 +230,7 @@ namespace Horde.Storage.Controllers
             }
         }
 
-        private async Task<BlobIdentifier> PutImpl(NamespaceId ns, BlobIdentifier id, IBufferedPayload payload)
+        private async Task<ContentId> PutImpl(NamespaceId ns, ContentId id, IBufferedPayload payload)
         {
             // decompress the content and generate a identifier from it to verify the identifier we got
             await using Stream decompressStream = payload.GetStream();
@@ -237,7 +238,7 @@ namespace Horde.Storage.Controllers
             byte[] decompressedContent = _compressedBufferUtils.DecompressContent(await decompressStream.ToByteArray());
 
             MemoryStream decompressedStream = new MemoryStream(decompressedContent);
-            BlobIdentifier identifierDecompressedPayload = await _storage.VerifyContentMatchesHash(decompressedStream, id);
+            ContentId identifierDecompressedPayload = ContentId.FromContentHash(await _storage.VerifyContentMatchesHash(decompressedStream, id));
 
             BlobIdentifier identifierCompressedPayload;
             {
@@ -303,4 +304,10 @@ namespace Horde.Storage.Controllers
             await _storage.DeleteObject(ns, id);
         }*/
     }
+
+    
+    public class ExistCheckMultipleContentIdResponse
+    {
+        public ContentId[] Needs { get; set; } = null!;
+    }
 }
diff --git a/Engine/Source/Programs/Horde/HordeStorage/Horde.Storage/Controllers/ObjectController.cs b/Engine/Source/Programs/Horde/HordeStorage/Horde.Storage/Controllers/ObjectController.cs
--- a/Engine/Source/Programs/Horde/HordeStorage/Horde.Storage/Controllers/ObjectController.cs
+++ b/Engine/Source/Programs/Horde/HordeStorage/Horde.Storage/Controllers/ObjectController.cs
@@ -179,9 +179,12 @@ namespace Horde.Storage.Controllers
             }
             catch (PartialReferenceResolveException e)
             {
-                return BadRequest(new ValidationProblemDetails {Title = $"Object {id} is missing blobs", Detail = $"Following blobs are missing: {string.Join(",", e.UnresolvedReferences)}"});
+                return BadRequest(new ValidationProblemDetails {Title = $"Object {id} is missing content ids", Detail = $"Following content ids are invalid: {string.Join(",", e.UnresolvedReferences)}"});
+            }
+            catch (ReferenceIsMissingBlobsException e)
+            {
+                return BadRequest(new ValidationProblemDetails {Title = $"Object {id} is missing blobs", Detail = $"Following blobs are missing: {string.Join(",", e.MissingBlobs)}"});
             }
-
         }
 
 
diff --git a/Engine/Source/Programs/Horde/HordeStorage/Horde.Storage/Controllers/ReferencesController.cs b/Engine/Source/Programs/Horde/HordeStorage/Horde.Storage/Controllers/ReferencesController.cs
--- a/Engine/Source/Programs/Horde/HordeStorage/Horde.Storage/Controllers/ReferencesController.cs
+++ b/Engine/Source/Programs/Horde/HordeStorage/Horde.Storage/Controllers/ReferencesController.cs
@@ -388,8 +388,11 @@ namespace Horde.Storage.Controllers
             }
 
 
-            PutObjectResult result = await _objectService.Put(ns, bucket, key, blobHeader, payloadObject);
-            return Ok(new PutObjectResponse(result.MissingReferences));
+            (ContentId[] missingReferences, BlobIdentifier[] missingBlobs) = await _objectService.Put(ns, bucket, key, blobHeader, payloadObject);
+
+            List<ContentHash> missingHashes = new List<ContentHash>(missingReferences);
+            missingHashes.AddRange(missingBlobs);
+            return Ok(new PutObjectResponse(missingHashes.ToArray()));
         }
 
         [HttpPost("{ns}/{bucket}/{key}/finalize/{hash}.{format?}")]
@@ -413,8 +416,11 @@ namespace Horde.Storage.Controllers
                 }
             }
 
-            BlobIdentifier[] missingReferences = await _objectService.Finalize(ns, bucket, key, hash);
-            return Ok(new PutObjectResponse(missingReferences));
+            (ContentId[] missingReferences, BlobIdentifier[] missingBlobs) = await _objectService.Finalize(ns, bucket, key, hash);
+            List<ContentHash> missingHashes = new List<ContentHash>(missingReferences);
+            missingHashes.AddRange(missingBlobs);
+
+            return Ok(new PutObjectResponse(missingHashes.ToArray()));
         }
 
 
@@ -530,11 +536,11 @@ namespace Horde.Storage.Controllers
 
     public class PutObjectResponse
     {
-        public PutObjectResponse(BlobIdentifier[] missingReferences)
+        public PutObjectResponse(ContentHash[] missingReferences)
         {
             Needs = missingReferences;
         }
 
-        public BlobIdentifier[] Needs { get; set; }
+        public ContentHash[] Needs { get; set; }
     }
 }
diff --git a/Engine/Source/Programs/Horde/HordeStorage/Horde.Storage/Implementation/Blob/IBlobService.cs b/Engine/Source/Programs/Horde/HordeStorage/Horde.Storage/Implementation/Blob/IBlobService.cs
--- a/Engine/Source/Programs/Horde/HordeStorage/Horde.Storage/Implementation/Blob/IBlobService.cs
+++ b/Engine/Source/Programs/Horde/HordeStorage/Horde.Storage/Implementation/Blob/IBlobService.cs
@@ -21,7 +21,7 @@ namespace Horde.Storage.Implementation;
 
 public interface IBlobService
 {
-    Task<BlobIdentifier> VerifyContentMatchesHash(Stream content, BlobIdentifier identifier);
+    Task<ContentHash> VerifyContentMatchesHash(Stream content, ContentHash identifier);
     Task<BlobIdentifier> PutObjectKnownHash(NamespaceId ns, IBufferedPayload content, BlobIdentifier identifier);
     Task<BlobIdentifier> PutObject(NamespaceId ns, IBufferedPayload payload, BlobIdentifier identifier);
     Task<BlobIdentifier> PutObject(NamespaceId ns, byte[] payload, BlobIdentifier identifier);
@@ -85,9 +85,9 @@ public class BlobService : IBlobService
         return store;
     }
 
-    public async Task<BlobIdentifier> VerifyContentMatchesHash(Stream content, BlobIdentifier identifier)
+    public async Task<ContentHash> VerifyContentMatchesHash(Stream content, ContentHash identifier)
     {
-        BlobIdentifier blobHash;
+        ContentHash blobHash;
         {
             using IScope _ = Tracer.Instance.StartActive("web.hash");
             blobHash = await BlobIdentifier.FromStream(content);
@@ -106,7 +106,7 @@ public class BlobService : IBlobService
         using IScope _ = Tracer.Instance.StartActive("put_blob");
 
         await using Stream hashStream = payload.GetStream();
-        BlobIdentifier id = await VerifyContentMatchesHash(hashStream, identifier);
+        BlobIdentifier id = BlobIdentifier.FromContentHash(await VerifyContentMatchesHash(hashStream, identifier));
 
         Task<BlobIdentifier> putObjectTask = PutObjectToStores(ns, payload, id);
         Task addToBlobIndexTask = _blobIndex.AddBlobToIndex(ns, id);
@@ -122,7 +122,7 @@ public class BlobService : IBlobService
         using IScope _ = Tracer.Instance.StartActive("put_blob");
 
         await using Stream hashStream = new MemoryStream(payload);
-        BlobIdentifier id = await VerifyContentMatchesHash(hashStream, identifier);
+        BlobIdentifier id = BlobIdentifier.FromContentHash(await VerifyContentMatchesHash(hashStream, identifier));
 
         Task<BlobIdentifier> putObjectTask = PutObjectToStores(ns, payload, id);
         Task addToBlobIndexTask = _blobIndex.AddBlobToIndex(ns, id);
diff --git a/Engine/Source/Programs/Horde/HordeStorage/Horde.Storage/Implementation/ContentId/IContentIdStore.cs b/Engine/Source/Programs/Horde/HordeStorage/Horde.Storage/Implementation/ContentId/IContentIdStore.cs
--- a/Engine/Source/Programs/Horde/HordeStorage/Horde.Storage/Implementation/ContentId/IContentIdStore.cs
+++ b/Engine/Source/Programs/Horde/HordeStorage/Horde.Storage/Implementation/ContentId/IContentIdStore.cs
@@ -14,7 +14,7 @@ namespace Horde.Storage.Implementation
         /// <param name="ns">The namespace to operate in</param>
         /// <param name="contentId">The identifier for the content id</param>
         /// <returns></returns>
-        Task<BlobIdentifier[]?> Resolve(NamespaceId ns, BlobIdentifier contentId);
+        Task<BlobIdentifier[]?> Resolve(NamespaceId ns, ContentId contentId);
 
         /// <summary>
         /// Add a mapping from contentId to blobIdentifier
@@ -24,12 +24,12 @@ namespace Horde.Storage.Implementation
         /// <param name="blobIdentifier">The blob the content id maps to</param>
         /// <param name="contentWeight">Weight of this identifier compared to previous mappings, used to determine which is more important, lower weight is considered a better fit</param>
         /// <returns></returns>
-        Task Put(NamespaceId ns, BlobIdentifier contentId, BlobIdentifier blobIdentifier, int contentWeight);
+        Task Put(NamespaceId ns, ContentId contentId, BlobIdentifier blobIdentifier, int contentWeight);
     }
 
     public class InvalidContentIdException : Exception
     {
-        public InvalidContentIdException(BlobIdentifier contentId) : base($"Unknown content id {contentId}")
+        public InvalidContentIdException(ContentId contentId) : base($"Unknown content id {contentId}")
         {
             
         }
@@ -37,9 +37,9 @@ namespace Horde.Storage.Implementation
 
     public class ContentIdResolveException : Exception
     {
-        public BlobIdentifier ContentId { get; }
+        public ContentId ContentId { get; }
 
-        public ContentIdResolveException(BlobIdentifier contentId): base($"Unable to find any mapping of contentId {contentId} that has all blobs present")
+        public ContentIdResolveException(ContentId contentId): base($"Unable to find any mapping of contentId {contentId} that has all blobs present")
         {
             ContentId = contentId;
         }
diff --git a/Engine/Source/Programs/Horde/HordeStorage/Horde.Storage/Implementation/ContentId/MemoryContentIdStore.cs b/Engine/Source/Programs/Horde/HordeStorage/Horde.Storage/Implementation/ContentId/MemoryContentIdStore.cs
--- a/Engine/Source/Programs/Horde/HordeStorage/Horde.Storage/Implementation/ContentId/MemoryContentIdStore.cs
+++ b/Engine/Source/Programs/Horde/HordeStorage/Horde.Storage/Implementation/ContentId/MemoryContentIdStore.cs
@@ -14,16 +14,16 @@ namespace Horde.Storage.Implementation
     public class MemoryContentIdStore : IContentIdStore
     {
         private readonly IBlobService _blobStore;
-        private readonly ConcurrentDictionary<NamespaceId, ConcurrentDictionary<BlobIdentifier, SortedList<int, BlobIdentifier[]>>> _contentIds = new ConcurrentDictionary<NamespaceId, ConcurrentDictionary<BlobIdentifier, SortedList<int, BlobIdentifier[]>>>();
+        private readonly ConcurrentDictionary<NamespaceId, ConcurrentDictionary<ContentId, SortedList<int, BlobIdentifier[]>>> _contentIds = new ConcurrentDictionary<NamespaceId, ConcurrentDictionary<ContentId, SortedList<int, BlobIdentifier[]>>>();
         
         public MemoryContentIdStore(IBlobService blobStore)
         {
             _blobStore = blobStore;
         }
 
-        public async Task<BlobIdentifier[]?> Resolve(NamespaceId ns, BlobIdentifier contentId)
+        public async Task<BlobIdentifier[]?> Resolve(NamespaceId ns, ContentId contentId)
         {
-            if (_contentIds.TryGetValue(ns, out ConcurrentDictionary<BlobIdentifier, SortedList<int, BlobIdentifier[]>>? contentIdsForNamespace))
+            if (_contentIds.TryGetValue(ns, out ConcurrentDictionary<ContentId, SortedList<int, BlobIdentifier[]>>? contentIdsForNamespace))
             {
                 if (contentIdsForNamespace.TryGetValue(contentId, out SortedList<int, BlobIdentifier[]>? contentIdMappings))
                 {
@@ -37,18 +37,19 @@ namespace Horde.Storage.Implementation
                 }
             }
 
+            BlobIdentifier uncompressedBlobIdentifier = contentId.AsBlobIdentifier();
             // if no content id is found, but we have a blob that matches the content id (so a unchunked and uncompressed version of the data) we use that instead
-            if (await _blobStore.Exists(ns, contentId))
-                return new[] { contentId };
+            if (await _blobStore.Exists(ns, uncompressedBlobIdentifier))
+                return new[] { uncompressedBlobIdentifier };
 
             return null;
         }
 
-        public Task Put(NamespaceId ns, BlobIdentifier contentId, BlobIdentifier blobIdentifier, int contentWeight)
+        public Task Put(NamespaceId ns, ContentId contentId, BlobIdentifier blobIdentifier, int contentWeight)
         {
             _contentIds.AddOrUpdate(ns, (_) =>
             {
-                ConcurrentDictionary<BlobIdentifier, SortedList<int, BlobIdentifier[]>> dict = new()
+                ConcurrentDictionary<ContentId, SortedList<int, BlobIdentifier[]>> dict = new()
                 {
                     [contentId] = new SortedList<int, BlobIdentifier[]>
                     {
diff --git a/Engine/Source/Programs/Horde/HordeStorage/Horde.Storage/Implementation/ContentId/ScyllaContentIdStore.cs b/Engine/Source/Programs/Horde/HordeStorage/Horde.Storage/Implementation/ContentId/ScyllaContentIdStore.cs
--- a/Engine/Source/Programs/Horde/HordeStorage/Horde.Storage/Implementation/ContentId/ScyllaContentIdStore.cs
+++ b/Engine/Source/Programs/Horde/HordeStorage/Horde.Storage/Implementation/ContentId/ScyllaContentIdStore.cs
@@ -34,14 +34,14 @@ namespace Horde.Storage.Implementation
         }
 
 
-        public async Task<BlobIdentifier[]?> Resolve(NamespaceId ns, BlobIdentifier contentId)
+        public async Task<BlobIdentifier[]?> Resolve(NamespaceId ns, ContentId contentId)
         {
             using IScope scope = Tracer.Instance.StartActive("ScyllaContentIdStore.ResolveContentId");
             scope.Span.ResourceName = contentId.ToString();
 
-            Task<bool> blobStoreExistsTask = _blobStore.Exists(ns, contentId);
+            BlobIdentifier contentIdBlob = contentId.AsBlobIdentifier();
+            Task<bool> blobStoreExistsTask = _blobStore.Exists(ns, contentIdBlob);
 
-            BlobIdentifier[]? resolvedBlobs = null;
             // lower content_weight means its a better candidate to resolve to
             foreach (ScyllaContentId? resolvedContentId in await _mapper.FetchAsync<ScyllaContentId>("WHERE content_id = ? ORDER BY content_weight DESC", new ScyllaBlobIdentifier(contentId)))
             {
@@ -58,27 +58,19 @@ namespace Horde.Storage.Implementation
                         return blobs;
                 }
                 // blobs are missing continue testing with the next content id in the weighted list as that might exist
-
-                resolvedBlobs = blobs;
             }
 
             // if no content id is found, but we have a blob that matches the content id (so a unchunked and uncompressed version of the data) we use that instead
             bool contentIdBlobExists = await blobStoreExistsTask;
 
             if (contentIdBlobExists)
-                return new[] { contentId };
+                return new[] { contentIdBlob };
             
-            if (resolvedBlobs != null)
-            {
-                // content id found, but blobs is missing, still we know which blobs it was so we return that list
-                return resolvedBlobs;
-            }
-
             // unable to resolve the content id
             return null;
         }
 
-        public async Task Put(NamespaceId ns, BlobIdentifier contentId, BlobIdentifier blobIdentifier, int contentWeight)
+        public async Task Put(NamespaceId ns, ContentId contentId, BlobIdentifier blobIdentifier, int contentWeight)
         {
             await _mapper.UpdateAsync<ScyllaContentId>("SET chunks = ? WHERE content_id = ? AND content_weight = ?", new [] {new ScyllaBlobIdentifier(blobIdentifier)}, new ScyllaBlobIdentifier(contentId), contentWeight);
         }
diff --git a/Engine/Source/Programs/Horde/HordeStorage/Horde.Storage/Implementation/IReferenceResolver.cs b/Engine/Source/Programs/Horde/HordeStorage/Horde.Storage/Implementation/IReferenceResolver.cs
--- a/Engine/Source/Programs/Horde/HordeStorage/Horde.Storage/Implementation/IReferenceResolver.cs
+++ b/Engine/Source/Programs/Horde/HordeStorage/Horde.Storage/Implementation/IReferenceResolver.cs
@@ -32,9 +32,10 @@ namespace Horde.Storage.Implementation
             // TODO: This is cacheable and we should store the result somewhere
             Queue<CompactBinaryObject> objectsToVisit = new Queue<CompactBinaryObject>();
             objectsToVisit.Enqueue(cb);
-            List<BlobIdentifier> unresolvedReferences = new List<BlobIdentifier>();
+            List<ContentId> unresolvedContentIdReferences = new List<ContentId>();
+            List<BlobIdentifier> unresolvedBlobReferences = new List<BlobIdentifier>();
 
-            List<Task<(BlobIdentifier, BlobIdentifier[]?)>> pendingContentIdResolves = new();
+            List<Task<(ContentId, BlobIdentifier[]?)>> pendingContentIdResolves = new();
             List<Task<CompactBinaryObject>> pendingCompactBinaryAttachments = new();
 
             while(pendingCompactBinaryAttachments.Count != 0 || pendingContentIdResolves.Count != 0 || objectsToVisit.Count != 0)
@@ -53,7 +54,7 @@ namespace Horde.Storage.Implementation
 
                         if (field.IsBinaryAttachment())
                         {
-                            pendingContentIdResolves.Add(ResolveContentId(ns, blobIdentifier));
+                            pendingContentIdResolves.Add(ResolveContentId(ns, ContentId.FromBlobIdentifier(blobIdentifier)));
                         }
                         else
                         {
@@ -69,14 +70,14 @@ namespace Horde.Storage.Implementation
                     }
                 }
 
-                List<Task<(BlobIdentifier, BlobIdentifier[]?)>> contentIdResolvesToRemove = new();
+                List<Task<(ContentId, BlobIdentifier[]?)>> contentIdResolvesToRemove = new();
 
-                foreach (Task<(BlobIdentifier, BlobIdentifier[]?)> pendingContentIdResolveTask in pendingContentIdResolves)
+                foreach (Task<(ContentId, BlobIdentifier[]?)> pendingContentIdResolveTask in pendingContentIdResolves)
                 {
                     // check for any content id resolve that has finished and return those blobs it found
                     if (pendingContentIdResolveTask.IsCompleted)
                     {
-                        (BlobIdentifier id, BlobIdentifier[]? resolvedBlobs) = await pendingContentIdResolveTask;
+                        (ContentId contentId, BlobIdentifier[]? resolvedBlobs) = await pendingContentIdResolveTask;
                         if (resolvedBlobs != null)
                         {
                             foreach (BlobIdentifier b in resolvedBlobs)
@@ -86,14 +87,14 @@ namespace Horde.Storage.Implementation
                         }
                         else
                         {
-                            unresolvedReferences.Add(id);
+                            unresolvedContentIdReferences.Add(contentId);
                         }
                         contentIdResolvesToRemove.Add(pendingContentIdResolveTask);
                     }
                 }
 
                 // cleanup finished tasks
-                foreach (Task<(BlobIdentifier, BlobIdentifier[]?)> finishedTask in contentIdResolvesToRemove)
+                foreach (Task<(ContentId, BlobIdentifier[]?)> finishedTask in contentIdResolvesToRemove)
                 {
                     pendingContentIdResolves.Remove(finishedTask);
                 }
@@ -111,7 +112,7 @@ namespace Horde.Storage.Implementation
                         }
                         catch (BlobNotFoundException e)
                         {
-                            unresolvedReferences.Add(e.Blob);
+                            unresolvedBlobReferences.Add(e.Blob);
                         }
                         finishedCompactBinaryResolves.Add(pendingCompactBinaryAttachment);
                     }
@@ -124,9 +125,14 @@ namespace Horde.Storage.Implementation
                 }
             }
 
-            if (unresolvedReferences.Count != 0)
+            if (unresolvedContentIdReferences.Count != 0)
             {
-                throw new PartialReferenceResolveException(unresolvedReferences);
+                throw new PartialReferenceResolveException(unresolvedContentIdReferences);
+            }
+
+            if (unresolvedBlobReferences.Count != 0)
+            {
+                throw new ReferenceIsMissingBlobsException(unresolvedBlobReferences);
             }
         }
 
@@ -139,22 +145,32 @@ namespace Horde.Storage.Implementation
             return childBinaryObject;
         }
 
-        private async Task<(BlobIdentifier, BlobIdentifier[]?)> ResolveContentId(NamespaceId ns, BlobIdentifier blobIdentifier)
+        private async Task<(ContentId, BlobIdentifier[]?)> ResolveContentId(NamespaceId ns, ContentId contentId)
         {
             using IScope scope = Tracer.Instance.StartActive("ReferenceResolver.ResolveContentId");
-            scope.Span.ResourceName = blobIdentifier.ToString();
-            BlobIdentifier[]? resolvedBlobs = await _contentIdStore.Resolve(ns, blobIdentifier);
-            return (blobIdentifier, resolvedBlobs);
+            scope.Span.ResourceName = contentId.ToString();
+            BlobIdentifier[]? resolvedBlobs = await _contentIdStore.Resolve(ns, contentId);
+            return (contentId, resolvedBlobs);
         }
     }
 
     public class PartialReferenceResolveException : Exception
     {
-        public List<BlobIdentifier> UnresolvedReferences { get; }
+        public List<ContentId> UnresolvedReferences { get; }
 
-        public PartialReferenceResolveException(List<BlobIdentifier> unresolvedReferences) : base($"References missing: {string.Join(',', unresolvedReferences)}")
+        public PartialReferenceResolveException(List<ContentId> unresolvedReferences) : base($"References missing: {string.Join(',', unresolvedReferences)}")
         {
             UnresolvedReferences = unresolvedReferences;
         }
     }
+
+    public class ReferenceIsMissingBlobsException : Exception
+    {
+        public List<BlobIdentifier> MissingBlobs { get; }
+
+        public ReferenceIsMissingBlobsException(List<BlobIdentifier> missingBlobs) : base($"References is missing these blobs: {string.Join(',', missingBlobs)}")
+        {
+            MissingBlobs = missingBlobs;
+        }
+    }
 }
diff --git a/Engine/Source/Programs/Horde/HordeStorage/Horde.Storage/Implementation/ObjectService.cs b/Engine/Source/Programs/Horde/HordeStorage/Horde.Storage/Implementation/ObjectService.cs
--- a/Engine/Source/Programs/Horde/HordeStorage/Horde.Storage/Implementation/ObjectService.cs
+++ b/Engine/Source/Programs/Horde/HordeStorage/Horde.Storage/Implementation/ObjectService.cs
@@ -17,8 +17,8 @@ namespace Horde.Storage.Implementation
     public interface IObjectService
     {
         Task<(ObjectRecord, BlobContents?)> Get(NamespaceId ns, BucketId bucket, IoHashKey key, string[] fields);
-        Task<PutObjectResult> Put(NamespaceId ns, BucketId bucket, IoHashKey key, BlobIdentifier blobHash, CompactBinaryObject payload);
-        Task<BlobIdentifier[]> Finalize(NamespaceId ns, BucketId bucket, IoHashKey key, BlobIdentifier blobHash);
+        Task<(ContentId[], BlobIdentifier[])> Put(NamespaceId ns, BucketId bucket, IoHashKey key, BlobIdentifier blobHash, CompactBinaryObject payload);
+        Task<(ContentId[], BlobIdentifier[])> Finalize(NamespaceId ns, BucketId bucket, IoHashKey key, BlobIdentifier blobHash);
 
         IAsyncEnumerable<NamespaceId> GetNamespaces();
 
@@ -90,7 +90,7 @@ namespace Horde.Storage.Implementation
             return (o, blobContents);
         }
 
-        public async Task<PutObjectResult> Put(NamespaceId ns, BucketId bucket, IoHashKey key, BlobIdentifier blobHash, CompactBinaryObject payload)
+        public async Task<(ContentId[], BlobIdentifier[])> Put(NamespaceId ns, BucketId bucket, IoHashKey key, BlobIdentifier blobHash, CompactBinaryObject payload)
         {
             bool hasReferences = payload.GetAllFields().Any(field => field.IsAttachment());
 
@@ -101,7 +101,9 @@ namespace Horde.Storage.Implementation
 
             Task<BlobIdentifier> blobStorePut = _blobService.PutObject(ns, payload.Data, blobHash);
             Task addRefToBodyTask = _blobIndex.AddRefToBlobs(ns, bucket, key, new [] {blobHash});
-            BlobIdentifier[] missingReferences = Array.Empty<BlobIdentifier>();
+            ContentId[] missingReferences = Array.Empty<ContentId>();
+            BlobIdentifier[] missingBlobs = Array.Empty<BlobIdentifier>();
+
             if (hasReferences)
             {
                 using IScope _ = Tracer.Instance.StartActive("ObjectService.ResolveReferences");
@@ -113,26 +115,30 @@ namespace Horde.Storage.Implementation
                     Task addRefsTask = _blobIndex.AddRefToBlobs(ns, bucket, key, referencesArray);
                     Task<BlobIdentifier[]> checkReferencesTask = _blobService.FilterOutKnownBlobs(ns, referencesArray);
                     await Task.WhenAll(addRefsTask, checkReferencesTask);
-                    missingReferences = await checkReferencesTask;
+                    missingBlobs = await checkReferencesTask;
                 }
                 catch (PartialReferenceResolveException e)
                 {
                     missingReferences = e.UnresolvedReferences.ToArray();
                 }
+                catch (ReferenceIsMissingBlobsException e)
+                {
+                    missingBlobs = e.MissingBlobs.ToArray();
+                }
             }
 
             await Task.WhenAll(objectStorePut, blobStorePut, addRefToBodyTask);
-
-            if (missingReferences.Length == 0)
+            
+            if (missingReferences.Length == 0 && missingBlobs.Length == 0)
             {
                 await _referencesStore.Finalize(ns, bucket, key);
                 await _replicationLog.InsertAddEvent(ns, bucket, key, blobHash);
             }
 
-            return new PutObjectResult(missingReferences);
+            return (missingReferences, missingBlobs);
         }
 
-        public async Task<BlobIdentifier[]> Finalize(NamespaceId ns, BucketId bucket, IoHashKey key, BlobIdentifier blobHash)
+        public async Task<(ContentId[], BlobIdentifier[])> Finalize(NamespaceId ns, BucketId bucket, IoHashKey key, BlobIdentifier blobHash)
         {
             (ObjectRecord o, BlobContents? blob) = await Get(ns, bucket, key);
             if (blob == null)
@@ -145,28 +151,33 @@ namespace Horde.Storage.Implementation
 
             bool hasReferences = payload.GetAllFields().Any(field => field.IsAttachment());
 
-            BlobIdentifier[] missingReferences = Array.Empty<BlobIdentifier>();
+            ContentId[] missingReferences = Array.Empty<ContentId>();
+            BlobIdentifier[] missingBlobs = Array.Empty<BlobIdentifier>();
             if (hasReferences)
             {
                 using IScope _ = Tracer.Instance.StartActive("ObjectService.ResolveReferences");
                 try
                 {
                     IAsyncEnumerable<BlobIdentifier> references = _referenceResolver.ResolveReferences(ns, payload);
-                    missingReferences = await _blobService.FilterOutKnownBlobs(ns, references);
+                    missingBlobs = await _blobService.FilterOutKnownBlobs(ns, references);
                 }
                 catch (PartialReferenceResolveException e)
                 {
                     missingReferences = e.UnresolvedReferences.ToArray();
                 }
+                catch (ReferenceIsMissingBlobsException e)
+                {
+                    missingBlobs = e.MissingBlobs.ToArray();
+                }
             }
 
-            if (missingReferences.Length == 0)
+            if (missingReferences.Length == 0 && missingBlobs.Length == 0)
             {
                 await _referencesStore.Finalize(ns, bucket, key);
                 await _replicationLog.InsertAddEvent(ns, bucket, key, blobHash);
             }
 
-            return missingReferences;
+            return (missingReferences, missingBlobs);
         }
 
         public IAsyncEnumerable<NamespaceId> GetNamespaces()
@@ -189,14 +200,4 @@ namespace Horde.Storage.Implementation
             return _referencesStore.DeleteBucket(ns, bucket);
         }
     }
-
-    public class PutObjectResult
-    {
-        public PutObjectResult(BlobIdentifier[] missingReferences)
-        {
-            MissingReferences = missingReferences;
-        }
-
-        public BlobIdentifier[] MissingReferences { get; set; }
-    }
 }
diff --git a/Engine/Source/Programs/Horde/HordeStorage/Horde.Storage/Tests/Functional/ContentId/CompressedBlobTests.cs b/Engine/Source/Programs/Horde/HordeStorage/Horde.Storage/Tests/Functional/ContentId/CompressedBlobTests.cs
--- a/Engine/Source/Programs/Horde/HordeStorage/Horde.Storage/Tests/Functional/ContentId/CompressedBlobTests.cs
+++ b/Engine/Source/Programs/Horde/HordeStorage/Horde.Storage/Tests/Functional/ContentId/CompressedBlobTests.cs
@@ -126,7 +126,7 @@ namespace Horde.Storage.FunctionalTests.CompressedBlobs
         {
             byte[] texturePayload = await File.ReadAllBytesAsync("ContentId/Payloads/UncompressedTexture_CAS_dea81b6c3b565bb5089695377c98ce0f1c13b0c3.udd");
             BlobIdentifier compressedPayloadIdentifier = BlobIdentifier.FromBlob(texturePayload);
-            BlobIdentifier uncompressedPayloadIdentifier = new BlobIdentifier("DEA81B6C3B565BB5089695377C98CE0F1C13B0C3");
+            ContentId uncompressedPayloadIdentifier = new ContentId("DEA81B6C3B565BB5089695377C98CE0F1C13B0C3");
 
             {
                 ByteArrayContent content = new ByteArrayContent(texturePayload);
@@ -135,8 +135,9 @@ namespace Horde.Storage.FunctionalTests.CompressedBlobs
                 result.EnsureSuccessStatusCode();
 
                 InsertResponse response = await result.Content.ReadAsAsync<InsertResponse>();
+                Assert.IsNotNull(response.Identifier);
                 Assert.AreNotEqual(compressedPayloadIdentifier, response.Identifier);
-                Assert.AreEqual(uncompressedPayloadIdentifier, response.Identifier);
+                Assert.AreEqual(uncompressedPayloadIdentifier, ContentId.FromBlobIdentifier(response.Identifier));
             }
 
 
diff --git a/Engine/Source/Programs/Horde/HordeStorage/Horde.Storage/Tests/Functional/Objects/ReferencesTests.cs b/Engine/Source/Programs/Horde/HordeStorage/Horde.Storage/Tests/Functional/Objects/ReferencesTests.cs
--- a/Engine/Source/Programs/Horde/HordeStorage/Horde.Storage/Tests/Functional/Objects/ReferencesTests.cs
+++ b/Engine/Source/Programs/Horde/HordeStorage/Horde.Storage/Tests/Functional/Objects/ReferencesTests.cs
@@ -838,6 +838,81 @@ namespace Horde.Storage.FunctionalTests.References
             }
         }
 
+        
+        [TestMethod]
+        public async Task PutContentIdMissingBlob()
+        {
+            IContentIdStore? contentIdStore = _server!.Services.GetService<IContentIdStore>();
+            Assert.IsNotNull(contentIdStore);
+
+            // submit a object which contains a content id, which exists but points to a blob that does not exist
+            string blobContents = "This is a string that is referenced as a blob";
+            byte[] blobData = Encoding.ASCII.GetBytes(blobContents);
+            BlobIdentifier blobHash = BlobIdentifier.FromBlob(blobData);
+            ContentId contentId = new ContentId("0000000000000000000000000000000000000000");
+
+            await contentIdStore.Put(TestNamespace, contentId, blobHash, blobData.Length);
+
+            CompactBinaryWriter writer = new CompactBinaryWriter();
+            writer.BeginObject();
+            writer.AddBinaryAttachment(contentId.AsBlobIdentifier(), "blob");
+            writer.EndObject();
+
+            byte[] objectData = writer.Save();
+            BlobIdentifier objectHash = BlobIdentifier.FromBlob(objectData);
+            
+            IoHashKey key = IoHashKey.FromName("putContentIdMissingBlob");
+
+            {
+                HttpContent requestContent = new ByteArrayContent(objectData);
+                requestContent.Headers.ContentType = new MediaTypeHeaderValue(CustomMediaTypeNames.UnrealCompactBinary);
+                requestContent.Headers.Add(CommonHeaders.HashHeaderName, objectHash.ToString());
+
+                HttpResponseMessage result = await _httpClient!.PutAsync(requestUri: $"api/v1/refs/{TestNamespace}/bucket/{key}.uecb", requestContent);
+                result.EnsureSuccessStatusCode();
+
+                {
+                    Assert.AreEqual(result!.Content.Headers.ContentType!.MediaType, CustomMediaTypeNames.UnrealCompactBinary);
+                    await using MemoryStream ms = new MemoryStream();
+                    await result.Content.CopyToAsync(ms);
+                    byte[] roundTrippedBuffer = ms.ToArray();
+                    ReadOnlyMemory<byte> localMemory = new ReadOnlyMemory<byte>(roundTrippedBuffer);
+                    CompactBinaryObject cb = CompactBinaryObject.Load(ref localMemory);
+                    CompactBinaryField? needsField = cb["needs"];
+                    Assert.IsNotNull(needsField);
+                    List<BlobIdentifier?> missingBlobs = needsField!.AsArray().Select(field => field.AsHash()).ToList();
+                    Assert.AreEqual(1, missingBlobs.Count);
+
+                    Assert.AreNotEqual(blobHash, missingBlobs[0], "Refs should not be returning the mapped blob identifiers as this is unknown to the client attempting to put a new ref");
+                    Assert.AreEqual(contentId.AsBlobIdentifier(), missingBlobs[0]);
+                }
+            }
+
+            {
+                HttpContent requestContent = new ByteArrayContent(objectData);
+                requestContent.Headers.ContentType = new MediaTypeHeaderValue(CustomMediaTypeNames.UnrealCompactBinary);
+                requestContent.Headers.Add(CommonHeaders.HashHeaderName, objectHash.ToString());
+
+                HttpResponseMessage result = await _httpClient!.PutAsync(requestUri: $"api/v1/refs/{TestNamespace}/bucket/{key}.json", requestContent);
+                result.EnsureSuccessStatusCode();
+
+                {
+                    Assert.AreEqual(result!.Content.Headers.ContentType!.MediaType, MediaTypeNames.Application.Json);
+                    await using MemoryStream ms = new MemoryStream();
+                    await result.Content.CopyToAsync(ms);
+                    byte[] roundTrippedBuffer = ms.ToArray();
+                    string s = Encoding.ASCII.GetString(roundTrippedBuffer);
+                    JObject jObject = JObject.Parse(s);
+                    Assert.AreEqual(1, jObject.Children().Count());
+
+                    JToken? needs = jObject["needs"];
+                    BlobIdentifier[] missingBlobs = needs!.ToArray().Select(token => new BlobIdentifier(token.Value<string>()!)).ToArray();
+                    Assert.AreEqual(1, missingBlobs.Length);
+                    Assert.AreEqual(contentId.AsBlobIdentifier(), missingBlobs[0]);
+                }
+            }
+        }
+
         [TestMethod]
         public async Task PutAndFinalize()
         {
diff --git a/Engine/Source/Programs/Horde/HordeStorage/Horde.Storage/Tests/Functional/RefTests.cs b/Engine/Source/Programs/Horde/HordeStorage/Horde.Storage/Tests/Functional/RefTests.cs
--- a/Engine/Source/Programs/Horde/HordeStorage/Horde.Storage/Tests/Functional/RefTests.cs
+++ b/Engine/Source/Programs/Horde/HordeStorage/Horde.Storage/Tests/Functional/RefTests.cs
@@ -342,11 +342,13 @@ namespace Horde.Storage.FunctionalTests.Ref
             ReadOnlyMemory<byte> memory = new ReadOnlyMemory<byte>(buffer);
             CompactBinaryObject o = CompactBinaryObject.Load(ref memory);
             List<CompactBinaryField> fields = o.GetFields().ToList();
-            Assert.AreEqual(4, fields.Count);
+            Assert.AreEqual(5, fields.Count);
             Assert.AreEqual($"{TestNamespace}.bucket.testObjectWithMetadata", o["name"]!.AsString());
             // we do not support serializing random dictionaries to compact binary
             Assert.IsNull(o["metadata"]);
 
+            Assert.IsNotNull(o["contentHash"]);
+
             Assert.IsNotNull(o["lastAccessTime"]);
 
             Assert.IsNotNull(o["blob"]);
diff --git a/Engine/Source/Programs/Horde/HordeStorage/Jupiter.Common/CompactBinaryOutputFormatter.cs b/Engine/Source/Programs/Horde/HordeStorage/Jupiter.Common/CompactBinaryOutputFormatter.cs
--- a/Engine/Source/Programs/Horde/HordeStorage/Jupiter.Common/CompactBinaryOutputFormatter.cs
+++ b/Engine/Source/Programs/Horde/HordeStorage/Jupiter.Common/CompactBinaryOutputFormatter.cs
@@ -64,6 +64,11 @@ namespace Jupiter
                     BlobIdentifier blob = (BlobIdentifier) value!;
                     writer.AddBinaryAttachment(blob, name);
                 }
+                else if (propertyType == typeof(ContentHash)  || propertyType.IsAssignableTo(typeof(ContentHash)))
+                {
+                    ContentHash blob = (ContentHash) value!;
+                    writer.AddHash(blob, name);
+                }
                 else if (propertyType == typeof(string))
                 {
                     string s = (string) value!;
@@ -88,6 +93,11 @@ namespace Jupiter
                         BlobIdentifier[] blobs = (BlobIdentifier[])value;
                         writer.AddUniformArray(blobs, CompactBinaryFieldType.BinaryAttachment, name);
                     }
+                    else if (elementType == typeof(ContentHash) || elementType.IsAssignableTo(typeof(ContentHash)))
+                    {
+                        ContentHash[] blobs = (ContentHash[])value;
+                        writer.AddUniformArray(blobs, CompactBinaryFieldType.Hash, name);
+                    }
                     else if (elementType == typeof(string))
                     {
                         string[] strings = (string[])value;
diff --git a/Engine/Source/Programs/Horde/HordeStorage/Jupiter.Common/Implementation/CompactBinaryWriter.cs b/Engine/Source/Programs/Horde/HordeStorage/Jupiter.Common/Implementation/CompactBinaryWriter.cs
--- a/Engine/Source/Programs/Horde/HordeStorage/Jupiter.Common/Implementation/CompactBinaryWriter.cs
+++ b/Engine/Source/Programs/Horde/HordeStorage/Jupiter.Common/Implementation/CompactBinaryWriter.cs
@@ -149,6 +149,19 @@ namespace Jupiter.Implementation
             _fields.Add(field);
         }
 
+        public void AddHash(ContentHash hash, string? fieldName = null)
+        {
+            Field field = BeginField();
+            SetName(ref field, fieldName);
+
+            if (hash.HashData.Length != 20)
+                throw new Exception("Hash data is assumed to be 20 bytes");
+            field.Payload = hash.HashData;
+
+            EndField(ref field, CompactBinaryFieldType.Hash);
+        }
+
+
         public void AddBinaryAttachment(BlobIdentifier hash, string? fieldName = null)
         {
             Field field = BeginField();
@@ -281,9 +294,11 @@ namespace Jupiter.Implementation
                 case CompactBinaryFieldType.CompactBinaryAttachment:
                 case CompactBinaryFieldType.BinaryAttachment:
                 case CompactBinaryFieldType.Hash:
-                    if (member is BlobIdentifier blob)
+                   
+                    if (member != null && member.GetType().IsAssignableTo(typeof(ContentHash)))
                     {
-                        return blob.HashData;
+                        if (member is ContentHash hash)
+                            return hash.HashData;
                     }
 
                     break;
diff --git a/Engine/Source/Programs/Horde/HordeStorage/Jupiter.Common/Implementation/ContentId.cs b/Engine/Source/Programs/Horde/HordeStorage/Jupiter.Common/Implementation/ContentId.cs
new file mode 100644
--- /dev/null
+++ b/Engine/Source/Programs/Horde/HordeStorage/Jupiter.Common/Implementation/ContentId.cs
@@ -0,0 +1,113 @@
+// Copyright Epic Games, Inc. All Rights Reserved.
+
+using System;
+using System.ComponentModel;
+using System.Globalization;
+using Newtonsoft.Json;
+
+namespace Jupiter.Implementation
+{
+    [JsonConverter(typeof(ContentIdConverter))]
+    [TypeConverter(typeof(ContentIdTypeConverter))]
+    public class ContentId : ContentHash, IEquatable<ContentId>
+    {
+        public ContentId(byte[] identifier) : base(identifier)
+        {
+        }
+
+        [JsonConstructor]
+        public ContentId(string identifier) : base(identifier)
+        {
+
+        }
+
+        public override int GetHashCode()
+        {
+            return Comparer.GetHashCode(Identifier);
+        }
+
+        public bool Equals(ContentId? other)
+        {
+            if (other == null)
+                return false;
+
+            return Comparer.Equals(Identifier, other.Identifier);
+        }
+
+        public override bool Equals(object? obj)
+        {
+            if (ReferenceEquals(null, obj))
+            {
+                return false;
+            }
+
+            if (ReferenceEquals(this, obj))
+            {
+                return true;
+            }
+
+            if (obj.GetType() != this.GetType())
+            {
+                return false;
+            }
+
+            return Equals((ContentId) obj);
+        }
+
+        public static ContentId FromContentHash(ContentHash contentHash)
+        {
+            return new ContentId(contentHash.HashData);
+        }
+
+        public static ContentId FromBlobIdentifier(BlobIdentifier blobIdentifier)
+        {
+            return new ContentId(blobIdentifier.HashData);
+        }
+
+        public BlobIdentifier AsBlobIdentifier()
+        {
+            return new BlobIdentifier(HashData);
+        }
+    }
+
+    public class ContentIdTypeConverter : TypeConverter
+    {
+        public override bool CanConvertFrom(ITypeDescriptorContext? context, Type sourceType)
+        {  
+            if (sourceType == typeof(string))  
+            {  
+                return true;
+            }  
+            return base.CanConvertFrom(context, sourceType);
+        }  
+  
+        public override object? ConvertFrom(ITypeDescriptorContext? context, CultureInfo? culture, object value)  
+        {
+            if (value is string)
+            {
+                return new ContentId((string) value);
+            }
+
+            return base.ConvertFrom(context, culture, value);  
+        }  
+    }
+
+    public class ContentIdConverter : JsonConverter<ContentId?>
+    {
+        public override void WriteJson(JsonWriter writer, ContentId? value, JsonSerializer serializer)
+        {
+            writer.WriteValue(value!.ToString());
+        }
+
+        public override ContentId? ReadJson(JsonReader reader, Type objectType, ContentId? existingValue, bool hasExistingValue, JsonSerializer serializer)
+        {
+            if (reader.TokenType == JsonToken.Null) 
+                return null;
+            string? s = (string?)reader.Value;
+
+            if (s == null)
+                return null;
+            return new ContentId(s!);
+        }
+    }
+}
