Horde Storage - Added some more scopes to annotate the orphan blob cleanup of the new structured cache api

diff --git a/Engine/Source/Programs/Horde/HordeStorage/Horde.Storage/Implementation/GC/OrphanBlobCleanupRefs.cs b/Engine/Source/Programs/Horde/HordeStorage/Horde.Storage/Implementation/GC/OrphanBlobCleanupRefs.cs
--- a/Engine/Source/Programs/Horde/HordeStorage/Horde.Storage/Implementation/GC/OrphanBlobCleanupRefs.cs
+++ b/Engine/Source/Programs/Horde/HordeStorage/Horde.Storage/Implementation/GC/OrphanBlobCleanupRefs.cs
@@ -7,6 +7,7 @@ using System.Threading;
 using System.Threading.Tasks;
 using async_enumerable_dotnet;
 using Dasync.Collections;
+using Datadog.Trace;
 using Horde.Storage.Controllers;
 using Horde.Storage.Implementation.Blob;
 using Jupiter;
@@ -56,6 +57,9 @@ namespace Horde.Storage.Implementation
                 if (cancellationToken.IsCancellationRequested)
                     break;
 
+                using IScope scope = Tracer.Instance.StartActive("gc.blob.namespace");
+                scope.Span.ResourceName = @namespace.ToString();
+
                 // only consider blobs that have been around for 60 minutes
                 // this due to cases were blobs are uploaded first
                 DateTime cutoff = DateTime.Now.AddMinutes(-60);
@@ -102,6 +106,8 @@ namespace Horde.Storage.Implementation
 
                     if (!found)
                     {
+                        using IScope removeBlobScope = Tracer.Instance.StartActive("gc.blob.remove-blob");
+                        removeBlobScope.Span.ResourceName = blob.ToString();
                         await RemoveBlob(@namespace, blob);
                         removedBlobs.Add(blob);
                     }
