Horde Storage - Accept lower case admin claims to grant access to the internal namespace, this is the casing of the admin claim we get from Okta.

diff --git a/Engine/Source/Programs/Horde/HordeStorage/Jupiter.Common/Implementation/NamespacePolicyResolver.cs b/Engine/Source/Programs/Horde/HordeStorage/Jupiter.Common/Implementation/NamespacePolicyResolver.cs
--- a/Engine/Source/Programs/Horde/HordeStorage/Jupiter.Common/Implementation/NamespacePolicyResolver.cs
+++ b/Engine/Source/Programs/Horde/HordeStorage/Jupiter.Common/Implementation/NamespacePolicyResolver.cs
@@ -31,7 +31,7 @@ namespace Jupiter.Common
                 return new NamespaceSettings.PerNamespaceSettings()
                 {
                     // the internal jupiter namespace is sensitive so only admins should have access, but we can still use the default storage pool
-                    Claims = new string[] { "Admin" },
+                    Claims = new string[] { "Admin", "admin" },
                     StoragePool = ""
                 };
             }
