Horde Storage - Fixes and polish to the snapshot builder
The snapshots currently produced are corrupted, most likely due using the wrong memory stream operation (fetching the underlying buffer which can be larger then what has been written to it due to it grows).

Also made sure we use the high level IObjectService to submit the snapshots instead of the low level ReferenceStore, which means we automatically verify that it is finalized, do the correct tracking and other such things.

diff --git a/Engine/Source/Programs/Horde/HordeStorage/Horde.Storage/Implementation/TransactionLog/ReplicationLogSnapshotBuilder.cs b/Engine/Source/Programs/Horde/HordeStorage/Horde.Storage/Implementation/TransactionLog/ReplicationLogSnapshotBuilder.cs
--- a/Engine/Source/Programs/Horde/HordeStorage/Horde.Storage/Implementation/TransactionLog/ReplicationLogSnapshotBuilder.cs
+++ b/Engine/Source/Programs/Horde/HordeStorage/Horde.Storage/Implementation/TransactionLog/ReplicationLogSnapshotBuilder.cs
@@ -2,6 +2,7 @@
 
 using System;
 using System.IO;
+using System.Linq;
 using System.Threading;
 using System.Threading.Tasks;
 using Jupiter.Implementation;
@@ -12,13 +13,13 @@ namespace Horde.Storage.Implementation.TransactionLog
     {
         private readonly IReplicationLog _replicationLog;
         private readonly IBlobService _blobService;
-        private readonly IReferencesStore _referencesStore;
+        private readonly IObjectService _objectService;
 
-        public ReplicationLogSnapshotBuilder(IReplicationLog replicationLog, IBlobService blobService, IReferencesStore referencesStore)
+        public ReplicationLogSnapshotBuilder(IReplicationLog replicationLog, IBlobService blobService, IObjectService objectService)
         {
             _replicationLog = replicationLog;
             _blobService = blobService;
-            _referencesStore = referencesStore;
+            _objectService = objectService;
         }
 
         public async Task<BlobIdentifier> BuildSnapshot(NamespaceId ns, NamespaceId storeInNamespace, CancellationToken cancellationToken = default(CancellationToken))
@@ -62,13 +63,19 @@ namespace Horde.Storage.Implementation.TransactionLog
                 snapshot.ProcessEvent(entry);
             }
 
+
+            byte[] buf;
             {
                 await using MemoryStream ms = new MemoryStream();
                 await snapshot.Serialize(ms);
-                byte[] buf = ms.GetBuffer();
+                await ms.FlushAsync(cancellationToken);
+                buf = ms.ToArray();
+            }
+            
+            
+            {
                 BlobIdentifier blobIdentifier = BlobIdentifier.FromBlob(buf);
 
-
                 CompactBinaryWriter writer = new CompactBinaryWriter();
                 writer.BeginObject();
                 writer.AddBinaryAttachment(blobIdentifier, "snapshotBlob");
@@ -77,8 +84,16 @@ namespace Horde.Storage.Implementation.TransactionLog
 
                 byte[] cbObjectBytes = writer.Save();
                 BlobIdentifier cbBlobId = BlobIdentifier.FromBlob(cbObjectBytes);
-                await _referencesStore.Put(storeInNamespace, new BucketId("snapshot"), new IoHashKey(blobIdentifier.ToString()), cbBlobId, cbObjectBytes, true);
+
+                if (cancellationToken.IsCancellationRequested)
+                    throw new TaskCanceledException();
+                
+                // upload the attachment first so we are not missing any references when we go to create the ref
                 await _blobService.PutObject(storeInNamespace, buf, blobIdentifier);
+                
+                PutObjectResult result = await _objectService.Put(storeInNamespace, new BucketId("snapshot"), new IoHashKey(blobIdentifier.ToString()), cbBlobId, CompactBinaryObject.Load(cbObjectBytes));
+                if (result.MissingReferences.Length != 0)
+                    throw new Exception($"Failed to upload snapshot to object service, missing references {string.Join(',' , result.MissingReferences.Select(b => b.ToString()))}");
 
                 if (cancellationToken.IsCancellationRequested)
                     throw new TaskCanceledException();
