Horde Storage - Build fixes to the replication snapshot builder

diff --git a/Engine/Source/Programs/Horde/HordeStorage/Horde.Storage/Implementation/TransactionLog/ReplicationLogSnapshotBuilder.cs b/Engine/Source/Programs/Horde/HordeStorage/Horde.Storage/Implementation/TransactionLog/ReplicationLogSnapshotBuilder.cs
--- a/Engine/Source/Programs/Horde/HordeStorage/Horde.Storage/Implementation/TransactionLog/ReplicationLogSnapshotBuilder.cs
+++ b/Engine/Source/Programs/Horde/HordeStorage/Horde.Storage/Implementation/TransactionLog/ReplicationLogSnapshotBuilder.cs
@@ -1,6 +1,7 @@
 ﻿// Copyright Epic Games, Inc. All Rights Reserved.
 
 using System;
+using System.Collections.Generic;
 using System.IO;
 using System.Linq;
 using System.Threading;
@@ -91,9 +92,11 @@ namespace Horde.Storage.Implementation.TransactionLog
                 // upload the attachment first so we are not missing any references when we go to create the ref
                 await _blobService.PutObject(storeInNamespace, buf, blobIdentifier);
                 
-                PutObjectResult result = await _objectService.Put(storeInNamespace, new BucketId("snapshot"), new IoHashKey(blobIdentifier.ToString()), cbBlobId, CompactBinaryObject.Load(cbObjectBytes));
-                if (result.MissingReferences.Length != 0)
-                    throw new Exception($"Failed to upload snapshot to object service, missing references {string.Join(',' , result.MissingReferences.Select(b => b.ToString()))}");
+                (ContentId[] missingContentIds, BlobIdentifier[] missingBlobs) = await _objectService.Put(storeInNamespace, new BucketId("snapshot"), new IoHashKey(blobIdentifier.ToString()), cbBlobId, CompactBinaryObject.Load(cbObjectBytes));
+                List<ContentHash> missingHashes = new List<ContentHash>(missingContentIds);
+                missingHashes.AddRange(missingBlobs);
+                if (missingHashes.Count != 0)
+                    throw new Exception($"Failed to upload snapshot to object service, missing references {string.Join(',' , missingHashes.Select(b => b.ToString()))}");
 
                 if (cancellationToken.IsCancellationRequested)
                     throw new TaskCanceledException();
