Horde: Suppress 404 on agent software channel
#rnx
#skipci

diff --git a/Engine/Source/Programs/Horde/HordeDashboard/src/backend/Backend.ts b/Engine/Source/Programs/Horde/HordeDashboard/src/backend/Backend.ts
--- a/Engine/Source/Programs/Horde/HordeDashboard/src/backend/Backend.ts
+++ b/Engine/Source/Programs/Horde/HordeDashboard/src/backend/Backend.ts
@@ -1126,7 +1126,7 @@ export class Backend {
 
     getAgentSoftwareChannel(name: string = "default"): Promise<GetAgentSoftwareChannelResponse> {
         return new Promise<GetAgentSoftwareChannelResponse>((resolve, reject) => {
-            this.backend.get(`/api/v1/agentsoftware/${name}`).then((value) => {
+            this.backend.get(`/api/v1/agentsoftware/${name}`, { suppress404: true }).then((value) => {
                 resolve(value.data as GetAgentSoftwareChannelResponse);
             }).catch(reason => {
                 reject(reason);
