Horde: Add a clear all for pinned jobs
#rnx
#skipci

diff --git a/Engine/Source/Programs/Horde/HordeDashboard/src/backend/Dashboard.ts b/Engine/Source/Programs/Horde/HordeDashboard/src/backend/Dashboard.ts
--- a/Engine/Source/Programs/Horde/HordeDashboard/src/backend/Dashboard.ts
+++ b/Engine/Source/Programs/Horde/HordeDashboard/src/backend/Dashboard.ts
@@ -65,6 +65,21 @@ export class Dashboard {
         this.setUpdated();
     }
 
+    clearPinnedJobs() {
+
+        if (!this.data.pinnedJobIds?.length) {
+            return;
+        }
+
+        const jobs = this.data.pinnedJobIds;
+
+        this.data.pinnedJobIds = [];
+
+        backend.updateUser({ removePinnedJobIds:jobs });
+
+        this.setUpdated();
+    }
+
     get username(): string {
         return this.data.name;
     }
