Horde Storage - Run mongo for tests on a seperate port to play nice with a service install

diff --git a/Engine/Source/Programs/Horde/HordeStorage/Horde.Storage/appsettings.Testing.json b/Engine/Source/Programs/Horde/HordeStorage/Horde.Storage/appsettings.Testing.json
--- a/Engine/Source/Programs/Horde/HordeStorage/Horde.Storage/appsettings.Testing.json
+++ b/Engine/Source/Programs/Horde/HordeStorage/Horde.Storage/appsettings.Testing.json
@@ -18,7 +18,7 @@
   },
 
   "Mongo": {
-    "ConnectionString": "mongodb://root:abc123@localhost:27017"
+    "ConnectionString": "mongodb://root:abc123@localhost:27018"
   },
 
   "Scylla": {
