Horde: Update dashboard API response objects with user information
#rnx
#skipci

diff --git a/Engine/Source/Programs/Horde/HordeDashboard/src/backend/Api.ts b/Engine/Source/Programs/Horde/HordeDashboard/src/backend/Api.ts
--- a/Engine/Source/Programs/Horde/HordeDashboard/src/backend/Api.ts
+++ b/Engine/Source/Programs/Horde/HordeDashboard/src/backend/Api.ts
@@ -942,15 +942,9 @@ export type GetJobResponse = {
 	/** Hash of the graph for this job */
 	graphHash?: string;
 
-	/** @deprecated The user that started this job */
-	startedByUser?: string;
-
 	/** The user that started this job */
 	startedByUserInfo?: GetThinUserInfoResponse;
 
-	/** @deprecated The user that aborted this job */
-	abortedByUser?: string;
-
 	/** The user that started this job */
 	abortedByUserInfo?: GetThinUserInfoResponse;
 
@@ -1044,15 +1038,9 @@ export type GetStepResponse = {
 	/** If the step has been requested to abort	*/
 	abortRequested?: boolean;
 
-	/** @deprecated Name of the user that requested the abort of this step */
-	abortByUser?: string;
-
 	/* The user that requested the abort of this step */
 	abortedByUserInfo?: GetThinUserInfoResponse;
 
-	/** @deprecated Name of the user that requested this step be run again */
-	retryByUser?: string;
-
 	/* The user that retried this step */
 	retriedByUserInfo?: GetThinUserInfoResponse;
 
@@ -2042,9 +2030,6 @@ export type GetChangeSummaryResponse = {
 	/**  The source changelist number */
 	number: number;
 
-	/**  @deprecated Name of the user that authored this change */
-	author: string;
-
 	/**  The description text */
 	description: string;
 
@@ -2682,18 +2667,10 @@ export type GetIssueResponse = {
 	/**Whether the issue is promoted */
 	promoted: boolean;
 
-	/** @deprecated Owner of the issue */
-	owner?: string;
-
-	/** @deprecated Owner id of the issue */
-	ownerId?: string;
 
 	/** Owner of the issue */
 	ownerInfo: GetThinUserInfoResponse;
 
-	/** @deprecated User that nominated the current owner */
-	nominatedBy?: string;
-
 	/** Use that nominated the current owner */
 	nominatedByInfo: GetThinUserInfoResponse;
 
@@ -2706,12 +2683,6 @@ export type GetIssueResponse = {
 	/**Time at which the issue was resolved */
 	resolvedAt?: Date | string;
 
-	/** @deprecated Name of the user that resolved the issue */
-	resolvedBy?: string;
-
-	/** @deprecated User id of the person that resolved the issue */
-	resolvedById?: string;
-
 	/** Use info for the person that resolved the issue */
 	resolvedByInfo: GetThinUserInfoResponse;
 
@@ -2726,13 +2697,9 @@ export type GetIssueResponse = {
 
 	affectedStreams: GetIssueAffectedStreamResponse[];
 
-	/** @deprecated User id's of the Most likely suspects for causing this issue */
-	primarySuspectIds: string[];
-
 	/** Use info for the person that resolved the issue */
 	primarySuspectsInfo: GetThinUserInfoResponse[];
 
-
 	/** Whether to show alerts for this issue */
 	showDesktopAlerts: boolean;
 
diff --git a/Engine/Source/Programs/Horde/HordeDashboard/src/backend/CommitCache.ts b/Engine/Source/Programs/Horde/HordeDashboard/src/backend/CommitCache.ts
--- a/Engine/Source/Programs/Horde/HordeDashboard/src/backend/CommitCache.ts
+++ b/Engine/Source/Programs/Horde/HordeDashboard/src/backend/CommitCache.ts
@@ -23,7 +23,7 @@ export class CommitCache {
         
         const commit = map.get(change);
 
-        if (!commit || !commit.author) {
+        if (!commit || !commit.authorInfo?.id) {
             return undefined;
         }
 
@@ -76,7 +76,7 @@ export class CommitCache {
             }
 
             map!.set(c, {
-                number: 0, author: "", description: "", authorInfo : { id : "0", email:"unknown", name: "unknown"}
+                number: 0, description: "", authorInfo : { id : "", email:"", name: ""}
             })
 
         })
diff --git a/Engine/Source/Programs/Horde/HordeDashboard/src/backend/JobDetails.ts b/Engine/Source/Programs/Horde/HordeDashboard/src/backend/JobDetails.ts
--- a/Engine/Source/Programs/Horde/HordeDashboard/src/backend/JobDetails.ts
+++ b/Engine/Source/Programs/Horde/HordeDashboard/src/backend/JobDetails.ts
@@ -216,7 +216,7 @@ export class JobDetails {
 
     get aborted(): boolean {
 
-        return this.jobdata?.abortedByUser ? true : false;
+        return this.jobdata?.abortedByUserInfo?.id ? true : false;
     }
 
     stepByName(name: string): StepData | undefined {
@@ -1159,7 +1159,7 @@ export const getStepSummaryMarkdown = (jobDetails: JobDetails, stepId: string):
     const text: string[] = [];
 
     if (jobDetails.jobdata) {
-        text.push(`Job created by ${jobDetails.jobdata.startedByUser ? jobDetails.jobdata.startedByUser : "scheduler"}`);
+        text.push(`Job created by ${jobDetails.jobdata.startedByUserInfo ? jobDetails.jobdata.startedByUserInfo.name : "scheduler"}`);
     }
 
     const batchText = () => {
@@ -1175,12 +1175,12 @@ export const getStepSummaryMarkdown = (jobDetails: JobDetails, stepId: string):
 
     };
 
-    if (step.retryByUser) {
+    if (step.retriedByUserInfo) {
         const retryId = jobDetails.getRetryStepId(step.id);
         if (retryId) {
-            text.push(`[Step was retried by ${step.retryByUser}](/job/${jobDetails.id!}?step=${retryId})`);
+            text.push(`[Step was retried by ${step.retriedByUserInfo.name}](/job/${jobDetails.id!}?step=${retryId})`);
         } else {
-            text.push(`Step was retried by ${step.retryByUser}`);
+            text.push(`Step was retried by ${step.retriedByUserInfo.name}`);
         }
 
     }
@@ -1191,12 +1191,12 @@ export const getStepSummaryMarkdown = (jobDetails: JobDetails, stepId: string):
     } else if (step.state === JobStepState.Aborted) {
         eta.display = eta.server = "";
         let aborted = "";
-        if (step.abortByUser) {
+        if (step.abortedByUserInfo) {
             aborted = "This step was aborted";
-            aborted += ` by ${step.abortByUser}.`;
-        } else if (jobDetails.jobdata?.abortedByUser) {
+            aborted += ` by ${step.abortedByUserInfo.name}.`;
+        } else if (jobDetails.jobdata?.abortedByUserInfo) {
             aborted = "The job was aborted";
-            aborted += ` by ${jobDetails.jobdata?.abortedByUser}.`;
+            aborted += ` by ${jobDetails.jobdata?.abortedByUserInfo.name}.`;
         } else {
             aborted = "The step was aborted";
         }
diff --git a/Engine/Source/Programs/Horde/HordeDashboard/src/backend/JobHandler.ts b/Engine/Source/Programs/Horde/HordeDashboard/src/backend/JobHandler.ts
--- a/Engine/Source/Programs/Horde/HordeDashboard/src/backend/JobHandler.ts
+++ b/Engine/Source/Programs/Horde/HordeDashboard/src/backend/JobHandler.ts
@@ -78,7 +78,7 @@ export class JobHandler {
         try {            
 
             // discover new/updated jobs
-            let filter = "id,streamId,name,change,preflightChange,templateId,templateHash,graphHash,startedByUser,abortedByUser,createTime,state,arguments,updateTime,labels,defaultLabel";
+            let filter = "id,streamId,name,change,preflightChange,templateId,templateHash,graphHash,startedByUserInfo,abortedByUserInfo,createTime,state,arguments,updateTime,labels,defaultLabel";
 
             if (this.includeBatches) {
                 filter += ",batches";
diff --git a/Engine/Source/Programs/Horde/HordeDashboard/src/base/utilities/filter.ts b/Engine/Source/Programs/Horde/HordeDashboard/src/base/utilities/filter.ts
--- a/Engine/Source/Programs/Horde/HordeDashboard/src/base/utilities/filter.ts
+++ b/Engine/Source/Programs/Horde/HordeDashboard/src/base/utilities/filter.ts
@@ -20,8 +20,8 @@ const getJobKeywords = (job: GetJobResponse): string[] => {
         keywords.push(job.preflightChange.toString());
     }
 
-    if (job.startedByUser) {
-        keywords.push(job.startedByUser);
+    if (job.startedByUserInfo) {
+        keywords.push(job.startedByUserInfo.name);
     } else {
         keywords.push("Scheduler");
     }
diff --git a/Engine/Source/Programs/Horde/HordeDashboard/src/base/utilities/issueUtils.ts b/Engine/Source/Programs/Horde/HordeDashboard/src/base/utilities/issueUtils.ts
--- a/Engine/Source/Programs/Horde/HordeDashboard/src/base/utilities/issueUtils.ts
+++ b/Engine/Source/Programs/Horde/HordeDashboard/src/base/utilities/issueUtils.ts
@@ -56,7 +56,7 @@ export function getIssueStatus(issue: GetIssueResponse, showResolveTime?: boolea
 			resolvedText += ` in CL ${issue.fixChange}`;
 		}
 
-		resolvedText += ` by ${issue.resolvedBy ?? "Horde"}`;
+		resolvedText += ` by ${issue.resolvedByInfo?.name ?? "Horde"}`;
 
 		if (!showResolveTime) {
 			return resolvedText;
@@ -71,19 +71,19 @@ export function getIssueStatus(issue: GetIssueResponse, showResolveTime?: boolea
 		return resolvedText;
 	}
 
-	if (!issue.ownerId) {
+	if (!issue.ownerInfo) {
 		text = "Currently unassigned.";
 	} else {
-		if (issue.ownerId === dashboard.userId) {
-			if (issue.nominatedBy) {
-				text = `You have been nominated to fix this issue by ${issue.nominatedBy}.`;
+		if (issue.ownerInfo.id === dashboard.userId) {
+			if (issue.nominatedByInfo) {
+				text = `You have been nominated to fix this issue by ${issue.nominatedByInfo.name}.`;
 			} else {
-				text = `Assigned to ${issue.owner}.`;
+				text = `Assigned to ${issue.ownerInfo.name}.`;
 			}
 		} else {
-			text = `Assigned to ${issue.owner}`;
-			if (issue.nominatedBy) {
-				text += ` by ${issue.nominatedBy}`;
+			text = `Assigned to ${issue.ownerInfo.name}`;
+			if (issue.nominatedByInfo) {
+				text += ` by ${issue.nominatedByInfo.name}`;
 			}
 			if (!issue.acknowledgedAt) {
 				text += ` (Unacknowledged)`;
