Horde Storage - Added a Mongo backend that can be used for local testing with persistence between runs. This is now the default when using development mode.
We also now store blob in the filesystem (subfolder from the horde.storage assemblies).

This will allow us to automatically start horde storage with horde build and have horde build able to use storage for anything it needs to store, thus no longer needing a fallback path in case horde storage is not available.

Note that the mongo backend is not a fully supported production backend, it will likely work fine for a single studio deployment but we do not expect it to be used when you need replication as scylla handles these cases much better.

diff --git a/Engine/Source/Programs/Horde/HordeStorage/Horde.Storage/Horde.Storage.Startup.cs b/Engine/Source/Programs/Horde/HordeStorage/Horde.Storage/Horde.Storage.Startup.cs
--- a/Engine/Source/Programs/Horde/HordeStorage/Horde.Storage/Horde.Storage.Startup.cs
+++ b/Engine/Source/Programs/Horde/HordeStorage/Horde.Storage/Horde.Storage.Startup.cs
@@ -287,6 +287,8 @@ namespace Horde.Storage
                     return ActivatorUtilities.CreateInstance<MemoryBlobIndex>(provider);
                 case HordeStorageSettings.BlobIndexImplementations.Scylla:
                     return ActivatorUtilities.CreateInstance<ScyllaBlobIndex>(provider);
+                case HordeStorageSettings.BlobIndexImplementations.Mongo:
+                    return ActivatorUtilities.CreateInstance<MongoBlobIndex>(provider);
                 default:
                     throw new ArgumentOutOfRangeException();
             }
@@ -323,6 +325,8 @@ namespace Horde.Storage
                     return ActivatorUtilities.CreateInstance<MemoryContentIdStore>(provider);
                 case HordeStorageSettings.ContentIdStoreImplementations.Scylla:
                     return ActivatorUtilities.CreateInstance<ScyllaContentIdStore>(provider);
+                case HordeStorageSettings.ContentIdStoreImplementations.Mongo:
+                    return ActivatorUtilities.CreateInstance<MongoContentIdStore>(provider);
                 default:
                     throw new ArgumentOutOfRangeException();
             }
@@ -388,6 +392,8 @@ namespace Horde.Storage
                     return ActivatorUtilities.CreateInstance<MemoryReferencesStore>(provider);
                 case HordeStorageSettings.ReferencesDbImplementations.Scylla:
                     return ActivatorUtilities.CreateInstance<ScyllaReferencesStore>(provider);
+                case HordeStorageSettings.ReferencesDbImplementations.Mongo:
+                    return ActivatorUtilities.CreateInstance<MongoReferencesStore>(provider);
                 default:
                     throw new ArgumentOutOfRangeException();
             }
diff --git a/Engine/Source/Programs/Horde/HordeStorage/Horde.Storage/Implementation/Blob/FileSystemStore.cs b/Engine/Source/Programs/Horde/HordeStorage/Horde.Storage/Implementation/Blob/FileSystemStore.cs
--- a/Engine/Source/Programs/Horde/HordeStorage/Horde.Storage/Implementation/Blob/FileSystemStore.cs
+++ b/Engine/Source/Programs/Horde/HordeStorage/Horde.Storage/Implementation/Blob/FileSystemStore.cs
@@ -16,7 +16,7 @@ namespace Horde.Storage.Implementation
     public class FileSystemStore : IBlobStore, IBlobCleanup 
     {
         private readonly ILogger _logger = Log.ForContext<FileSystemStore>();
-        private IOptionsMonitor<FilesystemSettings> _settings;
+        private readonly IOptionsMonitor<FilesystemSettings> _settings;
 
         internal const int DefaultBufferSize = 4096;
 
@@ -25,6 +25,11 @@ namespace Horde.Storage.Implementation
             _settings = settings;
         }
 
+        private string GetRootDir()
+        {
+            return PathUtil.ResolvePath(_settings.CurrentValue.RootDir);
+        }
+
         public static FileInfo GetFilesystemPath(string rootDir, NamespaceId ns, BlobIdentifier blob)
         {
             const int CountOfCharactersPerDirectory = 2;
@@ -38,12 +43,12 @@ namespace Horde.Storage.Implementation
 
         public FileInfo GetFilesystemPath(NamespaceId ns, BlobIdentifier objectName)
         {
-            return GetFilesystemPath(_settings.CurrentValue.RootDir, ns, objectName);
+            return GetFilesystemPath(GetRootDir(), ns, objectName);
         }
 
         public DirectoryInfo GetFilesystemPath(NamespaceId ns)
         {
-            return new DirectoryInfo(Path.Combine(_settings.CurrentValue.RootDir, ns.ToString()));
+            return new DirectoryInfo(Path.Combine(GetRootDir(), ns.ToString()));
         }
 
         public async Task<BlobIdentifier> PutObject(NamespaceId ns, ReadOnlyMemory<byte> content, BlobIdentifier blobIdentifier)
@@ -261,7 +266,7 @@ namespace Horde.Storage.Implementation
         /// <returns>Enumerable of least recently accessed objects as FileInfos</returns>
         public IEnumerable<FileInfo> GetLeastRecentlyAccessedObjects(NamespaceId? ns = null, int maxResults = 10000)
         {
-            string path = ns != null ? Path.Combine(_settings.CurrentValue.RootDir, ns.ToString()!) : _settings.CurrentValue.RootDir;
+            string path = ns != null ? Path.Combine(GetRootDir(), ns.ToString()!) : GetRootDir();
             DirectoryInfo di = new DirectoryInfo(path);
             if (!di.Exists)
             {
@@ -278,7 +283,7 @@ namespace Horde.Storage.Implementation
         /// <returns>Total size of blobs in bytes</returns>
         public async Task<long> CalculateDiskSpaceUsed(NamespaceId? ns = null)
         {
-            string path = ns != null ? Path.Combine(_settings.CurrentValue.RootDir, ns.ToString()!) : _settings.CurrentValue.RootDir;
+            string path = ns != null ? Path.Combine(GetRootDir(), ns.ToString()!) : GetRootDir();
             DirectoryInfo di = new DirectoryInfo(path);
             if (!di.Exists)
             {
@@ -301,7 +306,7 @@ namespace Horde.Storage.Implementation
         
         public  IAsyncEnumerable<NamespaceId> ListNamespaces()
         {
-            DirectoryInfo di = new DirectoryInfo(_settings.CurrentValue.RootDir);
+            DirectoryInfo di = new DirectoryInfo(GetRootDir());
             return di.GetDirectories().Select(x => new NamespaceId(x.Name)).ToAsyncEnumerable();
         }
     }
diff --git a/Engine/Source/Programs/Horde/HordeStorage/Horde.Storage/Implementation/Blob/MongoBlobIndex.cs b/Engine/Source/Programs/Horde/HordeStorage/Horde.Storage/Implementation/Blob/MongoBlobIndex.cs
new file mode 100644
--- /dev/null
+++ b/Engine/Source/Programs/Horde/HordeStorage/Horde.Storage/Implementation/Blob/MongoBlobIndex.cs
@@ -0,0 +1,144 @@
+// Copyright Epic Games, Inc. All Rights Reserved.
+
+using System.Collections.Generic;
+using System.Linq;
+using System.Threading.Tasks;
+using Jupiter;
+using Jupiter.Implementation;
+using Microsoft.Extensions.Options;
+using MongoDB.Bson.Serialization.Attributes;
+using MongoDB.Bson.Serialization.Options;
+using MongoDB.Driver;
+
+namespace Horde.Storage.Implementation.Blob;
+
+public class MongoBlobIndex : MongoStore, IBlobIndex
+{
+    private readonly IOptionsMonitor<JupiterSettings> _jupiterSettings;
+
+    public MongoBlobIndex(IOptionsMonitor<JupiterSettings> jupiterSettings, IOptionsMonitor<MongoSettings> settings) : base(settings)
+    {
+        _jupiterSettings = jupiterSettings;
+
+        CreateCollectionIfNotExists<MongoBlobIndexModelV0>().Wait();
+
+        IndexKeysDefinitionBuilder<MongoBlobIndexModelV0> indexKeysDefinitionBuilder = Builders<MongoBlobIndexModelV0>.IndexKeys;
+        CreateIndexModel<MongoBlobIndexModelV0> indexModel = new CreateIndexModel<MongoBlobIndexModelV0>(
+            indexKeysDefinitionBuilder.Combine(
+                indexKeysDefinitionBuilder.Ascending(m => m.Ns), 
+                indexKeysDefinitionBuilder.Ascending(m => m.BlobId)
+            )
+            , new CreateIndexOptions()
+            {
+                Name = "CompoundIndex"
+            });
+
+        
+        AddIndexFor<MongoBlobIndexModelV0>().CreateMany(new[] { 
+            indexModel, 
+        });
+    }
+
+    public async Task AddBlobToIndex(NamespaceId ns, BlobIdentifier id)
+    {
+        IMongoCollection<MongoBlobIndexModelV0> collection = GetCollection<MongoBlobIndexModelV0>();
+        MongoBlobIndexModelV0 model = new MongoBlobIndexModelV0(ns, id);
+        model.Regions.Add(_jupiterSettings.CurrentValue.CurrentSite);
+            
+        FilterDefinition<MongoBlobIndexModelV0> filter = Builders<MongoBlobIndexModelV0>.Filter.Where(m => m.Ns == ns.ToString() && m.BlobId == id.ToString());
+        await collection.FindOneAndReplaceAsync(filter, model, new FindOneAndReplaceOptions<MongoBlobIndexModelV0, MongoBlobIndexModelV0>
+        {
+            IsUpsert = true
+        });
+    }
+
+    public async Task<IBlobIndex.BlobInfo?> GetBlobInfo(NamespaceId ns, BlobIdentifier id)
+    {
+        IMongoCollection<MongoBlobIndexModelV0> collection = GetCollection<MongoBlobIndexModelV0>();
+        IAsyncCursor<MongoBlobIndexModelV0>? cursor = await collection.FindAsync(m => m.Ns == ns.ToString() && m.BlobId == id.ToString());
+        MongoBlobIndexModelV0? model = await cursor.FirstOrDefaultAsync();
+        if (model == null)
+            return null;
+
+        return model.ToBlobInfo();
+    }
+
+    public async Task<bool> RemoveBlobFromIndex(NamespaceId ns, BlobIdentifier id)
+    {
+        IMongoCollection<MongoBlobIndexModelV0> collection = GetCollection<MongoBlobIndexModelV0>();
+        DeleteResult result = await collection.DeleteOneAsync(m => m.Ns == ns.ToString() && m.BlobId == id.ToString());
+
+        return result.IsAcknowledged;
+    }
+
+    public async Task<bool> BlobExistsInRegion(NamespaceId ns, BlobIdentifier blobIdentifier)
+    {
+        IBlobIndex.BlobInfo? blobInfo = await GetBlobInfo(ns, blobIdentifier);
+        return blobInfo?.Regions.Contains(_jupiterSettings.CurrentValue.CurrentSite) ?? false;
+    }
+
+    public async Task AddRefToBlobs(NamespaceId ns, BucketId bucket, IoHashKey key, BlobIdentifier[] blobs)
+    {
+        IMongoCollection<MongoBlobIndexModelV0> collection = GetCollection<MongoBlobIndexModelV0>();
+
+        string nsAsString = ns.ToString();
+        Task[] refUpdateTasks = new Task[blobs.Length];
+        for (int i = 0; i < blobs.Length; i++)
+        {
+            BlobIdentifier id = blobs[i];
+            refUpdateTasks[i] = Task.Run(async () =>
+            {
+                UpdateDefinition<MongoBlobIndexModelV0> update = Builders<MongoBlobIndexModelV0>.Update.AddToSet(m => m.References, new Dictionary<string, string> { {"bucket", bucket.ToString()}, {"key", key.ToString()}});
+                FilterDefinition<MongoBlobIndexModelV0> filter = Builders<MongoBlobIndexModelV0>.Filter.Where(m => m.Ns == nsAsString && m.BlobId == id.ToString());
+
+                await collection.FindOneAndUpdateAsync(filter, update);
+            });
+        }
+
+        await Task.WhenAll(refUpdateTasks);
+    }
+
+}
+
+[BsonDiscriminator("blob-index.v0")]
+[BsonIgnoreExtraElements]
+[MongoCollectionName("BlobIndex")]
+class MongoBlobIndexModelV0
+{
+    [BsonConstructor]
+    public MongoBlobIndexModelV0(string ns, string blobId, List<string> regions, List<Dictionary<string, string>> references)
+    {
+        Ns = ns;
+        BlobId = blobId;
+        Regions = regions;
+        References = references;
+    }
+
+    public MongoBlobIndexModelV0(NamespaceId ns, BlobIdentifier blobId)
+    {
+        Ns = ns.ToString();
+        BlobId = blobId.ToString();
+    }
+
+    [BsonRequired]
+    public string Ns { get; set; }
+
+    [BsonRequired]
+    public string BlobId { get;set; }
+
+    public List<string> Regions { get; set; } = new List<string>();
+
+    [BsonDictionaryOptions(DictionaryRepresentation.Document)]
+    public List<Dictionary<string, string>> References { get; set; } = new List<Dictionary<string, string>>();
+
+    public IBlobIndex.BlobInfo ToBlobInfo()
+    {
+        return new IBlobIndex.BlobInfo()
+        {
+            Namespace = new NamespaceId(Ns),
+            Regions = Regions.ToHashSet(),
+            References = References.Select(dictionary =>
+                (new BucketId(dictionary["bucket"]), new IoHashKey(dictionary["key"]))).ToList(),
+        };
+    }
+}
diff --git a/Engine/Source/Programs/Horde/HordeStorage/Horde.Storage/Implementation/Blob/PathUtil.cs b/Engine/Source/Programs/Horde/HordeStorage/Horde.Storage/Implementation/Blob/PathUtil.cs
new file mode 100644
--- /dev/null
+++ b/Engine/Source/Programs/Horde/HordeStorage/Horde.Storage/Implementation/Blob/PathUtil.cs
@@ -0,0 +1,16 @@
+// Copyright Epic Games, Inc. All Rights Reserved.
+
+using System;
+using System.Reflection;
+using System.IO;
+
+namespace Horde.Storage.Implementation;
+
+public static class PathUtil
+{
+    public static string ResolvePath(string path)
+    {
+        return Environment.ExpandEnvironmentVariables(path)
+            .Replace("$(ExecutableLocation)", Path.GetDirectoryName(Assembly.GetEntryAssembly()!.Location));
+    }
+}
diff --git a/Engine/Source/Programs/Horde/HordeStorage/Horde.Storage/Implementation/ContentId/MongoContentIdStore.cs b/Engine/Source/Programs/Horde/HordeStorage/Horde.Storage/Implementation/ContentId/MongoContentIdStore.cs
new file mode 100644
--- /dev/null
+++ b/Engine/Source/Programs/Horde/HordeStorage/Horde.Storage/Implementation/ContentId/MongoContentIdStore.cs
@@ -0,0 +1,120 @@
+// Copyright Epic Games, Inc. All Rights Reserved.
+
+using System.Collections.Generic;
+using System.Linq;
+using System.Threading.Tasks;
+using Jupiter.Implementation;
+using Microsoft.Extensions.Options;
+using MongoDB.Bson.Serialization.Attributes;
+using MongoDB.Bson.Serialization.Options;
+using MongoDB.Driver;
+
+namespace Horde.Storage.Implementation
+{
+    public class MongoContentIdStore : MongoStore, IContentIdStore
+    {
+        private readonly IBlobService _blobStore;
+
+        public MongoContentIdStore(IBlobService blobStore, IOptionsMonitor<MongoSettings> settings) : base(settings)
+        {
+            _blobStore = blobStore;
+
+            CreateCollectionIfNotExists<MongoContentIdModelV0>().Wait();
+
+            IndexKeysDefinitionBuilder<MongoContentIdModelV0> indexKeysDefinitionBuilder = Builders<MongoContentIdModelV0>.IndexKeys;
+            CreateIndexModel<MongoContentIdModelV0> indexModel = new CreateIndexModel<MongoContentIdModelV0>(
+                indexKeysDefinitionBuilder.Combine(
+                    indexKeysDefinitionBuilder.Ascending(m => m.Ns), 
+                    indexKeysDefinitionBuilder.Ascending(m => m.ContentId)
+                    )
+                , new CreateIndexOptions()
+                {
+                    Name = "CompoundIndex"
+                });
+
+            AddIndexFor<MongoContentIdModelV0>().CreateOne(indexModel);
+        }
+
+        public async Task<BlobIdentifier[]?> Resolve(NamespaceId ns, ContentId contentId)
+        {
+            IMongoCollection<MongoContentIdModelV0> collection = GetCollection<MongoContentIdModelV0>();
+
+            IAsyncCursor<MongoContentIdModelV0> cursor =
+                await collection.FindAsync(model => model.Ns.Equals(ns) && model.ContentId.Equals(contentId));
+
+            MongoContentIdModelV0 model = await cursor.FirstOrDefaultAsync();
+            if (model != null)
+            {
+                foreach (int weight in model.ContentWeightToBlobsMap.Keys.OrderBy(contentWeight => contentWeight))
+                {
+                    BlobIdentifier[] blobs = model.ContentWeightToBlobsMap[weight].Select(s => new BlobIdentifier(s)).ToArray();
+
+                    {
+                        BlobIdentifier[] missingBlobs = await _blobStore.FilterOutKnownBlobs(ns, blobs);
+                        if (missingBlobs.Length == 0)
+                            return blobs;
+                    }
+                    // blobs are missing continue testing with the next content id in the weighted list as that might exist
+
+                }
+            }
+
+            BlobIdentifier contentIdAsBlobIdentifier = contentId.AsBlobIdentifier();
+            // no content id where all blobs are present, check if its present in the blob store as a uncompressed version of the blob
+            if (await _blobStore.Exists(ns, contentIdAsBlobIdentifier))
+            {
+                return new[] { contentIdAsBlobIdentifier };
+            }
+
+            return null;
+        }
+
+        public async Task Put(NamespaceId ns, ContentId contentId, BlobIdentifier blobIdentifier, int contentWeight)
+        {
+            IMongoCollection<MongoContentIdModelV0> collection = GetCollection<MongoContentIdModelV0>();
+
+            UpdateDefinition<MongoContentIdModelV0> update = Builders<MongoContentIdModelV0>.Update.AddToSet(m => m.ContentWeightToBlobsMap,
+                new KeyValuePair<int, string[]>(contentWeight, new[] { blobIdentifier.ToString() }));
+            FilterDefinition<MongoContentIdModelV0> filter = Builders<MongoContentIdModelV0>.Filter.Where(m => m.Ns.Equals(ns) && m.ContentId.Equals(contentId));
+            await collection.FindOneAndUpdateAsync(filter, update, new FindOneAndUpdateOptions<MongoContentIdModelV0>()
+            {
+                IsUpsert = true
+            });
+        }
+    }
+
+    [BsonDiscriminator("content-id.v0")]
+    [BsonIgnoreExtraElements]
+    [MongoCollectionName("ContentId")]
+    public class MongoContentIdModelV0
+    {
+        [BsonConstructor]
+        public MongoContentIdModelV0(string ns, string contentId, Dictionary<int, string[]> contentWeightToBlobsMap)
+        {
+            Ns = ns;
+            ContentId = contentId;
+            ContentWeightToBlobsMap = contentWeightToBlobsMap;
+        }
+
+        public MongoContentIdModelV0(NamespaceId ns, ContentId contentId, int contentWeight, BlobIdentifier[] blobs)
+        {
+            Ns = ns.ToString();
+            ContentId = contentId.ToString();
+
+            ContentWeightToBlobsMap = new Dictionary<int, string[]>
+            {
+                [contentWeight] = blobs.Select(identifier => identifier.ToString()).ToArray()
+            };
+        }
+
+        [BsonRequired]
+        public string Ns { get; set; }
+
+        [BsonRequired]
+        public string ContentId { get; set; }
+
+        [BsonDictionaryOptions(DictionaryRepresentation.ArrayOfDocuments)]
+        public Dictionary<int, string[]> ContentWeightToBlobsMap { get; set; }
+
+    }
+}
diff --git a/Engine/Source/Programs/Horde/HordeStorage/Horde.Storage/Implementation/MongoStore.cs b/Engine/Source/Programs/Horde/HordeStorage/Horde.Storage/Implementation/MongoStore.cs
new file mode 100644
--- /dev/null
+++ b/Engine/Source/Programs/Horde/HordeStorage/Horde.Storage/Implementation/MongoStore.cs
@@ -0,0 +1,86 @@
+// Copyright Epic Games, Inc. All Rights Reserved.
+
+using System;
+using System.Collections.Generic;
+using System.Security.Authentication;
+using System.Threading.Tasks;
+using Microsoft.Extensions.Options;
+using MongoDB.Bson.Serialization;
+using MongoDB.Driver;
+
+namespace Horde.Storage.Implementation;
+
+public class MongoStore
+{
+    protected readonly MongoClient _client;
+
+    public MongoStore(IOptionsMonitor<MongoSettings> settings)
+    {
+        MongoClientSettings mongoClientSettings = MongoClientSettings.FromUrl(
+            new MongoUrl(settings.CurrentValue.ConnectionString)
+        );
+        if (settings.CurrentValue.RequireTls12)
+        {
+            mongoClientSettings.SslSettings = new SslSettings {EnabledSslProtocols = SslProtocols.Tls12};
+        }
+
+        _client = new MongoClient(mongoClientSettings);
+    }
+
+    protected async Task CreateCollectionIfNotExists<T>()
+    {
+        string collectionName = GetCollectionName<T>();
+
+        try
+        {
+            await _client.GetDatabase(GetDatabaseName()).CreateCollectionAsync(collectionName);
+        }
+        catch (MongoCommandException e)
+        {
+            if (e.CodeName != "NamespaceExists")
+                throw ;
+        }
+    }
+
+    private string GetCollectionName<T>()
+    {
+        object[] attr = typeof(T).GetCustomAttributes(typeof(MongoCollectionNameAttribute), true);
+        foreach (MongoCollectionNameAttribute o in attr)
+        {
+            return o.CollectionName;
+        }
+
+        throw new ArgumentException($"No MongoCollectionNameAttribute found on type {nameof(T)}");
+    }
+
+    protected IMongoIndexManager<T> AddIndexFor<T>()
+    {
+        return GetCollection<T>().Indexes;
+    }
+
+    protected IMongoCollection<T> GetCollection<T>()
+    {
+        string collectionName = GetCollectionName<T>();
+        string dbName = GetDatabaseName();
+        return _client.GetDatabase(dbName).GetCollection<T>(collectionName);
+    }
+
+    private string GetDatabaseName()
+    {
+        // TODO: Database per namespace seems a bit much?
+        //string cleanedNs = ns.ToString().Replace(".", "_");
+        //string dbName = $"HordeStorage_{cleanedNs}";
+        string dbName = $"HordeStorage";
+        return dbName;
+    }
+}
+
+public class MongoCollectionNameAttribute : Attribute
+{
+    public string CollectionName { get; }
+
+    public MongoCollectionNameAttribute(string collectionName)
+    {
+        CollectionName = collectionName;
+    }
+}
diff --git a/Engine/Source/Programs/Horde/HordeStorage/Horde.Storage/Implementation/Objects/MongoReferencesStore.cs b/Engine/Source/Programs/Horde/HordeStorage/Horde.Storage/Implementation/Objects/MongoReferencesStore.cs
new file mode 100644
--- /dev/null
+++ b/Engine/Source/Programs/Horde/HordeStorage/Horde.Storage/Implementation/Objects/MongoReferencesStore.cs
@@ -0,0 +1,261 @@
+// Copyright Epic Games, Inc. All Rights Reserved.
+
+using System;
+using System.Collections.Generic;
+using System.Threading.Tasks;
+using Dasync.Collections;
+using Jupiter.Implementation;
+using Microsoft.Extensions.Options;
+using MongoDB.Bson.Serialization.Attributes;
+using MongoDB.Driver;
+using Serilog;
+
+namespace Horde.Storage.Implementation
+{
+    public class MongoReferencesStore : MongoStore, IReferencesStore
+    {
+        private readonly ILogger _logger = Log.ForContext<MongoReferencesStore>();
+
+        public MongoReferencesStore(IOptionsMonitor<MongoSettings> settings) : base(settings)
+        {
+            CreateCollectionIfNotExists<MongoReferencesModelV0>().Wait();
+            CreateCollectionIfNotExists<MongoNamespacesModelV0>().Wait();
+
+            IndexKeysDefinitionBuilder<MongoReferencesModelV0> indexKeysDefinitionBuilder = Builders<MongoReferencesModelV0>.IndexKeys;
+            CreateIndexModel<MongoReferencesModelV0> indexModelClusteredKey = new CreateIndexModel<MongoReferencesModelV0>(
+                indexKeysDefinitionBuilder.Combine(
+                    indexKeysDefinitionBuilder.Ascending(m => m.Ns),
+                    indexKeysDefinitionBuilder.Ascending(m => m.Bucket),
+                    indexKeysDefinitionBuilder.Ascending(m => m.Key)
+                ), new CreateIndexOptions {Name = "CompoundIndex"}
+            );
+
+            CreateIndexModel<MongoReferencesModelV0> indexModelNamespace = new CreateIndexModel<MongoReferencesModelV0>(
+                indexKeysDefinitionBuilder.Ascending(m => m.Ns),
+                new CreateIndexOptions {Name = "NamespaceIndex"}
+            );
+
+            AddIndexFor<MongoReferencesModelV0>().CreateMany(new[] { 
+                indexModelClusteredKey, 
+                indexModelNamespace
+            });
+
+        }
+
+        public async Task<ObjectRecord> Get(NamespaceId ns, BucketId bucket, IoHashKey key, IReferencesStore.FieldFlags flags)
+        {
+            IMongoCollection<MongoReferencesModelV0> collection = GetCollection<MongoReferencesModelV0>();
+            IAsyncCursor<MongoReferencesModelV0>? cursor = await collection.FindAsync(m => m.Ns == ns.ToString() && m.Bucket == bucket.ToString() && m.Key == key.ToString());
+            MongoReferencesModelV0? model = await cursor.FirstOrDefaultAsync();
+            if (model == null)
+                throw new ObjectNotFoundException(ns, bucket, key);
+
+            return model.ToObjectRecord();
+        }
+
+        public async Task Put(NamespaceId ns, BucketId bucket, IoHashKey key, BlobIdentifier blobHash, byte[] blob, bool isFinalized)
+        {
+            IMongoCollection<MongoReferencesModelV0> collection = GetCollection<MongoReferencesModelV0>();
+
+            //if (blob.LongLength > _settings.CurrentValue.InlineBlobMaxSize)
+            if (blob.LongLength > 64_000) // TODO: add a setting to control this
+            {
+                // do not inline large blobs
+                blob = Array.Empty<byte>();
+            }
+
+            Task addNamespaceTask = AddNamespaceIfNotExist(ns);
+            MongoReferencesModelV0 model = new MongoReferencesModelV0(ns, bucket, key, blobHash, blob, isFinalized, DateTime.Now);
+            
+            FilterDefinition<MongoReferencesModelV0> filter = Builders<MongoReferencesModelV0>.Filter.Where(m => m.Ns == ns.ToString() && m.Bucket == bucket.ToString() && m.Key == key.ToString());
+            FindOneAndReplaceOptions<MongoReferencesModelV0, MongoReferencesModelV0> options = new FindOneAndReplaceOptions<MongoReferencesModelV0, MongoReferencesModelV0>
+            {
+                IsUpsert = true
+            };
+            await collection.FindOneAndReplaceAsync(filter, model, options);
+
+            await addNamespaceTask;
+        }
+
+        public async Task Finalize(NamespaceId ns, BucketId bucket, IoHashKey key)
+        {
+            IMongoCollection<MongoReferencesModelV0> collection = GetCollection<MongoReferencesModelV0>();
+
+            UpdateResult _ = await collection.UpdateOneAsync(
+                model => model.Ns == ns.ToString() && model.Bucket == bucket.ToString() && model.Key == key.ToString(),
+                Builders<MongoReferencesModelV0>.Update.Set(model => model.IsFinalized, true)
+            );
+        }
+
+        public async Task UpdateLastAccessTime(NamespaceId ns, BucketId bucket, IoHashKey key, DateTime newLastAccessTime)
+        {
+            IMongoCollection<MongoReferencesModelV0> collection = GetCollection<MongoReferencesModelV0>();
+
+            UpdateResult _ = await collection.UpdateOneAsync(
+                model => model.Ns == ns.ToString() && model.Bucket == bucket.ToString() && model.Key == key.ToString(),
+                Builders<MongoReferencesModelV0>.Update.Set(model => model.LastAccessTime, newLastAccessTime)
+            );
+        }
+
+        public async IAsyncEnumerable<(BucketId, IoHashKey, DateTime)> GetRecords(NamespaceId ns)
+        {
+            IMongoCollection<MongoReferencesModelV0> collection = GetCollection<MongoReferencesModelV0>();
+            IAsyncCursor<MongoReferencesModelV0>? cursor = await collection.FindAsync(model => model.Ns == ns.ToString());
+
+            while (await cursor.MoveNextAsync())
+            {
+                await foreach (MongoReferencesModelV0 model in cursor.Current)
+                {
+                    yield return (new BucketId(model.Bucket), new IoHashKey(model.Key), model.LastAccessTime);
+                }
+            }
+        }
+
+        public async Task AddNamespaceIfNotExist(NamespaceId ns)
+        {
+            FilterDefinition<MongoNamespacesModelV0> filter = Builders<MongoNamespacesModelV0>.Filter.Where(m => m.Ns == ns.ToString());
+            FindOneAndReplaceOptions<MongoNamespacesModelV0, MongoNamespacesModelV0> options = new FindOneAndReplaceOptions<MongoNamespacesModelV0, MongoNamespacesModelV0>
+            {
+                IsUpsert = true
+            };
+
+            IMongoCollection<MongoNamespacesModelV0> collection = GetCollection<MongoNamespacesModelV0>();
+            await collection.FindOneAndReplaceAsync(filter, new MongoNamespacesModelV0(ns), options);
+        }
+
+        public async IAsyncEnumerable<NamespaceId> GetNamespaces()
+        {
+            IMongoCollection<MongoNamespacesModelV0> collection = GetCollection<MongoNamespacesModelV0>();
+
+            IAsyncCursor<MongoNamespacesModelV0> cursor = await collection.FindAsync(FilterDefinition<MongoNamespacesModelV0>.Empty);
+            while (await cursor.MoveNextAsync())
+            {
+                foreach (MongoNamespacesModelV0? document in cursor.Current)
+                {
+                    yield return new NamespaceId(document.Ns);
+                }
+            }
+        }
+
+        public async Task<bool> Delete(NamespaceId ns, BucketId bucket, IoHashKey key)
+        {
+            IMongoCollection<MongoReferencesModelV0> collection = GetCollection<MongoReferencesModelV0>();
+
+            DeleteResult result = await collection.DeleteOneAsync(model =>
+                model.Ns == ns.ToString() && model.Bucket == bucket.ToString() && model.Key == key.ToString());
+
+            return result.IsAcknowledged;
+        }
+
+        public async Task<long> DropNamespace(NamespaceId ns)
+        {
+            IMongoCollection<MongoReferencesModelV0> collection = GetCollection<MongoReferencesModelV0>();
+
+            DeleteResult result = await collection.DeleteManyAsync(model =>
+                model.Ns == ns.ToString());
+
+            if (result.IsAcknowledged)
+            {
+                return result.DeletedCount;
+            }
+
+            // failed to delete
+            return 0L;
+
+        }
+
+        public async Task<long> DeleteBucket(NamespaceId ns, BucketId bucket)
+        {
+            IMongoCollection<MongoReferencesModelV0> collection = GetCollection<MongoReferencesModelV0>();
+
+            DeleteResult result = await collection.DeleteManyAsync(model =>
+                model.Ns == ns.ToString() && model.Bucket == bucket.ToString());
+
+            if (result.IsAcknowledged)
+            {
+                return result.DeletedCount;
+            }
+
+            // failed to delete
+            return 0L;
+        }
+    }
+
+    
+    // we do versioning by writing a discriminator into object
+    [BsonDiscriminator("ref.v0")]
+    [BsonIgnoreExtraElements]
+    [MongoCollectionName("References")]
+    public class MongoReferencesModelV0
+    {
+        [BsonConstructor]
+        public MongoReferencesModelV0(string ns, string bucket, string key, string blobIdentifier, byte[] inlineBlob, bool isFinalized, DateTime lastAccessTime)
+        {
+            Ns = ns;
+            Bucket = bucket;
+            Key = key;
+            BlobIdentifier = blobIdentifier;
+            InlineBlob = inlineBlob;
+            IsFinalized = isFinalized;
+            LastAccessTime = lastAccessTime;
+        }
+
+        public MongoReferencesModelV0(NamespaceId ns, BucketId bucket, IoHashKey key, BlobIdentifier blobIdentifier, byte[] blob, bool isFinalized, DateTime lastAccessTime)
+        {
+            Ns = ns.ToString();
+            Bucket = bucket.ToString();
+            Key = key.ToString();
+            InlineBlob = blob;
+            BlobIdentifier = blobIdentifier.ToString();
+            IsFinalized = isFinalized;
+            LastAccessTime = lastAccessTime;
+        }
+
+        [BsonRequired]
+        public string Ns { get; set; }
+
+        [BsonRequired]
+        public string Bucket { get; set; }
+
+        [BsonRequired]
+        public string Key { get; set; }
+
+        public string BlobIdentifier { get; set; }
+
+        public bool IsFinalized { get; set; }
+
+        [BsonRequired]
+        public DateTime LastAccessTime { get; set; }
+
+        public byte[] InlineBlob { get; set; }
+
+        public ObjectRecord ToObjectRecord()
+        {
+            return new ObjectRecord(new NamespaceId(Ns), new BucketId(Bucket), new IoHashKey(Key), 
+                LastAccessTime,
+                InlineBlob, new BlobIdentifier(BlobIdentifier), IsFinalized);
+        }
+    }
+
+    [MongoCollectionName("Namespaces")]
+    [BsonDiscriminator("namespace.v0")]
+    [BsonIgnoreExtraElements]
+    public class MongoNamespacesModelV0
+    {
+        [BsonRequired]
+        public string Ns { get; set; }
+
+
+        [BsonConstructor]
+        public MongoNamespacesModelV0(string ns)
+        {
+            Ns = ns;
+        }
+
+        public MongoNamespacesModelV0(NamespaceId ns)
+        {
+            Ns = ns.ToString();
+        }
+
+    }
+}
diff --git a/Engine/Source/Programs/Horde/HordeStorage/Horde.Storage/Implementation/Replication/RefsReplicator.cs b/Engine/Source/Programs/Horde/HordeStorage/Horde.Storage/Implementation/Replication/RefsReplicator.cs
--- a/Engine/Source/Programs/Horde/HordeStorage/Horde.Storage/Implementation/Replication/RefsReplicator.cs
+++ b/Engine/Source/Programs/Horde/HordeStorage/Horde.Storage/Implementation/Replication/RefsReplicator.cs
@@ -53,7 +53,7 @@ namespace Horde.Storage.Implementation
             _httpClient.BaseAddress = new Uri(replicatorSettings.ConnectionString);
 
             string stateFileName = $"{_name}.json";
-            DirectoryInfo stateRoot = new DirectoryInfo(replicationSettings.CurrentValue.StateRoot);
+            DirectoryInfo stateRoot = new DirectoryInfo(PathUtil.ResolvePath(replicationSettings.CurrentValue.StateRoot));
             _stateFile = new FileInfo(Path.Combine(stateRoot.FullName, stateFileName));
 
             if (_stateFile.Exists)
diff --git a/Engine/Source/Programs/Horde/HordeStorage/Horde.Storage/Implementation/Replication/Replicator.cs b/Engine/Source/Programs/Horde/HordeStorage/Horde.Storage/Implementation/Replication/Replicator.cs
--- a/Engine/Source/Programs/Horde/HordeStorage/Horde.Storage/Implementation/Replication/Replicator.cs
+++ b/Engine/Source/Programs/Horde/HordeStorage/Horde.Storage/Implementation/Replication/Replicator.cs
@@ -123,7 +123,7 @@ namespace Horde.Storage.Implementation
             _replicatorSettings = replicatorSettings;
 
             string stateFileName = $"{_name}.json";
-            DirectoryInfo stateRoot = new DirectoryInfo(replicationSettings.CurrentValue.StateRoot);
+            DirectoryInfo stateRoot = new DirectoryInfo(PathUtil.ResolvePath(replicationSettings.CurrentValue.StateRoot));
             _stateFile = new FileInfo(Path.Combine(stateRoot.FullName, stateFileName));
 
             if (_stateFile.Exists)
diff --git a/Engine/Source/Programs/Horde/HordeStorage/Horde.Storage/Properties/launchSettings.json b/Engine/Source/Programs/Horde/HordeStorage/Horde.Storage/Properties/launchSettings.json
--- a/Engine/Source/Programs/Horde/HordeStorage/Horde.Storage/Properties/launchSettings.json
+++ b/Engine/Source/Programs/Horde/HordeStorage/Horde.Storage/Properties/launchSettings.json
@@ -1,43 +1,11 @@
 {
-  "iisSettings": {
-    "windowsAuthentication": false,
-    "anonymousAuthentication": true,
-    "iisExpress": {
-      "applicationUrl": "http://localhost:57001",
-      "sslPort": 44351
-    }
-  },
-  "$schema": "http://json.schemastore.org/launchsettings.json",
   "profiles": {
-    "IIS Express": {
-      "commandName": "IISExpress",
-      "launchBrowser": true,
-      "launchUrl": "api/v1/c/ddc/test.ddc/default/foo.json",
-      "environmentVariables": {
-        "ASPNETCORE_ENVIRONMENT": "Development",
-        "DD_TRACE_ENABLED": "false"
-      }
-    },
-    "Local Deployment": {
-      "commandName": "IISExpress",
-      "launchBrowser": true,
-      "launchUrl": "api/v1/c/ddc/ue4.ddc/terrain/foo.json",
-      "environmentVariables": {
-        "CORECLR_ENABLE_PROFILING": "0",
-        "Callisto__ConnectionString": "http://localhost:57002",
-        "Horde_Storage__StorageImplementation": "Memory",
-        "Horde_Storage__TransactionLogWriterImplementation": "Callisto",
-        "ASPNETCORE_ENVIRONMENT": "Development"
-      }
-    },
-    "Kestrel": {
+    "Project": {
       "commandName": "Project",
-      "launchUrl": "/",
       "environmentVariables": {
         "ASPNETCORE_ENVIRONMENT": "Development",
-        "DD_TRACE_ENABLED": "false"
-      },
-      "applicationUrl": ""
+        "ASPNETCORE_PORTS": "\"\""
+      }
     }
   }
-}
+}
\ No newline at end of file
diff --git a/Engine/Source/Programs/Horde/HordeStorage/Horde.Storage/Settings.cs b/Engine/Source/Programs/Horde/HordeStorage/Horde.Storage/Settings.cs
--- a/Engine/Source/Programs/Horde/HordeStorage/Horde.Storage/Settings.cs
+++ b/Engine/Source/Programs/Horde/HordeStorage/Horde.Storage/Settings.cs
@@ -29,7 +29,8 @@ namespace Horde.Storage
         public enum ReplicationLogWriterImplementations
         {
             Memory,
-            Scylla
+            Scylla,
+            Mongo
         }
 
         public enum TreeStoreImplementations
@@ -55,19 +56,22 @@ namespace Horde.Storage
         public enum ReferencesDbImplementations
         {
             Memory,
-            Scylla
+            Scylla,
+            Mongo
         }
 
         public enum ContentIdStoreImplementations
         {
             Memory,
-            Scylla
+            Scylla,
+            Mongo
         }
 
         public enum BlobIndexImplementations
         {
             Memory,
-            Scylla
+            Scylla,
+            Mongo
         }
 
         public enum LeaderElectionImplementations
@@ -137,6 +141,7 @@ namespace Horde.Storage
         [Required] public string ConnectionString { get; set; } = "";
 
         public bool RequireTls12 { get; set; } = true;
+        public bool CreateDatabaseIfMissing { get; set; } = true;
     }
 
 
diff --git a/Engine/Source/Programs/Horde/HordeStorage/Horde.Storage/Tests/Functional/Blob/BlobIndexTests.cs b/Engine/Source/Programs/Horde/HordeStorage/Horde.Storage/Tests/Functional/Blob/BlobIndexTests.cs
--- a/Engine/Source/Programs/Horde/HordeStorage/Horde.Storage/Tests/Functional/Blob/BlobIndexTests.cs
+++ b/Engine/Source/Programs/Horde/HordeStorage/Horde.Storage/Tests/Functional/Blob/BlobIndexTests.cs
@@ -191,4 +191,23 @@ namespace Horde.Storage.FunctionalTests.Storage
             await replicatedKeyspace.ExecuteAsync(new SimpleStatement("DROP TABLE IF EXISTS blob_index"));
         }
     }
+
+    [TestClass]
+    public class MongoBlobIndexTests : BlobIndexTests
+    {
+        protected override IEnumerable<KeyValuePair<string, string>> GetSettings()
+        {
+            return new[] { new KeyValuePair<string, string>("Horde_Storage:BlobIndexImplementation", HordeStorageSettings.BlobIndexImplementations.Mongo.ToString()) };
+        }
+
+        protected override Task Seed(IServiceProvider serverServices)
+        {
+            return Task.CompletedTask;
+        }
+
+        protected override async Task Teardown(IServiceProvider provider)
+        {
+            await Task.CompletedTask;
+        }
+    }
 }
diff --git a/Engine/Source/Programs/Horde/HordeStorage/Horde.Storage/Tests/Functional/ContentId/CompressedBlobTests.cs b/Engine/Source/Programs/Horde/HordeStorage/Horde.Storage/Tests/Functional/ContentId/CompressedBlobTests.cs
--- a/Engine/Source/Programs/Horde/HordeStorage/Horde.Storage/Tests/Functional/ContentId/CompressedBlobTests.cs
+++ b/Engine/Source/Programs/Horde/HordeStorage/Horde.Storage/Tests/Functional/ContentId/CompressedBlobTests.cs
@@ -30,6 +30,7 @@ namespace Horde.Storage.FunctionalTests.CompressedBlobs
             // Use the S3 storage backend so we can handle large blobs
             return new[]
             {
+                new KeyValuePair<string, string>("Horde_Storage:ContentIdStoreImplementation", HordeStorageSettings.ContentIdStoreImplementations.Scylla.ToString()),
                 new KeyValuePair<string, string>("Horde_Storage:StorageImplementations:0", HordeStorageSettings.StorageBackendImplementations.S3.ToString()),
                 new KeyValuePair<string, string>("S3:BucketName", $"tests-{TestNamespace}")
             };
@@ -46,6 +47,32 @@ namespace Horde.Storage.FunctionalTests.CompressedBlobs
         }
     }
 
+    [TestClass]
+    public class MongoCompressedBlobTests : CompressedBlobTests
+    {
+        protected override IEnumerable<KeyValuePair<string, string>> GetSettings()
+        {
+            // Use the S3 storage backend so we can handle large blobs
+            return new[]
+            {
+                new KeyValuePair<string, string>("Horde_Storage:ContentIdStoreImplementation", HordeStorageSettings.ContentIdStoreImplementations.Mongo.ToString()),
+                new KeyValuePair<string, string>("Horde_Storage:StorageImplementations:0", HordeStorageSettings.StorageBackendImplementations.S3.ToString()),
+                new KeyValuePair<string, string>("S3:BucketName", $"tests-{TestNamespace}")
+            };
+        }
+
+        protected override async Task Seed(IServiceProvider provider)
+        {
+            await Task.CompletedTask;
+        }
+
+        protected override async Task Teardown()
+        {
+            await Task.CompletedTask;
+        }
+    }
+
+
 
     public abstract class CompressedBlobTests
     {
diff --git a/Engine/Source/Programs/Horde/HordeStorage/Horde.Storage/Tests/Functional/Objects/ReferencesTests.cs b/Engine/Source/Programs/Horde/HordeStorage/Horde.Storage/Tests/Functional/Objects/ReferencesTests.cs
--- a/Engine/Source/Programs/Horde/HordeStorage/Horde.Storage/Tests/Functional/Objects/ReferencesTests.cs
+++ b/Engine/Source/Programs/Horde/HordeStorage/Horde.Storage/Tests/Functional/Objects/ReferencesTests.cs
@@ -82,6 +82,42 @@ namespace Horde.Storage.FunctionalTests.References
         }
     }
 
+    [TestClass]
+    public class MongoReferencesTests : ReferencesTests
+    {
+        protected override IEnumerable<KeyValuePair<string, string>> GetSettings()
+        {
+            return new[]
+            {
+                new KeyValuePair<string, string>("Horde_Storage:ReferencesDbImplementation", HordeStorageSettings.ReferencesDbImplementations.Mongo.ToString()),
+                new KeyValuePair<string, string>("Horde_Storage:ContentIdStoreImplementation", HordeStorageSettings.ContentIdStoreImplementations.Mongo.ToString()),
+                // we do not have a mongo version of the replication log, as the mongo deployment is only intended for single servers
+                new KeyValuePair<string, string>("Horde_Storage:ReplicationLogWriterImplementation", HordeStorageSettings.ReplicationLogWriterImplementations.Memory.ToString()),
+            };
+        }
+
+        protected override async Task SeedDb(IServiceProvider provider)
+        {
+            IReferencesStore referencesStore = provider.GetService<IReferencesStore>()!;
+            //verify we are using the expected store
+            Assert.IsTrue(referencesStore.GetType() == typeof(MongoReferencesStore));
+
+            IContentIdStore contentIdStore = provider.GetService<IContentIdStore>()!;
+            //verify we are using the expected store
+            Assert.IsTrue(contentIdStore.GetType() == typeof(MongoContentIdStore));
+
+            IReplicationLog replicationLog = provider.GetService<IReplicationLog>()!;
+            //verify we are using the replication log writer
+            Assert.IsTrue(replicationLog.GetType() == typeof(MemoryReplicationLog));
+
+            await SeedTestData(referencesStore);
+        }
+
+        protected override async Task TeardownDb(IServiceProvider provider)
+        {
+            await Task.CompletedTask;
+            }
+    }
 
     [TestClass]
     public class MemoryReferencesTests : ReferencesTests
diff --git a/Engine/Source/Programs/Horde/HordeStorage/Horde.Storage/appsettings.Development.json b/Engine/Source/Programs/Horde/HordeStorage/Horde.Storage/appsettings.Development.json
--- a/Engine/Source/Programs/Horde/HordeStorage/Horde.Storage/appsettings.Development.json
+++ b/Engine/Source/Programs/Horde/HordeStorage/Horde.Storage/appsettings.Development.json
@@ -9,118 +9,89 @@
       }
     },
     "WriteTo": [
-      { "Name": "Console" }
+      {
+        "Name": "Console"
+      }
     ],
-    "Enrich": [ "FromLogContext", "WithMachineName", "WithThreadId" ],
-    "Properties": {
-      "Application": "Horde.Storage"
-    }
+    "Enrich": [ "FromLogContext" ]
   },
 
-  "Horde_Storage": {
-    "refDBImplementation": "memory",
-
-    "TransactionLogWriterImplementation": "memory",
+  "Mongo": {
+    // this assumes a mongo db were we have localhost exemption, e.g. we do not need to authenticate
+    // if using mongo via a docker container you will need to update this
+    "ConnectionString": "mongodb://localhost:27017",
+    "RequireTls12": true
+  },
 
-    "treeStoreImplementation": "memory",
+  // scylla is normally not enabled in development but we still configure it so we can easily switch
+  "Scylla": {
+    "LocalDatacenterName": "datacenter1",
+    "LocalKeyspaceSuffix": "local"
+  },
 
-    "replicationLogWriterImplementation": "memory",
+  "Horde_Storage": {
+    // assume a single instance is running
+    "LeaderElectionImplementation": "Static",
 
-    "blobIndexImplementation": "memory"
-  },
+    // use mongo for any database needs
+    "ReferencesDbImplementations": "Mongo",
+    "ContentIdStoreImplementations": "Mongo",
+    "BlobIndexImplementation": "Mongo",
 
-  // settings s3 defaults to use a local minio instance, same as tests, not enabled by default but useful to do local testing
-  "S3": {
-    "ConnectionString": "http://127.0.0.1:9000",
-    "ForceAWSPathStyle": true, // required by minio,
-    "SetBucketPolicies": false, // minio does not support bucket access policies
-    "BucketName": "local-horde-storage-bucket"
-  },
-  "AWSCredentials": {
-    // minio and health checks expect access keys to have been set
-    "AwsAccessKey": "FSYPTVCPKC2IV5D1U8RE",
-    "AwsSecretKey": "WaUf4Zq+J6YMlvbCVvTRkQJIPc7DvZwTm5UP0tPJ"
-  },
+    // store blobs in the filesystem
+    "StorageImplementations": [
+      "Filesystem"
+    ],
 
-  "Callisto": {
-    "ConnectionString": "http://localhost"
+    // disable old api that we do not want to use
+    "TransactionLogWriterImplementation": "Memory",
+    "RefDbImplementation": "Memory"
   },
 
-  "GC": {
-    "BlobCleanupServiceEnabled": false,
-    "CleanOldRefRecords": false,
-    "CleanOldBlobs": false
+  "Filesystem": {
+    "RootDir": "$(ExecutableLocation)/BlobStorage"
   },
 
   "Replication": {
-    "StateRoot": "./ganymede",
+    "StateRoot": "$(ExecutableLocation)/ganymede",
     "ReplicationPollFrequencySeconds": 15,
     "CurrentSite": "local-test",
-    "Replicators": [
-      {
-        "ReplicatorName": "Local test",
-        "NamespaceToReplicate": "ue4.ddc",
-
-        "ConnectionString": "https://jupiter.devtools-dev.epicgames.com"
-      }
-    ]
-  },
-
-  "Auth": {
-    "Method": "Disabled"
-  },
-
-  "Mongo": {
-    "ConnectionString": "mongodb://root:abc123@localhost:27017",
-    "RequireTls12": false
-  },
-
-  "Scylla": {
-    "LocalDatacenterName": "datacenter1",
-    "LocalKeyspaceSuffix": "local"
+    "Enabled": false
   },
 
-  "DynamoDb": {
-    "ConnectionString": "localhost:8000",
-    "UseGlobalTables": false
+  "GC": {
+    "BlobCleanupServiceEnabled": false,
+    "CleanOldBlobs": false,
+    "CleanOldRefRecords": false
   },
 
   "Jupiter": {
-    "CurrentSite": "local",
-    "DisableAuthOnPorts": [
-      57003
-    ]
+    "CurrentSite": "local"
   },
-  "ConsistencyCheck": {
-    "Enabled": false
+
+  // disable auth and allow any namespace
+  "Auth": {
+    "Method": "Disabled"
   },
 
   "Namespaces": {
     "Policies": {
-      "ue4.ddc": {
-        "Claims": [ "UE4" ]
-      },
-      "ue.ddc": {
-        "Claims": [ "UE4" ],
-        "StoragePool": "ddc2"
-      },
       "*": {
         "Claims": [ "*" ]
       }
     }
   },
 
+  "AllowedHosts": "*",
+
   "Kestrel": {
     "Endpoints": {
       "Http": {
-        "Url": "http://localhost:57001"
+        "Url": "http://0.0.0.0:57001"
       },
       "Http2": {
-        "Url": "http://localhost:57002",
+        "Url": "http://0.0.0.0:57002",
         "Protocols": "Http2"
-      },
-      "AuthDisabledHttp": {
-        "Url": "http://localhost:57003"
       }
     }
   }
