Horde: Include costing data in lease responses (if the user is authorized to view it). Also refactor the way that payload data is generated for lease responses so that we can always show relevant fields.

diff --git a/Engine/Source/Programs/Horde/Horde.Build/Agents/IAgent.cs b/Engine/Source/Programs/Horde/Horde.Build/Agents/IAgent.cs
--- a/Engine/Source/Programs/Horde/Horde.Build/Agents/IAgent.cs
+++ b/Engine/Source/Programs/Horde/Horde.Build/Agents/IAgent.cs
@@ -304,48 +304,6 @@ namespace HordeServer.Models
 			return false;
 		}
 
-		/// <summary>
-		/// Gets user-readable payload information
-		/// </summary>
-		/// <param name="Payload">The payload data</param>
-		/// <returns>Dictionary of key/value pairs for the payload</returns>
-		public static Dictionary<string, string>? GetPayloadDetails(ReadOnlyMemory<byte>? Payload)
-		{
-			Dictionary<string, string>? Details = null;
-			if (Payload != null)
-			{
-				Any BasePayload = Any.Parser.ParseFrom(Payload.Value.ToArray());
-
-				Details = new Dictionary<string, string>();
-
-				ConformTask ConformTask;
-				if(BasePayload.TryUnpack(out ConformTask))
-				{
-					Details["Type"] = "Conform";
-					Details["LogId"] = ConformTask.LogId;
-					Details["Full"] = ConformTask.RemoveUntrackedFiles.ToString();
-				}
-
-				ExecuteJobTask JobTask;
-				if (BasePayload.TryUnpack(out JobTask))
-				{
-					Details["Type"] = "Job";
-					Details["JobId"] = JobTask.JobId;
-					Details["BatchId"] = JobTask.BatchId;
-					Details["LogId"] = JobTask.LogId;
-				}
-
-				UpgradeTask UpgradeTask;
-				if (BasePayload.TryUnpack(out UpgradeTask))
-				{
-					Details["Type"] = "Upgrade";
-					Details["SoftwareId"] = UpgradeTask.SoftwareId;
-					Details["LogId"] = UpgradeTask.LogId;
-				}
-			}
-			return Details;
-		}
-
 		/// <summary>
 		/// Converts this lease to an RPC message
 		/// </summary>
diff --git a/Engine/Source/Programs/Horde/Horde.Build/Agents/ITaskSource.cs b/Engine/Source/Programs/Horde/Horde.Build/Agents/ITaskSource.cs
--- a/Engine/Source/Programs/Horde/Horde.Build/Agents/ITaskSource.cs
+++ b/Engine/Source/Programs/Horde/Horde.Build/Agents/ITaskSource.cs
@@ -74,6 +74,13 @@ namespace HordeServer.Tasks
 		/// <param name="Output">Output from the task</param>
 		/// <param name="Logger">Logger for the agent</param>
 		Task OnLeaseFinishedAsync(IAgent Agent, LeaseId LeaseId, Any Payload, LeaseOutcome Outcome, ReadOnlyMemory<byte> Output, ILogger Logger);
+
+		/// <summary>
+		/// Gets information to include for a lease in a lease info response
+		/// </summary>
+		/// <param name="Payload">The lease payload</param>
+		/// <param name="Details">Properties for the lease</param>
+		void GetLeaseDetails(Any Payload, Dictionary<string, string> Details);
 	}
 	
 	/// <summary>
@@ -89,6 +96,7 @@ namespace HordeServer.Tasks
 		{
 			internal StringBuilder FormatString = new StringBuilder();
 			internal List<Func<TMessage, object>> Accessors = new List<Func<TMessage, object>>();
+			internal List<(string, Func<TMessage, object>)> JsonAccessors = new List<(string, Func<TMessage, object>)>();
 
 			/// <summary>
 			/// Adds a new property to the list
@@ -105,10 +113,12 @@ namespace HordeServer.Tasks
 			/// </summary>
 			/// <param name="Name">Name of the property</param>
 			/// <param name="Accessor">Accessor for the property</param>
+			[System.Diagnostics.CodeAnalysis.SuppressMessage("Globalization", "CA1308:Normalize strings to uppercase", Justification = "<Pending>")]
 			public PropertyList Add(string Name, Func<TMessage, object> Accessor)
 			{
 				FormatString.Append(CultureInfo.InvariantCulture, $", {Name}={{{Name}}}");
 				Accessors.Add(Accessor);
+				JsonAccessors.Add((Name[0..1].ToLower(CultureInfo.InvariantCulture) + Name[1..], Accessor));
 				return this;
 			}
 		}
@@ -161,5 +171,20 @@ namespace HordeServer.Tasks
 			Logger.LogInformation("Lease {LeaseId} complete", LeaseId);
 			return Task.CompletedTask;
 		}
+
+		/// <inheritdoc/>
+		public virtual void GetLeaseDetails(Any Payload, Dictionary<string, string> Details)
+		{
+			Details["type"] = Type;
+
+			if (OnLeaseStartedProperties.Accessors.Count > 0)
+			{
+				TMessage Message = Payload.Unpack<TMessage>();
+				foreach ((string Name, Func<TMessage, object> GetMethod) in OnLeaseStartedProperties.JsonAccessors)
+				{
+					Details[Name] = GetMethod(Message)?.ToString() ?? String.Empty;
+				}
+			}
+		}
 	}
 }
diff --git a/Engine/Source/Programs/Horde/Horde.Build/Agents/Impl/AgentService.cs b/Engine/Source/Programs/Horde/Horde.Build/Agents/Impl/AgentService.cs
--- a/Engine/Source/Programs/Horde/Horde.Build/Agents/Impl/AgentService.cs
+++ b/Engine/Source/Programs/Horde/Horde.Build/Agents/Impl/AgentService.cs
@@ -132,6 +132,30 @@ namespace HordeServer.Services
 			Ticker.Dispose();
 		}
 
+		/// <summary>
+		/// Gets user-readable payload information
+		/// </summary>
+		/// <param name="Payload">The payload data</param>
+		/// <returns>Dictionary of key/value pairs for the payload</returns>
+		public Dictionary<string, string>? GetPayloadDetails(ReadOnlyMemory<byte>? Payload)
+		{
+			Dictionary<string, string>? Details = null;
+			if (Payload != null)
+			{
+				Any BasePayload = Any.Parser.ParseFrom(Payload.Value.ToArray());
+				foreach (ITaskSource TaskSource in TaskSources)
+				{
+					if (BasePayload.Is(TaskSource.Descriptor))
+					{
+						Details = new Dictionary<string, string>();
+						TaskSource.GetLeaseDetails(BasePayload, Details);
+						break;
+					}
+				}
+			}
+			return Details;
+		}
+
 		/// <summary>
 		/// Issues a bearer token for the given session id
 		/// </summary>
diff --git a/Engine/Source/Programs/Horde/Horde.Build/Api/Agents.cs b/Engine/Source/Programs/Horde/Horde.Build/Api/Agents.cs
--- a/Engine/Source/Programs/Horde/Horde.Build/Api/Agents.cs
+++ b/Engine/Source/Programs/Horde/Horde.Build/Api/Agents.cs
@@ -137,6 +137,11 @@ namespace HordeServer.Api
 		/// </summary>
 		public string? AgentId { get; set; }
 
+		/// <summary>
+		/// Cost of this agent, per hour
+		/// </summary>
+		public double? AgentRate { get; set; }
+
 		/// <summary>
 		/// Name of the lease
 		/// </summary>
@@ -181,7 +186,8 @@ namespace HordeServer.Api
 		/// Constructor
 		/// </summary>
 		/// <param name="Lease">The lease to initialize from</param>
-		public GetAgentLeaseResponse(AgentLease Lease)
+		/// <param name="Details">The payload details</param>
+		public GetAgentLeaseResponse(AgentLease Lease, Dictionary<string, string>? Details)
 		{
 			this.Id = Lease.Id.ToString();
 			this.Name = Lease.Name;
@@ -190,23 +196,26 @@ namespace HordeServer.Api
 			this.StartTime = Lease.StartTime;
 			this.Executing = Lease.Active;
 			this.FinishTime = Lease.ExpiryTime;
-			this.Details = AgentLease.GetPayloadDetails(Lease.Payload);			
+			this.Details = Details;
 		}
 
 		/// <summary>
 		/// Converts this lease to a response object
 		/// </summary>
 		/// <param name="Lease">The lease to initialize from</param>
-		public GetAgentLeaseResponse(ILease Lease)
+		/// <param name="Details">The payload details</param>
+		/// <param name="AgentRate">Rate for running this agent</param>
+		public GetAgentLeaseResponse(ILease Lease, Dictionary<string, string>? Details, double? AgentRate)
 		{
 			this.Id = Lease.Id.ToString();
 			this.AgentId = Lease.AgentId.ToString();
+			this.AgentRate = AgentRate;
 			this.Name = Lease.Name;
 			this.LogId = Lease.LogId?.ToString();
 			this.StartTime = Lease.StartTime;
 			this.Executing = (Lease.FinishTime == null);
 			this.FinishTime = Lease.FinishTime;
-			this.Details = AgentLease.GetPayloadDetails(Lease.Payload);
+			this.Details = Details;
 			this.Outcome = Lease.Outcome;
 		}
 	}
@@ -446,9 +455,10 @@ namespace HordeServer.Api
 		/// Constructor
 		/// </summary>
 		/// <param name="Agent">The agent to construct from</param>
+		/// <param name="Leases">Active leases</param>
 		/// <param name="Rate">Rate for this agent</param>
 		/// <param name="bIncludeAcl">Whether to include the ACL in the response</param>
-		public GetAgentResponse(IAgent Agent, double? Rate, bool bIncludeAcl)
+		public GetAgentResponse(IAgent Agent, List<GetAgentLeaseResponse> Leases, double? Rate, bool bIncludeAcl)
 		{
 			this.Id = Agent.Id.ToString();
 			this.Name = Agent.Id.ToString();
@@ -477,7 +487,7 @@ namespace HordeServer.Api
 			this.Pools = Agent.GetPools().Select(x => x.ToString()).ToList();
 			this.Workspaces = Agent.Workspaces.ConvertAll(x => new GetAgentWorkspaceResponse(x));
 			this.Capabilities = new { Devices = new[] { new { Properties = Agent.Properties, Resources = Agent.Resources } } };
-			this.Leases = Agent.Leases.ConvertAll(x => new GetAgentLeaseResponse(x));
+			this.Leases = Leases;
 			this.Acl = (bIncludeAcl && Agent.Acl != null) ? new GetAclResponse(Agent.Acl) : null;
 			this.Comment = Agent.Comment;
 		}
diff --git a/Engine/Source/Programs/Horde/Horde.Build/Compute/Impl/ComputeService.cs b/Engine/Source/Programs/Horde/Horde.Build/Compute/Impl/ComputeService.cs
--- a/Engine/Source/Programs/Horde/Horde.Build/Compute/Impl/ComputeService.cs
+++ b/Engine/Source/Programs/Horde/Horde.Build/Compute/Impl/ComputeService.cs
@@ -156,7 +156,7 @@ namespace HordeServer.Compute.Impl
 			this.CachedPools = new LazyCachedValue<Task<List<IPool>>>(() => PoolCollection.GetAsync(), TimeSpan.FromSeconds(30.0));
 			this.Logger = Logger;
 
-			OnLeaseStartedProperties.Add(x => x.Task);
+			OnLeaseStartedProperties.Add(nameof(ComputeTaskMessage.Task), x => x.Task.Hash);
 		}
 
 		/// <inheritdoc/>
diff --git a/Engine/Source/Programs/Horde/Horde.Build/Controllers/AgentsController.cs b/Engine/Source/Programs/Horde/Horde.Build/Controllers/AgentsController.cs
--- a/Engine/Source/Programs/Horde/Horde.Build/Controllers/AgentsController.cs
+++ b/Engine/Source/Programs/Horde/Horde.Build/Controllers/AgentsController.cs
@@ -98,13 +98,7 @@ namespace HordeServer.Controllers
 			List<object> Responses = new List<object>();
 			foreach (IAgent Agent in Agents)
 			{
-				double? Rate = null;
-				if (await AgentService.AuthorizeAsync(Agent, AclAction.ViewCosts, User, PermissionsCache))
-				{
-					Rate = await AgentService.GetRateAsync(Agent.Id);
-				}
-				bool bIncludeAcl = await AgentService.AuthorizeAsync(Agent, AclAction.ViewPermissions, User, PermissionsCache);
-				Responses.Add(new GetAgentResponse(Agent, Rate, bIncludeAcl).ApplyFilter(Filter));
+				Responses.Add(await GetAgentResponseAsync(Agent, PermissionsCache, Filter));
 			}
 
 			return Responses;
@@ -127,16 +121,29 @@ namespace HordeServer.Controllers
 				return NotFound(AgentId);
 			}
 
-			GlobalPermissionsCache Cache = new GlobalPermissionsCache();
+			return await GetAgentResponseAsync(Agent, new GlobalPermissionsCache(), Filter);
+		}
 
+		/// <summary>
+		/// Gets an individual agent response
+		/// </summary>
+		async ValueTask<object> GetAgentResponseAsync(IAgent Agent, GlobalPermissionsCache PermissionsCache, PropertyFilter? Filter = null)
+		{
 			double? Rate = null;
-			if (await AgentService.AuthorizeAsync(Agent, AclAction.ViewCosts, User, Cache))
+			if (await AgentService.AuthorizeAsync(Agent, AclAction.ViewCosts, User, PermissionsCache))
 			{
 				Rate = await AgentService.GetRateAsync(Agent.Id);
 			}
 
-			bool bIncludeAcl = await AgentService.AuthorizeAsync(Agent, AclAction.ViewPermissions, User, Cache);
-			return new GetAgentResponse(Agent, Rate, bIncludeAcl).ApplyFilter(Filter);
+			List<GetAgentLeaseResponse> Leases = new List<GetAgentLeaseResponse>();
+			foreach (AgentLease Lease in Agent.Leases)
+			{
+				Dictionary<string, string>? Details = AgentService.GetPayloadDetails(Lease.Payload);
+				Leases.Add(new GetAgentLeaseResponse(Lease, Details));
+			}
+
+			bool bIncludeAcl = await AgentService.AuthorizeAsync(Agent, AclAction.ViewPermissions, User, PermissionsCache);
+			return new GetAgentResponse(Agent, Leases, Rate, bIncludeAcl).ApplyFilter(Filter);
 		}
 
 		/// <summary>
@@ -318,10 +325,12 @@ namespace HordeServer.Controllers
 		/// <param name="FinishTime">End of the time window to consider</param>
 		/// <param name="Index">Index of the first result to return</param>
 		/// <param name="Count">Number of results to return</param>
+		/// <param name="Filter">Filter to apply to the properties</param>
 		/// <returns>Sessions </returns>
 		[HttpGet]
 		[Route("/api/v1/agents/{AgentId}/leases")]
-		public async Task<ActionResult<List<GetAgentLeaseResponse>>> FindLeasesAsync(AgentId AgentId, [FromQuery] SessionId? SessionId, [FromQuery] DateTimeOffset? StartTime, [FromQuery] DateTimeOffset? FinishTime, [FromQuery] int Index = 0, [FromQuery] int Count = 1000)
+		[ProducesResponseType(200, Type = typeof(List<GetAgentLeaseResponse>))]
+		public async Task<ActionResult<List<object>>> FindLeasesAsync(AgentId AgentId, [FromQuery] SessionId? SessionId, [FromQuery] DateTimeOffset? StartTime, [FromQuery] DateTimeOffset? FinishTime, [FromQuery] int Index = 0, [FromQuery] int Count = 1000, [FromQuery] PropertyFilter? Filter = null)
 		{
 			IAgent? Agent = await AgentService.GetAgentAsync(AgentId);
 			if (Agent == null)
@@ -334,7 +343,21 @@ namespace HordeServer.Controllers
 			}
 
 			List<ILease> Leases = await AgentService.FindLeasesAsync(AgentId, SessionId, StartTime?.UtcDateTime, FinishTime?.UtcDateTime, Index, Count);
-			return Leases.ConvertAll(x => new GetAgentLeaseResponse(x));
+
+			double? AgentRate = null;
+			if (await AclService.AuthorizeAsync(AclAction.ViewCosts, User))
+			{
+				AgentRate = await AgentService.GetRateAsync(AgentId);
+			}
+
+			List<object> Responses = new List<object>();
+			foreach(ILease Lease in Leases)
+			{
+				Dictionary<string, string>? Details = AgentService.GetPayloadDetails(Lease.Payload);
+				Responses.Add(PropertyFilter.Apply(new GetAgentLeaseResponse(Lease, Details, AgentRate), Filter));
+			}
+
+			return Responses;
 		}
 
 		/// <summary>
@@ -363,7 +386,14 @@ namespace HordeServer.Controllers
 				return NotFound(AgentId, LeaseId);
 			}
 
-			return new GetAgentLeaseResponse(Lease);
+			double? AgentRate = null;
+			if (await AclService.AuthorizeAsync(AclAction.ViewCosts, User))
+			{
+				AgentRate = await AgentService.GetRateAsync(AgentId);
+			}
+
+			Dictionary<string, string>? Details = AgentService.GetPayloadDetails(Lease.Payload);
+			return new GetAgentLeaseResponse(Lease, Details, AgentRate);
 		}
 	}
 }
diff --git a/Engine/Source/Programs/Horde/Horde.Build/Controllers/LeasesController.cs b/Engine/Source/Programs/Horde/Horde.Build/Controllers/LeasesController.cs
--- a/Engine/Source/Programs/Horde/Horde.Build/Controllers/LeasesController.cs
+++ b/Engine/Source/Programs/Horde/Horde.Build/Controllers/LeasesController.cs
@@ -74,14 +74,17 @@ namespace HordeServer.Controllers
 		/// <param name="FinishTime">End of the time window to consider</param>
 		/// <param name="Index">Index of the first result to return</param>
 		/// <param name="Count">Number of results to return</param>
+		/// <param name="Filter">Filter to apply to the results</param>
 		/// <returns>Sessions </returns>
 		[HttpGet]
 		[Route("/api/v1/leases")]
-		public async Task<ActionResult<List<GetAgentLeaseResponse>>> FindLeasesAsync([FromQuery] AgentId? AgentId, [FromQuery] SessionId? SessionId, [FromQuery] DateTimeOffset? StartTime, [FromQuery] DateTimeOffset? FinishTime, [FromQuery] int Index = 0, [FromQuery] int Count = 1000)
+		[ProducesResponseType(200, Type = typeof(List<GetAgentLeaseResponse>))]
+		public async Task<ActionResult<List<object>>> FindLeasesAsync([FromQuery] AgentId? AgentId, [FromQuery] SessionId? SessionId, [FromQuery] DateTimeOffset? StartTime, [FromQuery] DateTimeOffset? FinishTime, [FromQuery] int Index = 0, [FromQuery] int Count = 1000, [FromQuery] PropertyFilter? Filter = null)
 		{
+			GlobalPermissionsCache PermissionsCache = new GlobalPermissionsCache();
 			if (AgentId == null)
 			{
-				if (!await AclService.AuthorizeAsync(AclAction.ViewLeases, User))
+				if (!await AclService.AuthorizeAsync(AclAction.ViewLeases, User, PermissionsCache))
 				{
 					return Forbid(AclAction.ViewLeases);
 				}
@@ -93,14 +96,33 @@ namespace HordeServer.Controllers
 				{
 					return NotFound(AgentId.Value);
 				}
-				if (!await AgentService.AuthorizeAsync(Agent, AclAction.ViewLeases, User, null))
+				if (!await AgentService.AuthorizeAsync(Agent, AclAction.ViewLeases, User, PermissionsCache))
 				{
 					return Forbid(AclAction.ViewLeases, AgentId.Value);
 				}
 			}
 
+			bool IncludeCosts = await AclService.AuthorizeAsync(AclAction.ViewCosts, User, PermissionsCache);
+
 			List<ILease> Leases = await AgentService.FindLeasesAsync(AgentId, SessionId, StartTime?.UtcDateTime, FinishTime?.UtcDateTime, Index, Count);
-			return Leases.ConvertAll(x => new GetAgentLeaseResponse(x));
+
+			List<object> Responses = new List<object>();
+
+			Dictionary<AgentId, double?> CachedAgentRates = new Dictionary<AgentId, double?>();
+			foreach (ILease Lease in Leases)
+			{
+				double? AgentRate = null;
+				if (IncludeCosts && !CachedAgentRates.TryGetValue(Lease.AgentId, out AgentRate))
+				{
+					AgentRate = await AgentService.GetRateAsync(Lease.AgentId);
+					CachedAgentRates.Add(Lease.AgentId, AgentRate);
+				}
+
+				Dictionary<string, string>? Details = AgentService.GetPayloadDetails(Lease.Payload);
+				Responses.Add(PropertyFilter.Apply(new GetAgentLeaseResponse(Lease, Details, AgentRate), Filter));
+			}
+
+			return Responses;
 		}
 
 		/// <summary>
@@ -123,12 +145,21 @@ namespace HordeServer.Controllers
 			{
 				return NotFound(Lease.AgentId);
 			}
-			if (!await AgentService.AuthorizeAsync(Agent, AclAction.ViewLeases, User, null))
+
+			GlobalPermissionsCache PermissionsCache = new GlobalPermissionsCache();
+			if (!await AgentService.AuthorizeAsync(Agent, AclAction.ViewLeases, User, PermissionsCache))
 			{
 				return Forbid(AclAction.ViewLeases);
 			}
 
-			return new GetAgentLeaseResponse(Lease);
+			double? AgentRate = null;
+			if (await AclService.AuthorizeAsync(AclAction.ViewCosts, User, PermissionsCache))
+			{
+				AgentRate = await AgentService.GetRateAsync(Agent.Id);
+			}
+
+			Dictionary<string, string>? Details = AgentService.GetPayloadDetails(Lease.Payload);
+			return new GetAgentLeaseResponse(Lease, Details, AgentRate);
 		}
 
 		/// <summary>
