Horde: Add flags to task sources indicating when they can issue tasks. Allow restart, shutdown, and upgrade tasks to run even during downtime or when a machine is disabled.

diff --git a/Engine/Source/Programs/Horde/Horde.Build/Agents/ITaskSource.cs b/Engine/Source/Programs/Horde/Horde.Build/Agents/ITaskSource.cs
--- a/Engine/Source/Programs/Horde/Horde.Build/Agents/ITaskSource.cs
+++ b/Engine/Source/Programs/Horde/Horde.Build/Agents/ITaskSource.cs
@@ -23,6 +23,28 @@ namespace HordeServer.Tasks
 {
 	using LeaseId = ObjectId<ILease>;
 
+	/// <summary>
+	/// Flags indicating when a task source is valid
+	/// </summary>
+	[Flags]
+	public enum TaskSourceFlags
+	{
+		/// <summary>
+		/// Normal behavior
+		/// </summary>
+		None = 0,
+
+		/// <summary>
+		/// Indicates that tasks from this source can run during downtime
+		/// </summary>
+		AllowDuringDowntime = 1,
+
+		/// <summary>
+		/// Allows this source to schedule tasks when the agent is disabled
+		/// </summary>
+		AllowWhenDisabled = 2,
+	}
+
 	/// <summary>
 	/// Handler for a certain lease type
 	/// </summary>
@@ -38,6 +60,11 @@ namespace HordeServer.Tasks
 		/// </summary>
 		MessageDescriptor Descriptor { get; }
 
+		/// <summary>
+		/// Flags controlling behavior of this source
+		/// </summary>
+		TaskSourceFlags Flags { get; }
+
 		/// <summary>
 		/// Assigns a lease or waits for one to be available
 		/// </summary>
@@ -128,6 +155,9 @@ namespace HordeServer.Tasks
 		/// <inheritdoc/>
 		public abstract string Type { get; }
 
+		/// <inheritdoc/>
+		public abstract TaskSourceFlags Flags { get; }
+
 		/// <inheritdoc/>
 		protected virtual PropertyList OnLeaseStartedProperties { get; } = new PropertyList();
 
diff --git a/Engine/Source/Programs/Horde/Horde.Build/Agents/Impl/AgentService.cs b/Engine/Source/Programs/Horde/Horde.Build/Agents/Impl/AgentService.cs
--- a/Engine/Source/Programs/Horde/Horde.Build/Agents/Impl/AgentService.cs
+++ b/Engine/Source/Programs/Horde/Horde.Build/Agents/Impl/AgentService.cs
@@ -351,6 +351,23 @@ namespace HordeServer.Services
 			}
 		}
 
+		/// <summary>
+		/// Determines whether a task source can currently issue tasks
+		/// </summary>
+		bool CanUseTaskSource(IAgent Agent, ITaskSource TaskSource)
+		{
+			TaskSourceFlags Flags = TaskSource.Flags;
+			if ((Flags & TaskSourceFlags.AllowWhenDisabled) == 0 && !Agent.Enabled)
+			{
+				return false;
+			}
+			if ((Flags & TaskSourceFlags.AllowDuringDowntime) == 0 && DowntimeService.IsDowntimeActive)
+			{
+				return false;
+			}
+			return true;
+		}
+
 		/// <summary>
 		/// Waits for a lease to be assigned to an agent
 		/// </summary>
@@ -388,18 +405,21 @@ namespace HordeServer.Services
 						WaitingAgents[Agent.Id] = CancellationSource;
 					}
 
-					// If we're in a maintenance window, just wait for the time to expire
-					if (DowntimeService.IsDowntimeActive)
-					{
-						await AsyncUtils.DelayNoThrow(MaxWaitTime, CancellationToken);
-						break;
-					}
-
 					// Create all the tasks to wait for
 					List<Task<(ITaskSource, AgentLease)?>> Tasks = new List<Task<(ITaskSource, AgentLease)?>>();
 					foreach (ITaskSource TaskSource in TaskSources)
 					{
-						Tasks.Add(GuardedWaitForLeaseAsync(TaskSource, Agent, CancellationSource.Token));
+						if (CanUseTaskSource(Agent, TaskSource))
+						{
+							Tasks.Add(GuardedWaitForLeaseAsync(TaskSource, Agent, CancellationSource.Token));
+						}
+					}
+
+					// If no task source is valid, just add a delay
+					if (Tasks.Count == 0)
+					{
+						await AsyncUtils.DelayNoThrow(MaxWaitTime, CancellationToken);
+						break;
 					}
 
 					// Wait for a lease to be available
diff --git a/Engine/Source/Programs/Horde/Horde.Build/Agents/Impl/Tasks/ConformTaskSource.cs b/Engine/Source/Programs/Horde/Horde.Build/Agents/Impl/Tasks/ConformTaskSource.cs
--- a/Engine/Source/Programs/Horde/Horde.Build/Agents/Impl/Tasks/ConformTaskSource.cs
+++ b/Engine/Source/Programs/Horde/Horde.Build/Agents/Impl/Tasks/ConformTaskSource.cs
@@ -33,6 +33,9 @@ namespace HordeServer.Tasks.Impl
 		/// <inheritdoc/>
 		public override string Type => "Conform";
 
+		/// <inheritdoc/>
+		public override TaskSourceFlags Flags => TaskSourceFlags.AllowWhenDisabled;
+
 		DatabaseService DatabaseService;
 		IAgentCollection AgentCollection;
 		PoolService PoolService;
diff --git a/Engine/Source/Programs/Horde/Horde.Build/Agents/Impl/Tasks/RestartTaskSource.cs b/Engine/Source/Programs/Horde/Horde.Build/Agents/Impl/Tasks/RestartTaskSource.cs
--- a/Engine/Source/Programs/Horde/Horde.Build/Agents/Impl/Tasks/RestartTaskSource.cs
+++ b/Engine/Source/Programs/Horde/Horde.Build/Agents/Impl/Tasks/RestartTaskSource.cs
@@ -26,6 +26,8 @@ namespace HordeServer.Tasks.Impl
 	{
 		public override string Type => "Restart";
 
+		public override TaskSourceFlags Flags => TaskSourceFlags.AllowWhenDisabled | TaskSourceFlags.AllowDuringDowntime;
+
 		ILogFileService LogService;
 
 		public RestartTaskSource(ILogFileService LogService)
diff --git a/Engine/Source/Programs/Horde/Horde.Build/Agents/Impl/Tasks/ShutdownTaskSource.cs b/Engine/Source/Programs/Horde/Horde.Build/Agents/Impl/Tasks/ShutdownTaskSource.cs
--- a/Engine/Source/Programs/Horde/Horde.Build/Agents/Impl/Tasks/ShutdownTaskSource.cs
+++ b/Engine/Source/Programs/Horde/Horde.Build/Agents/Impl/Tasks/ShutdownTaskSource.cs
@@ -27,6 +27,8 @@ namespace HordeServer.Tasks.Impl
 	{
 		public override string Type => "Shutdown";
 
+		public override TaskSourceFlags Flags => TaskSourceFlags.AllowWhenDisabled | TaskSourceFlags.AllowDuringDowntime;
+
 		ILogFileService LogService;
 
 		public ShutdownTaskSource(ILogFileService LogService)
diff --git a/Engine/Source/Programs/Horde/Horde.Build/Agents/Impl/Tasks/UpgradeTaskSource.cs b/Engine/Source/Programs/Horde/Horde.Build/Agents/Impl/Tasks/UpgradeTaskSource.cs
--- a/Engine/Source/Programs/Horde/Horde.Build/Agents/Impl/Tasks/UpgradeTaskSource.cs
+++ b/Engine/Source/Programs/Horde/Horde.Build/Agents/Impl/Tasks/UpgradeTaskSource.cs
@@ -27,6 +27,8 @@ namespace HordeServer.Tasks.Impl
 	{
 		public override string Type => "Upgrade";
 
+		public override TaskSourceFlags Flags => TaskSourceFlags.AllowWhenDisabled | TaskSourceFlags.AllowDuringDowntime;
+
 		AgentSoftwareService AgentSoftwareService;
 		ILogFileService LogService;
 		IClock Clock;
diff --git a/Engine/Source/Programs/Horde/Horde.Build/Compute/Impl/ComputeService.cs b/Engine/Source/Programs/Horde/Horde.Build/Compute/Impl/ComputeService.cs
--- a/Engine/Source/Programs/Horde/Horde.Build/Compute/Impl/ComputeService.cs
+++ b/Engine/Source/Programs/Horde/Horde.Build/Compute/Impl/ComputeService.cs
@@ -131,8 +131,12 @@ namespace HordeServer.Compute.Impl
 	/// </summary>
 	class ComputeService : TaskSourceBase<ComputeTaskMessage>, IComputeService, IHostedService, IDisposable
 	{
+		/// <inheritdoc/>
 		public override string Type => "Compute";
 
+		/// <inheritdoc/>
+		public override TaskSourceFlags Flags => TaskSourceFlags.None;
+
 		public static NamespaceId DefaultNamespaceId { get; } = new NamespaceId("default");
 
 		IObjectCollection ObjectCollection;
diff --git a/Engine/Source/Programs/Horde/Horde.Build/Jobs/Impl/JobTaskSource.cs b/Engine/Source/Programs/Horde/Horde.Build/Jobs/Impl/JobTaskSource.cs
--- a/Engine/Source/Programs/Horde/Horde.Build/Jobs/Impl/JobTaskSource.cs
+++ b/Engine/Source/Programs/Horde/Horde.Build/Jobs/Impl/JobTaskSource.cs
@@ -47,6 +47,9 @@ namespace HordeServer.Tasks.Impl
 		/// <inheritdoc/>
 		public override string Type => "Job";
 
+		/// <inheritdoc/>
+		public override TaskSourceFlags Flags => TaskSourceFlags.None;
+
 		/// <summary>
 		/// An item in the queue to be executed
 		/// </summary>
