Horde: Run finalize logic even if a lease is aborted.

diff --git a/Engine/Source/Programs/Horde/Horde.Agent/Services/WorkerService.cs b/Engine/Source/Programs/Horde/Horde.Agent/Services/WorkerService.cs
--- a/Engine/Source/Programs/Horde/Horde.Agent/Services/WorkerService.cs
+++ b/Engine/Source/Programs/Horde/Horde.Agent/Services/WorkerService.cs
@@ -1030,6 +1030,12 @@ namespace HordeAgent.Services
 				}
 			}
 
+			// If this lease was cancelled, don't bother updating the job state.
+			if (CancellationToken.IsCancellationRequested)
+			{
+				return LeaseResult.Cancelled;
+			}
+
 			// Mark the batch as complete
 			await RpcClient.InvokeAsync(x => x.FinishBatchAsync(new FinishBatchRequest(ExecuteTask.JobId, ExecuteTask.BatchId, LeaseId), null, null, CancellationToken), new RpcContext(), CancellationToken);
 			Logger.LogInformation("Done.");
@@ -1067,116 +1073,130 @@ namespace HordeAgent.Services
 				await Executor.InitializeAsync(BatchLogger, CancellationToken);
 			}
 
-			// Execute the steps
-			for (; ; )
+			try
 			{
-				// Get the next step to execute
-				BeginStepResponse Step = await RpcClient.InvokeAsync(x => x.BeginStepAsync(new BeginStepRequest(ExecuteTask.JobId, ExecuteTask.BatchId, LeaseId), null, null, CancellationToken), new RpcContext(), CancellationToken);
-				if (Step.State == BeginStepResponse.Types.Result.Waiting)
-				{
-					BatchLogger.LogInformation("Waiting for dependency to be ready");
-					await Task.Delay(TimeSpan.FromSeconds(20.0), CancellationToken);
-					continue;
-				}
-				else if (Step.State == BeginStepResponse.Types.Result.Complete)
-				{
-					break;
-				}
-				else if (Step.State != BeginStepResponse.Types.Result.Ready)
-				{
-					BatchLogger.LogError("Unexpected step state: {StepState}", Step.State);
-					break;
-				}
-
-				// Get current disk space available. This will allow us to more easily spot steps that eat up a lot of disk space.
-				string? DriveName;
-				if (RuntimeInformation.IsOSPlatform(OSPlatform.Windows))
-				{
-					DriveName = Path.GetPathRoot(WorkingDir.FullName);
-				}
-				else
-				{
-					DriveName = WorkingDir.FullName;
-				}
-
-				float AvailableFreeSpace = 0;
-				if (DriveName != null)
+				// Execute the steps
+				for (; ; )
 				{
-					try
+					// Get the next step to execute
+					BeginStepResponse Step = await RpcClient.InvokeAsync(x => x.BeginStepAsync(new BeginStepRequest(ExecuteTask.JobId, ExecuteTask.BatchId, LeaseId), null, null, CancellationToken), new RpcContext(), CancellationToken);
+					if (Step.State == BeginStepResponse.Types.Result.Waiting)
 					{
-						DriveInfo Info = new DriveInfo(DriveName);
-						AvailableFreeSpace = (1.0f * Info.AvailableFreeSpace) / 1024 / 1024 / 1024;
+						BatchLogger.LogInformation("Waiting for dependency to be ready");
+						await Task.Delay(TimeSpan.FromSeconds(20.0), CancellationToken);
+						continue;
 					}
-					catch (Exception Ex)
+					else if (Step.State == BeginStepResponse.Types.Result.Complete)
 					{
-						BatchLogger.LogWarning(Ex, "Unable to query disk info for path '{DriveName}'", DriveName);
+						break;
+					}
+					else if (Step.State != BeginStepResponse.Types.Result.Ready)
+					{
+						BatchLogger.LogError("Unexpected step state: {StepState}", Step.State);
+						break;
 					}
-				}
 
-				// Print the new state
-				Stopwatch StepTimer = Stopwatch.StartNew();
-				
-				BatchLogger.LogInformation("Starting job {JobId}, batch {BatchId}, step {StepId} (Drive Space Left: {DriveSpaceRemaining} GB)", ExecuteTask.JobId, ExecuteTask.BatchId, Step.StepId, AvailableFreeSpace.ToString("F1"));
-
-				// Create a trace span
-				using IScope Scope = GlobalTracer.Instance.BuildSpan("Execute").WithResourceName(Step.Name).StartActive();
-				Scope.Span.SetTag("stepId", Step.StepId);
-				Scope.Span.SetTag("logId", Step.LogId);
-//				using IDisposable TraceProperty = LogContext.PushProperty("dd.trace_id", CorrelationIdentifier.TraceId.ToString());
-//				using IDisposable SpanProperty = LogContext.PushProperty("dd.span_id", CorrelationIdentifier.SpanId.ToString());
-
-				// Update the context to include information about this step
-				JobStepOutcome StepOutcome;
-				JobStepState StepState;
-				using (BatchLogger.BeginIndentScope("  "))
-				{
-					// Start writing to the log file
-					await using (JsonRpcLogger StepLogger = new JsonRpcLogger(RpcClient, Step.LogId, ExecuteTask.JobId, ExecuteTask.BatchId, Step.StepId, Step.Warnings, Logger))
+					// Get current disk space available. This will allow us to more easily spot steps that eat up a lot of disk space.
+					string? DriveName;
+					if (RuntimeInformation.IsOSPlatform(OSPlatform.Windows))
 					{
-						// Execute the task
-						ILogger ForwardingLogger = new DefaultLoggerIndentHandler(StepLogger);
-						if (Settings.WriteStepOutputToLogger)
-						{
-							ForwardingLogger = new ForwardingLogger(Logger, ForwardingLogger);
-						}
+						DriveName = Path.GetPathRoot(WorkingDir.FullName);
+					}
+					else
+					{
+						DriveName = WorkingDir.FullName;
+					}
 
-						using CancellationTokenSource StepPollCancelSource = new CancellationTokenSource();
-						using CancellationTokenSource StepAbortSource = new CancellationTokenSource();
-						TaskCompletionSource<bool> StepFinishedSource = new TaskCompletionSource<bool>();
-						Task StepPollTask = Task.Run(() => PollForStepAbort(RpcClient, ExecuteTask.JobId, ExecuteTask.BatchId, Step.StepId, StepPollCancelSource.Token, StepAbortSource, StepFinishedSource.Task));
-						
+					float AvailableFreeSpace = 0;
+					if (DriveName != null)
+					{
 						try
 						{
-							(StepOutcome, StepState) = await ExecuteStepAsync(Executor, Step, ForwardingLogger, CancellationToken, StepAbortSource.Token);
+							DriveInfo Info = new DriveInfo(DriveName);
+							AvailableFreeSpace = (1.0f * Info.AvailableFreeSpace) / 1024 / 1024 / 1024;
 						}
-						finally
+						catch (Exception Ex)
 						{
-							// Will get called even when cancellation token for the lease/batch fires
-							StepFinishedSource.SetResult(true); // Tell background poll task to stop
-							await StepPollTask;
+							BatchLogger.LogWarning(Ex, "Unable to query disk info for path '{DriveName}'", DriveName);
 						}
-						
-						// Kill any processes spawned by the step
-						TerminateProcesses(StepLogger, CancellationToken);
-
-						// Wait for the logger to finish
-						await StepLogger.StopAsync();
-						
-						// Reflect the warnings/errors in the step outcome
-						if (StepOutcome > StepLogger.Outcome)
+					}
+
+					// Print the new state
+					Stopwatch StepTimer = Stopwatch.StartNew();
+
+					BatchLogger.LogInformation("Starting job {JobId}, batch {BatchId}, step {StepId} (Drive Space Left: {DriveSpaceRemaining} GB)", ExecuteTask.JobId, ExecuteTask.BatchId, Step.StepId, AvailableFreeSpace.ToString("F1"));
+
+					// Create a trace span
+					using IScope Scope = GlobalTracer.Instance.BuildSpan("Execute").WithResourceName(Step.Name).StartActive();
+					Scope.Span.SetTag("stepId", Step.StepId);
+					Scope.Span.SetTag("logId", Step.LogId);
+					//				using IDisposable TraceProperty = LogContext.PushProperty("dd.trace_id", CorrelationIdentifier.TraceId.ToString());
+					//				using IDisposable SpanProperty = LogContext.PushProperty("dd.span_id", CorrelationIdentifier.SpanId.ToString());
+
+					// Update the context to include information about this step
+					JobStepOutcome StepOutcome;
+					JobStepState StepState;
+					using (BatchLogger.BeginIndentScope("  "))
+					{
+						// Start writing to the log file
+						await using (JsonRpcLogger StepLogger = new JsonRpcLogger(RpcClient, Step.LogId, ExecuteTask.JobId, ExecuteTask.BatchId, Step.StepId, Step.Warnings, Logger))
 						{
-							StepOutcome = StepLogger.Outcome;
+							// Execute the task
+							ILogger ForwardingLogger = new DefaultLoggerIndentHandler(StepLogger);
+							if (Settings.WriteStepOutputToLogger)
+							{
+								ForwardingLogger = new ForwardingLogger(Logger, ForwardingLogger);
+							}
+
+							using CancellationTokenSource StepPollCancelSource = new CancellationTokenSource();
+							using CancellationTokenSource StepAbortSource = new CancellationTokenSource();
+							TaskCompletionSource<bool> StepFinishedSource = new TaskCompletionSource<bool>();
+							Task StepPollTask = Task.Run(() => PollForStepAbort(RpcClient, ExecuteTask.JobId, ExecuteTask.BatchId, Step.StepId, StepPollCancelSource.Token, StepAbortSource, StepFinishedSource.Task));
+
+							try
+							{
+								(StepOutcome, StepState) = await ExecuteStepAsync(Executor, Step, ForwardingLogger, CancellationToken, StepAbortSource.Token);
+							}
+							finally
+							{
+								// Will get called even when cancellation token for the lease/batch fires
+								StepFinishedSource.SetResult(true); // Tell background poll task to stop
+								await StepPollTask;
+							}
+
+							// Kill any processes spawned by the step
+							TerminateProcesses(StepLogger, CancellationToken);
+
+							// Wait for the logger to finish
+							await StepLogger.StopAsync();
+
+							// Reflect the warnings/errors in the step outcome
+							if (StepOutcome > StepLogger.Outcome)
+							{
+								StepOutcome = StepLogger.Outcome;
+							}
 						}
+
+						// Update the server with the outcome from the step
+						BatchLogger.LogInformation("Marking step as complete (Outcome={Outcome}, State={StepState})", StepOutcome, StepState);
+						await RpcClient.InvokeAsync(x => x.UpdateStepAsync(new UpdateStepRequest(ExecuteTask.JobId, ExecuteTask.BatchId, Step.StepId, StepState, StepOutcome), null, null, CancellationToken), new RpcContext(), CancellationToken);
 					}
 
-					// Update the server with the outcome from the step
-					BatchLogger.LogInformation("Marking step as complete (Outcome={Outcome}, State={StepState})", StepOutcome, StepState);
-					await RpcClient.InvokeAsync(x => x.UpdateStepAsync(new UpdateStepRequest(ExecuteTask.JobId, ExecuteTask.BatchId, Step.StepId, StepState, StepOutcome), null, null, CancellationToken), new RpcContext(), CancellationToken);
+					// Print the finishing state
+					StepTimer.Stop();
+					BatchLogger.LogInformation("Completed in {Time}", StepTimer.Elapsed);
+				}
+			}
+			catch (Exception Ex)
+			{
+				if (CancellationToken.IsCancellationRequested && IsCancellationException(Ex))
+				{
+					Logger.LogError("Step was aborted");
+				}
+				else
+				{
+					Logger.LogError(Ex, "Exception while executing batch: {Ex}", Ex);
 				}
-
-				// Print the finishing state
-				StepTimer.Stop();
-				BatchLogger.LogInformation("Completed in {Time}", StepTimer.Elapsed);
 			}
 
 			// Clean the environment
@@ -1184,7 +1204,7 @@ namespace HordeAgent.Services
 			using (BatchLogger.BeginIndentScope("  "))
 			{
 				using IScope Scope = GlobalTracer.Instance.BuildSpan("Finalize").StartActive();
-				await Executor.FinalizeAsync(BatchLogger, CancellationToken);
+				await Executor.FinalizeAsync(BatchLogger, CancellationToken.None);
 			}
 		}
 
