Horde: Fix incorrect decoding of escape characters in log text.

diff --git a/Engine/Source/Programs/Horde/Horde.Build.Tests/LogFileServiceTests.cs b/Engine/Source/Programs/Horde/Horde.Build.Tests/LogFileServiceTests.cs
--- a/Engine/Source/Programs/Horde/Horde.Build.Tests/LogFileServiceTests.cs
+++ b/Engine/Source/Programs/Horde/Horde.Build.Tests/LogFileServiceTests.cs
@@ -59,6 +59,21 @@ namespace HordeServerTests
 			LogFileService = new LogFileService(LogFileCollection, null!, LogBuilder, LogStorage, Clock, Logger);
         }
 
+		[TestMethod]
+		public void PlainTextDecoder()
+		{
+			string JsonText = "{\"time\":\"2022-01-26T14:18:29\",\"level\":\"Error\",\"message\":\"       \\tdepends on: /mnt/horde/U5\\u002BR5.0\\u002BInc\\u002BMin/Sync/Engine/Binaries/Linux/libUnrealEditor-Core.so\",\"id\":1,\"line\":5,\"lineCount\":1028}";
+			byte[] JsonData = Encoding.UTF8.GetBytes(JsonText);
+
+			byte[] OutputData = new byte[JsonData.Length];
+			int OutputLength = LogText.ConvertToPlainText(JsonData, OutputData, 0);
+			string OutputText = Encoding.UTF8.GetString(OutputData, 0, OutputLength);
+
+			string ExpectText = "       \tdepends on: /mnt/horde/U5+R5.0+Inc+Min/Sync/Engine/Binaries/Linux/libUnrealEditor-Core.so\n";
+			Assert.AreEqual(OutputText, ExpectText);
+			
+		}
+
 		[TestMethod]
         public async Task WriteLogLifecycleOldTest()
         {
diff --git a/Engine/Source/Programs/Horde/Horde.Build/Logs/Data/LogText.cs b/Engine/Source/Programs/Horde/Horde.Build/Logs/Data/LogText.cs
--- a/Engine/Source/Programs/Horde/Horde.Build/Logs/Data/LogText.cs
+++ b/Engine/Source/Programs/Horde/Horde.Build/Logs/Data/LogText.cs
@@ -364,11 +364,11 @@ namespace HordeServer.Logs
 					{
 						(byte)'\"' => (byte)'\"',
 						(byte)'\\' => (byte)'\\',
-						(byte)'\b' => (byte)'\b',
-						(byte)'\f' => (byte)'\f',
-						(byte)'\n' => (byte)'\n',
-						(byte)'\r' => (byte)'\r',
-						(byte)'\t' => (byte)'\t',
+						(byte)'b' => (byte)'\b',
+						(byte)'f' => (byte)'\f',
+						(byte)'n' => (byte)'\n',
+						(byte)'r' => (byte)'\r',
+						(byte)'t' => (byte)'\t',
 						_ => Source[1]
 					};
 					Source = Source.Slice(2);
