Horde: Preparations for auto-scale v2 rollout

- Add a flag for toggling shadow mode (no actual resizing done) and activating v2 itself
- Sizing strategy configured per pool
- Made AutoscaleServiceV2 more test friendly
- Added a feature flag section in settings

diff --git a/Engine/Source/Programs/Horde/Horde.Build.Tests/Fleet/AutoscaleServiceV2Test.cs b/Engine/Source/Programs/Horde/Horde.Build.Tests/Fleet/AutoscaleServiceV2Test.cs
new file mode 100644
--- /dev/null
+++ b/Engine/Source/Programs/Horde/Horde.Build.Tests/Fleet/AutoscaleServiceV2Test.cs
@@ -0,0 +1,152 @@
+// Copyright Epic Games, Inc. All Rights Reserved.
+
+using System;
+using System.Collections.Generic;
+using System.Threading;
+using System.Threading.Tasks;
+using Horde.Build.Fleet.Autoscale;
+using HordeServer;
+using HordeServer.Models;
+using HordeServer.Services;
+using HordeServer.Utilities;
+using HordeServerTests;
+using Microsoft.Extensions.DependencyInjection;
+using Microsoft.Extensions.Logging;
+using Microsoft.Extensions.Options;
+using Microsoft.VisualStudio.TestTools.UnitTesting;
+
+namespace Horde.Build.Tests.Fleet
+{
+	using PoolId = StringId<IPool>;
+	using ProjectId = StringId<IProject>;
+	using StreamId = StringId<IStream>;
+	
+	public class FleetManagerSpy : IFleetManager
+	{
+		public int ExpandPoolAsyncCallCount { get; private set; }
+		public int ShrinkPoolAsyncCallCount { get; private set; }
+		
+		public Task ExpandPoolAsync(IPool Pool, IReadOnlyList<IAgent> Agents, int Count)
+		{
+			ExpandPoolAsyncCallCount++;
+			return Task.CompletedTask;
+		}
+
+		public Task ShrinkPoolAsync(IPool Pool, IReadOnlyList<IAgent> Agents, int Count)
+		{
+			ShrinkPoolAsyncCallCount++;
+			return Task.CompletedTask;
+		}
+
+		public Task<int> GetNumStoppedInstancesAsync(IPool Pool)
+		{
+			throw new NotImplementedException();
+		}
+	}
+
+	public class PoolSizeStrategySpy : IPoolSizeStrategy
+	{
+		public int CallCount { get; private set; }
+		public HashSet<PoolId> PoolIdsSeen { get; } = new();
+		
+		public Task<List<PoolSizeData>> CalcDesiredPoolSizesAsync(List<PoolSizeData> Pools)
+		{
+			CallCount++;
+			foreach (PoolSizeData Data in Pools)
+			{
+				PoolIdsSeen.Add(Data.Pool.Id);
+			}
+			return Task.FromResult(Pools);
+		}
+
+		public string Name { get; } = "PoolSizeStrategySpy";
+	}
+
+	[TestClass]
+	public class AutoscaleServiceV2Test : TestSetup
+	{
+		FleetManagerSpy FleetManagerSpy = new();
+		
+		[TestMethod]
+		public async Task ShadowModeEnabled()
+		{
+			ServerSettings.FeatureFlags.AutoscaleServiceV2ShadowMode = false;
+ 
+			var Service = GetAutoscaleService(FleetManagerSpy);
+			var PoolSizeDatas = await GetPoolSizeData();
+			PoolSizeDatas[0] = PoolSizeDatas[0].Copy(DesiredAgentCount: 8);
+
+			await Service.ResizePools(PoolSizeDatas);
+			Assert.AreEqual(1, FleetManagerSpy.ExpandPoolAsyncCallCount);
+			Assert.AreEqual(0, FleetManagerSpy.ShrinkPoolAsyncCallCount);
+		}
+		
+		[TestMethod]
+		public async Task ShadowModeDisabled()
+		{
+			ServerSettings.FeatureFlags.AutoscaleServiceV2ShadowMode = true;
+ 
+			var Service = GetAutoscaleService(FleetManagerSpy);
+			var PoolSizeDatas = await GetPoolSizeData();
+			PoolSizeDatas[0] = PoolSizeDatas[0].Copy(DesiredAgentCount: 8);
+
+			await Service.ResizePools(PoolSizeDatas);
+			Assert.AreEqual(0, FleetManagerSpy.ExpandPoolAsyncCallCount);
+			Assert.AreEqual(0, FleetManagerSpy.ShrinkPoolAsyncCallCount);
+		}
+		
+		[TestMethod]
+		public async Task PerPoolStrategy()
+		{
+			var Service = GetAutoscaleService(FleetManagerSpy);
+			
+			IPool Pool1 = await PoolService.CreatePoolAsync("bogusPoolLease1", null, true, 0, 0, PoolSizeStrategy.LeaseUtilization);
+			IPool Pool2 = await PoolService.CreatePoolAsync("bogusPoolLease2", null, true, 0, 0, null);
+			IPool Pool3 = await PoolService.CreatePoolAsync("bogusPoolJobQueue1", null, true, 0, 0, PoolSizeStrategy.JobQueue);
+			IPool Pool4 = await PoolService.CreatePoolAsync("bogusPoolJobQueue2", null, true, 0, 0, PoolSizeStrategy.JobQueue);
+			IPool Pool5 = await PoolService.CreatePoolAsync("bogusPoolNoOp", null, true, 0, 0, PoolSizeStrategy.NoOp);
+
+			PoolSizeStrategySpy LeaseUtilizationSpy = new();
+			PoolSizeStrategySpy JobQueueSpy = new();
+			PoolSizeStrategySpy NoOpSpy = new();
+			Service.OverridePoolSizeStrategiesDuringTesting(LeaseUtilizationSpy, JobQueueSpy, NoOpSpy);
+
+			await Service.TickLeaderAsync(CancellationToken.None);
+			
+			Assert.AreEqual(1, LeaseUtilizationSpy.CallCount);
+			Assert.IsTrue(LeaseUtilizationSpy.PoolIdsSeen.Contains(Pool1.Id));
+			Assert.IsTrue(LeaseUtilizationSpy.PoolIdsSeen.Contains(Pool2.Id));
+			
+			Assert.AreEqual(1, JobQueueSpy.CallCount);
+			Assert.IsTrue(JobQueueSpy.PoolIdsSeen.Contains(Pool3.Id));
+			Assert.IsTrue(JobQueueSpy.PoolIdsSeen.Contains(Pool4.Id));
+			
+			Assert.AreEqual(1, NoOpSpy.CallCount);
+			Assert.IsTrue(NoOpSpy.PoolIdsSeen.Contains(Pool5.Id));
+		}
+
+		private async Task<List<PoolSizeData>> GetPoolSizeData()
+		{
+			IPool Pool1 = await PoolService.CreatePoolAsync("bogusPool1", null, true, 0, 0);
+			IAgent Agent1 = await CreateAgentAsync(Pool1);
+			IAgent Agent2 = await CreateAgentAsync(Pool1);
+
+			return new List<PoolSizeData>
+			{
+				new (Pool1, new List<IAgent> { Agent1, Agent2 }, null)
+			};
+		}
+
+		private AutoscaleServiceV2 GetAutoscaleService(IFleetManager FleetManager)
+		{
+			ILogger<AutoscaleServiceV2> Logger = ServiceProvider.GetRequiredService<ILogger<AutoscaleServiceV2>>();
+			IOptions<ServerSettings> ServerSettingsOpt = ServiceProvider.GetRequiredService<IOptions<ServerSettings>>();
+
+			LeaseUtilizationStrategy LeaseUtilizationStrategy = new(AgentCollection, PoolCollection, LeaseCollection, Clock);
+			JobQueueStrategy JobQueueStrategy = new(JobCollection, GraphCollection, StreamService, Clock);
+			AutoscaleServiceV2 Service = new AutoscaleServiceV2(LeaseUtilizationStrategy,JobQueueStrategy, new NoOpPoolSizeStrategy(), AgentCollection, PoolCollection, FleetManager, new NoOpDogStatsd(), Clock, ServerSettingsOpt, Logger);
+			return Service;
+		}
+
+	}
+}
\ No newline at end of file
diff --git a/Engine/Source/Programs/Horde/Horde.Build.Tests/Fleet/JobQueueStrategyTest.cs b/Engine/Source/Programs/Horde/Horde.Build.Tests/Fleet/JobQueueStrategyTest.cs
--- a/Engine/Source/Programs/Horde/Horde.Build.Tests/Fleet/JobQueueStrategyTest.cs
+++ b/Engine/Source/Programs/Horde/Horde.Build.Tests/Fleet/JobQueueStrategyTest.cs
@@ -33,36 +33,27 @@ namespace Horde.Build.Tests.Fleet
 		public async Task JobQueueWith5Jobs()
 		{
 			(JobQueueStrategy Strategy, PoolSizeData PoolSizeData) = await SetUpJobsAsync(1, 5);
-			Dictionary<PoolId, PoolSizeData> PoolSizeDatas = new ();
-			PoolSizeDatas[PoolSizeData.Pool.Id] = PoolSizeData;
-			await Strategy.CalcDesiredPoolSizesAsync(PoolSizeDatas);
-			PoolSizeData = PoolSizeDatas[PoolSizeData.Pool.Id];
-			Assert.AreEqual(1, PoolSizeDatas.Count);
-			Assert.AreEqual(PoolSizeData.Agents.Count + 0, PoolSizeData.DesiredAgentCount);
+			List<PoolSizeData> Result = await Strategy.CalcDesiredPoolSizesAsync(new() { PoolSizeData });
+			Assert.AreEqual(1, Result.Count);
+			Assert.AreEqual(Result[0].Agents.Count + 0, Result[0].DesiredAgentCount);
 		}
 		
 		[TestMethod]
 		public async Task JobQueueWith15Jobs()
 		{
 			(JobQueueStrategy Strategy, PoolSizeData PoolSizeData) = await SetUpJobsAsync(1, 15);
-			Dictionary<PoolId, PoolSizeData> PoolSizeDatas = new ();
-			PoolSizeDatas[PoolSizeData.Pool.Id] = PoolSizeData;
-			await Strategy.CalcDesiredPoolSizesAsync(PoolSizeDatas);
-			PoolSizeData = PoolSizeDatas[PoolSizeData.Pool.Id];
-			Assert.AreEqual(1, PoolSizeDatas.Count);
-			Assert.AreEqual(PoolSizeData.Agents.Count + 1, PoolSizeData.DesiredAgentCount);
+			List<PoolSizeData> Result = await Strategy.CalcDesiredPoolSizesAsync(new() { PoolSizeData });
+			Assert.AreEqual(1, Result.Count);
+			Assert.AreEqual(Result[0].Agents.Count + 1, Result[0].DesiredAgentCount);
 		}
 		
 		[TestMethod]
 		public async Task JobQueueWith25Jobs()
 		{
 			(JobQueueStrategy Strategy, PoolSizeData PoolSizeData) = await SetUpJobsAsync(1, 25);
-			Dictionary<PoolId, PoolSizeData> PoolSizeDatas = new ();
-			PoolSizeDatas[PoolSizeData.Pool.Id] = PoolSizeData;
-			await Strategy.CalcDesiredPoolSizesAsync(PoolSizeDatas);
-			PoolSizeData = PoolSizeDatas[PoolSizeData.Pool.Id];
-			Assert.AreEqual(1, PoolSizeDatas.Count);
-			Assert.AreEqual(PoolSizeData.Agents.Count + 2, PoolSizeData.DesiredAgentCount);
+			List<PoolSizeData> Result = await Strategy.CalcDesiredPoolSizesAsync(new() { PoolSizeData });
+			Assert.AreEqual(1, Result.Count);
+			Assert.AreEqual(Result[0].Agents.Count + 2, Result[0].DesiredAgentCount);
 		}
 	
 		/// <summary>
diff --git a/Engine/Source/Programs/Horde/Horde.Build.Tests/Fleet/LeaseUtilizationStrategyTest.cs b/Engine/Source/Programs/Horde/Horde.Build.Tests/Fleet/LeaseUtilizationStrategyTest.cs
--- a/Engine/Source/Programs/Horde/Horde.Build.Tests/Fleet/LeaseUtilizationStrategyTest.cs
+++ b/Engine/Source/Programs/Horde/Horde.Build.Tests/Fleet/LeaseUtilizationStrategyTest.cs
@@ -97,10 +97,9 @@ namespace Horde.Build.Tests.Fleet
 		private async Task AssertPoolSizeAsync(IPool Pool, int ExpectedNumAgents)
 		{
 			LeaseUtilizationStrategy Strategy = new (AgentCollection, PoolCollection, LeaseCollection, Clock);
-			Dictionary<PoolId, PoolSizeData> PoolSizeDatas = new ();
-			PoolSizeDatas[Pool.Id] = new(Pool, PoolAgents, null);
-			await Strategy.CalcDesiredPoolSizesAsync(PoolSizeDatas);
-			Assert.AreEqual(ExpectedNumAgents, PoolSizeDatas[Pool.Id].DesiredAgentCount);
+			List<PoolSizeData> Output = await Strategy.CalcDesiredPoolSizesAsync(new() { new(Pool, PoolAgents, null) });
+			Assert.AreEqual(1, Output.Count);
+			Assert.AreEqual(ExpectedNumAgents, Output[0].DesiredAgentCount);
 		}
 		
 		private async Task<ILease> AddPlaceholderLease(IAgent Agent, IPool Pool, DateTime StartTime, TimeSpan Duration)
diff --git a/Engine/Source/Programs/Horde/Horde.Build/Api/Pools.cs b/Engine/Source/Programs/Horde/Horde.Build/Api/Pools.cs
--- a/Engine/Source/Programs/Horde/Horde.Build/Api/Pools.cs
+++ b/Engine/Source/Programs/Horde/Horde.Build/Api/Pools.cs
@@ -7,6 +7,7 @@ using System.Collections.Generic;
 using System.ComponentModel.DataAnnotations;
 using System.Linq;
 using System.Threading.Tasks;
+using Horde.Build.Fleet.Autoscale;
 
 namespace HordeServer.Api
 {
@@ -30,6 +31,11 @@ namespace HordeServer.Api
 		/// Whether to enable autoscaling for this pool
 		/// </summary>
 		public bool? EnableAutoscaling { get; set; }
+		
+		/// <summary>
+		/// Pool sizing strategy
+		/// </summary>
+		public PoolSizeStrategy? SizeStrategy { get; set; }
 
 		/// <summary>
 		/// The minimum nunmber of agents to retain in this pool
@@ -86,6 +92,11 @@ namespace HordeServer.Api
 		/// Whether to enable autoscaling for this pool
 		/// </summary>
 		public bool? EnableAutoscaling { get; set; }
+		
+		/// <summary>
+		/// Pool sizing strategy
+		/// </summary>
+		public PoolSizeStrategy? SizeStrategy { get; set; }
 
 		/// <summary>
 		/// The minimum nunmber of agents to retain in this pool
diff --git a/Engine/Source/Programs/Horde/Horde.Build/Controllers/PoolsController.cs b/Engine/Source/Programs/Horde/Horde.Build/Controllers/PoolsController.cs
--- a/Engine/Source/Programs/Horde/Horde.Build/Controllers/PoolsController.cs
+++ b/Engine/Source/Programs/Horde/Horde.Build/Controllers/PoolsController.cs
@@ -60,7 +60,7 @@ namespace HordeServer.Controllers
 				return Forbid(AclAction.CreatePool);
 			}
 
-			IPool NewPool = await PoolService.CreatePoolAsync(Create.Name, Create.Condition, Create.EnableAutoscaling, Create.MinAgents, Create.NumReserveAgents, Create.Properties);
+			IPool NewPool = await PoolService.CreatePoolAsync(Create.Name, Create.Condition, Create.EnableAutoscaling, Create.MinAgents, Create.NumReserveAgents, Create.SizeStrategy, Create.Properties);
 			return new CreatePoolResponse(NewPool.Id.ToString());
 		}
 
diff --git a/Engine/Source/Programs/Horde/Horde.Build/Fleet/Autoscale/AutoscaleServiceV2.cs b/Engine/Source/Programs/Horde/Horde.Build/Fleet/Autoscale/AutoscaleServiceV2.cs
--- a/Engine/Source/Programs/Horde/Horde.Build/Fleet/Autoscale/AutoscaleServiceV2.cs
+++ b/Engine/Source/Programs/Horde/Horde.Build/Fleet/Autoscale/AutoscaleServiceV2.cs
@@ -7,12 +7,14 @@ using System.Linq;
 using System.Threading;
 using System.Threading.Tasks;
 using HordeCommon;
+using HordeServer;
 using HordeServer.Collections;
 using HordeServer.Models;
 using HordeServer.Services;
 using HordeServer.Utilities;
 using Microsoft.Extensions.Hosting;
 using Microsoft.Extensions.Logging;
+using Microsoft.Extensions.Options;
 using OpenTracing;
 using OpenTracing.Util;
 using StatsdClient;
@@ -28,27 +30,33 @@ namespace Horde.Build.Fleet.Autoscale
 	{
 		static readonly TimeSpan ShrinkPoolCoolDown = TimeSpan.FromMinutes(20.0);
 
-		IPoolSizeStrategy PoolSizeStrategy;
+		IPoolSizeStrategy LeaseUtilizationStrategy;
+		IPoolSizeStrategy JobQueueStrategy;
+		IPoolSizeStrategy NoOpStrategy;
 		IAgentCollection AgentCollection;
 		IPoolCollection PoolCollection;
 		IFleetManager FleetManager;
 		IDogStatsd DogStatsd;
 		IClock Clock;
 		ITicker Ticker;
+		ServerSettings Settings;
 		ILogger<AutoscaleServiceV2> Logger;
 
 		/// <summary>
 		/// Constructor
 		/// </summary>
-		public AutoscaleServiceV2(IPoolSizeStrategy PoolSizeStrategy, IAgentCollection AgentCollection, IPoolCollection PoolCollection, IFleetManager FleetManager, IDogStatsd DogStatsd, IClock Clock, ILogger<AutoscaleServiceV2> Logger)
+		public AutoscaleServiceV2(LeaseUtilizationStrategy LeaseUtilizationStrategy, JobQueueStrategy JobQueueStrategy, NoOpPoolSizeStrategy NoOpStrategy, IAgentCollection AgentCollection, IPoolCollection PoolCollection, IFleetManager FleetManager, IDogStatsd DogStatsd, IClock Clock, IOptions<ServerSettings> Settings, ILogger<AutoscaleServiceV2> Logger)
 		{
-			this.PoolSizeStrategy = PoolSizeStrategy;
+			this.LeaseUtilizationStrategy = LeaseUtilizationStrategy;
+			this.JobQueueStrategy = JobQueueStrategy;
+			this.NoOpStrategy = NoOpStrategy;
 			this.AgentCollection = AgentCollection;
 			this.PoolCollection = PoolCollection;
 			this.FleetManager = FleetManager;
 			this.DogStatsd = DogStatsd;
 			this.Clock = Clock;
 			this.Ticker = Clock.AddSharedTicker<AutoscaleServiceV2>(TimeSpan.FromMinutes(5.0), TickLeaderAsync, Logger);
+			this.Settings = Settings.Value;
 			this.Logger = Logger;
 		}
 
@@ -61,27 +69,45 @@ namespace Horde.Build.Fleet.Autoscale
 		/// <inheritdoc/>
 		public void Dispose() => Ticker.Dispose();
 
-		/// <inheritdoc/>
-		async ValueTask TickLeaderAsync(CancellationToken StoppingToken)
+		internal async ValueTask TickLeaderAsync(CancellationToken StoppingToken)
 		{
-			DateTime UtcNow = Clock.UtcNow;
-
-			Logger.LogInformation("Autoscaling pools using strategy {PoolSizeStrategyName}...", PoolSizeStrategy.Name);
+			Logger.LogInformation("Autoscaling pools...");
 			Stopwatch Stopwatch = Stopwatch.StartNew();
 			using IScope _ = GlobalTracer.Instance.BuildSpan("AutoscaleService").StartActive();
 			
-			List<IAgent> Agents = await AgentCollection.FindAsync(Status: AgentStatus.Ok);
-			List<IAgent> GetAgentsInPool(PoolId PoolId)
+			// Group pools by strategy type to ensure they can be called once (more optimal)
+			Dictionary<PoolSizeStrategy, List<PoolSizeData>> PoolSizeDataByStrategy = await GetPoolSizeDataByStrategyType();
+			foreach ((var StrategyType, var CurrentData) in PoolSizeDataByStrategy)
 			{
-				return Agents.FindAll(a => a.GetPools().Any(p => p == PoolId));
+				List<PoolSizeData> NewData = await GetPoolSizeStrategy(StrategyType).CalcDesiredPoolSizesAsync(CurrentData);
+				await ResizePools(NewData);
 			}
 			
+			Stopwatch.Stop();
+			Logger.LogInformation("Autoscaling pools took {ElapsedTime} ms", Stopwatch.ElapsedMilliseconds);
+		}
+
+		internal async Task<Dictionary<PoolSizeStrategy, List<PoolSizeData>>> GetPoolSizeDataByStrategyType()
+		{
+			List<IAgent> Agents = await AgentCollection.FindAsync(Status: AgentStatus.Ok);
+			List<IAgent> GetAgentsInPool(PoolId PoolId) => Agents.FindAll(a => a.GetPools().Any(p => p == PoolId));
 			List<IPool> Pools = await PoolCollection.GetAsync();
-			Dictionary<PoolId, PoolSizeData> PoolSizeDatas = Pools.ToDictionary(x => x.Id, x => new PoolSizeData(x, GetAgentsInPool(x.Id), null));
 
-			await PoolSizeStrategy.CalcDesiredPoolSizesAsync(PoolSizeDatas);
+			Dictionary<PoolSizeStrategy, List<PoolSizeData>> Result = new();
+			foreach (PoolSizeStrategy StrategyType in Enum.GetValues<PoolSizeStrategy>())
+			{
+				Result[StrategyType] = Pools
+					.Where(x => (x.SizeStrategy ?? Settings.DefaultAgentPoolSizeStrategy) == StrategyType)
+					.Select(x => new PoolSizeData(x, GetAgentsInPool(x.Id), null))
+					.ToList();
+			}
 
-			foreach (PoolSizeData PoolSizeData in PoolSizeDatas.Values.OrderByDescending(x => x.Agents.Count))
+			return Result;
+		}
+
+		internal async Task ResizePools(List<PoolSizeData> PoolSizeDatas)
+		{
+			foreach (PoolSizeData PoolSizeData in PoolSizeDatas.OrderByDescending(x => x.Agents.Count))
 			{
 				IPool Pool = PoolSizeData.Pool;
 
@@ -92,6 +118,11 @@ namespace Horde.Build.Fleet.Autoscale
 				int DeltaAgentCount = DesiredAgentCount - CurrentAgentCount;
 
 				Logger.LogInformation("{PoolName,-48} Current={Current,4} Target={Target,4} Delta={Delta,4} Status={Status}", Pool.Name, CurrentAgentCount, DesiredAgentCount, DeltaAgentCount, PoolSizeData.StatusMessage);
+
+				if (Settings.FeatureFlags.AutoscaleServiceV2ShadowMode)
+				{
+					continue;
+				}
 				
 				try
 				{
@@ -130,9 +161,31 @@ namespace Horde.Build.Fleet.Autoscale
 				DogStatsd.Gauge("agentpools.autoscale.target", DesiredAgentCount, tags: new []{"pool:" + Pool.Name});
 				DogStatsd.Gauge("agentpools.autoscale.current", CurrentAgentCount, tags: new []{"pool:" + Pool.Name});
 			}
-			
-			Stopwatch.Stop();
-			Logger.LogInformation("Autoscaling pools took {ElapsedTime} ms", Stopwatch.ElapsedMilliseconds);
+		}
+
+		/// <summary>
+		/// Backdoor for tests to override strategies with test doubles
+		/// These cannot supplied in the constructor since the DI requires concrete implementations.
+		/// </summary>
+		/// <param name="LeaseUtilizationStrategy"></param>
+		/// <param name="JobQueueStrategy"></param>
+		/// <param name="NoOpStrategy"></param>
+		internal void OverridePoolSizeStrategiesDuringTesting(IPoolSizeStrategy LeaseUtilizationStrategy, IPoolSizeStrategy JobQueueStrategy, IPoolSizeStrategy NoOpStrategy)
+		{
+			this.LeaseUtilizationStrategy = LeaseUtilizationStrategy;
+			this.JobQueueStrategy = JobQueueStrategy;
+			this.NoOpStrategy = NoOpStrategy;
+		}
+		
+		private IPoolSizeStrategy GetPoolSizeStrategy(PoolSizeStrategy Type)
+		{
+			return Type switch
+			{
+				PoolSizeStrategy.LeaseUtilization => LeaseUtilizationStrategy,
+				PoolSizeStrategy.JobQueue => JobQueueStrategy,
+				PoolSizeStrategy.NoOp => NoOpStrategy,
+				_ => NoOpStrategy,
+			};
 		}
 	}
 }
diff --git a/Engine/Source/Programs/Horde/Horde.Build/Fleet/Autoscale/IPoolSizeStrategy.cs b/Engine/Source/Programs/Horde/Horde.Build/Fleet/Autoscale/IPoolSizeStrategy.cs
--- a/Engine/Source/Programs/Horde/Horde.Build/Fleet/Autoscale/IPoolSizeStrategy.cs
+++ b/Engine/Source/Programs/Horde/Horde.Build/Fleet/Autoscale/IPoolSizeStrategy.cs
@@ -1,6 +1,7 @@
 // Copyright Epic Games, Inc. All Rights Reserved.
 
 using System.Collections.Generic;
+using System.Linq;
 using System.Threading.Tasks;
 using HordeServer.Models;
 using HordeServer.Utilities;
@@ -72,6 +73,20 @@ namespace Horde.Build.Fleet.Autoscale
 			this.DesiredAgentCount = DesiredAgentCount;
 			this.StatusMessage = StatusMessage;
 		}
+
+		/// <summary>
+		/// Copy the object, inheriting any unspecified values from current instance
+		/// Needed because the class is immutable 
+		/// </summary>
+		/// <param name="Pool"></param>
+		/// <param name="Agents"></param>
+		/// <param name="DesiredAgentCount"></param>
+		/// <param name="StatusMessage"></param>
+		/// <returns>A new copy</returns>
+		public PoolSizeData Copy(IPool? Pool = null, List<IAgent>? Agents = null, int? DesiredAgentCount = null, string? StatusMessage = null)
+		{
+			return new PoolSizeData(Pool ?? this.Pool, Agents ?? this.Agents, DesiredAgentCount ?? this.DesiredAgentCount, StatusMessage ?? this.StatusMessage);
+		}
 	}
 	
 	/// <summary>
@@ -84,7 +99,7 @@ namespace Horde.Build.Fleet.Autoscale
 		/// </summary>
 		/// <param name="Pools">Pools including attached agents</param>
 		/// <returns></returns>
-		Task CalcDesiredPoolSizesAsync(Dictionary<PoolId, PoolSizeData> Pools);
+		Task<List<PoolSizeData>> CalcDesiredPoolSizesAsync(List<PoolSizeData> Pools);
 		
 		/// <summary>
 		/// Name of the strategy
@@ -99,14 +114,10 @@ namespace Horde.Build.Fleet.Autoscale
 	public class NoOpPoolSizeStrategy : IPoolSizeStrategy
 	{
 		/// <inheritdoc/>
-		public Task CalcDesiredPoolSizesAsync(Dictionary<PoolId, PoolSizeData> Pools)
+		public Task<List<PoolSizeData>> CalcDesiredPoolSizesAsync(List<PoolSizeData> Pools)
 		{
-			foreach (PoolSizeData Data in Pools.Values)
-			{
-				Pools[Data.Pool.Id] = new(Data.Pool, Data.Agents, Data.Agents.Count, "(no-op)");
-			}
-
-			return Task.CompletedTask;
+			List<PoolSizeData> Result = Pools.Select(x => new PoolSizeData(x.Pool, x.Agents, x.Agents.Count, "(no-op)")).ToList();
+			return Task.FromResult(Result);
 		}
 
 		/// <inheritdoc/>
diff --git a/Engine/Source/Programs/Horde/Horde.Build/Fleet/Autoscale/JobQueueStrategy.cs b/Engine/Source/Programs/Horde/Horde.Build/Fleet/Autoscale/JobQueueStrategy.cs
--- a/Engine/Source/Programs/Horde/Horde.Build/Fleet/Autoscale/JobQueueStrategy.cs
+++ b/Engine/Source/Programs/Horde/Horde.Build/Fleet/Autoscale/JobQueueStrategy.cs
@@ -100,14 +100,16 @@ namespace Horde.Build.Fleet.Autoscale
 		}
 
 		/// <inheritdoc/>
-		public async Task CalcDesiredPoolSizesAsync(Dictionary<PoolId, PoolSizeData> PoolSizes)
+		public async Task<List<PoolSizeData>> CalcDesiredPoolSizesAsync(List<PoolSizeData> Pools)
 		{
 			DateTimeOffset MinCreateTime = Clock.UtcNow - LookbackPeriod;
 			Dictionary<PoolId, int> PoolQueueSizes = await GetPoolQueueSizesAsync(MinCreateTime);
+			List<PoolSizeData> Result = new();
 			
 			foreach ((PoolId PoolId, int QueueSize) in PoolQueueSizes.OrderBy(x => x.Value))
 			{
-				if (PoolSizes.TryGetValue(PoolId, out PoolSizeData? PoolSize))
+				PoolSizeData? PoolSize = Pools.Find(x => x.Pool.Id == PoolId);
+				if (PoolSize != null)
 				{
 					// TODO: Tweak these values once in production
 					int AdditionalAgentCount = QueueSize switch
@@ -122,9 +124,11 @@ namespace Horde.Build.Fleet.Autoscale
 					};
 
 					int DesiredAgentCount = PoolSize.Agents.Count + AdditionalAgentCount;
-					PoolSizes[PoolId] = new(PoolSize.Pool, PoolSize.Agents, DesiredAgentCount, $"QueueSize={QueueSize}");
+					Result.Add(new(PoolSize.Pool, PoolSize.Agents, DesiredAgentCount, $"QueueSize={QueueSize}"));
 				}
 			}
+
+			return Result;
 		}
 	}
 }
diff --git a/Engine/Source/Programs/Horde/Horde.Build/Fleet/Autoscale/LeaseUtilizationStrategy.cs b/Engine/Source/Programs/Horde/Horde.Build/Fleet/Autoscale/LeaseUtilizationStrategy.cs
--- a/Engine/Source/Programs/Horde/Horde.Build/Fleet/Autoscale/LeaseUtilizationStrategy.cs
+++ b/Engine/Source/Programs/Horde/Horde.Build/Fleet/Autoscale/LeaseUtilizationStrategy.cs
@@ -186,15 +186,17 @@ namespace Horde.Build.Fleet.Autoscale
 		public string Name { get; } = "LeaseUtilization";
 
 		/// <inheritdoc/>
-		public async Task CalcDesiredPoolSizesAsync(Dictionary<PoolId, PoolSizeData> PoolSizes)
+		public async Task<List<PoolSizeData>> CalcDesiredPoolSizesAsync(List<PoolSizeData> Pools)
 		{
 			Dictionary<PoolId, PoolData> PoolToData = await GetPoolDataAsync();
+			List<PoolSizeData> Result = new();
 
 			foreach (PoolData PoolData in PoolToData.Values.OrderByDescending(x => x.Agents.Count))
 			{
 				IPool Pool = PoolData.Pool;
 
-				if (PoolSizes.TryGetValue(Pool.Id, out PoolSizeData? PoolSize))
+				PoolSizeData? PoolSize = Pools.Find(x => x.Pool.Id == Pool.Id);
+				if (PoolSize != null)
 				{
 					int MinAgents = Pool.MinAgents ?? 1;
 					int NumReserveAgents = Pool.NumReserveAgents ?? 5;
@@ -214,9 +216,11 @@ namespace Horde.Build.Fleet.Autoscale
 					Sb.AppendFormat("Avg=[{0,5:0.0}] ", Utilization);
 					Sb.AppendFormat("Pct=[{0,5:0.0}] ", PoolData.Samples.Sum(x => x.JobWork) / NumSamples);
 
-					PoolSizes[Pool.Id] = new(Pool, PoolSize.Agents, DesiredAgentCount, Sb.ToString());
+					Result.Add(new(Pool, PoolSize.Agents, DesiredAgentCount, Sb.ToString()));
 				}
 			}
+
+			return Result;
 		}
 		
 		/// <summary>
diff --git a/Engine/Source/Programs/Horde/Horde.Build/Fleet/IPool.cs b/Engine/Source/Programs/Horde/Horde.Build/Fleet/IPool.cs
--- a/Engine/Source/Programs/Horde/Horde.Build/Fleet/IPool.cs
+++ b/Engine/Source/Programs/Horde/Horde.Build/Fleet/IPool.cs
@@ -10,6 +10,7 @@ using System.Collections.Generic;
 using System.Linq;
 using System.Security.Permissions;
 using System.Threading.Tasks;
+using Horde.Build.Fleet.Autoscale;
 
 namespace HordeServer.Models
 {
@@ -74,6 +75,11 @@ namespace HordeServer.Models
 		/// Last time the pool was (auto) scaled down
 		/// </summary>
 		public DateTime? LastScaleDownTime { get; }
+		
+		/// <summary>
+		/// Pool sizing strategy to be used for this pool
+		/// </summary>
+		public PoolSizeStrategy? SizeStrategy { get; }
 
 		/// <summary>
 		/// Update index for this document
diff --git a/Engine/Source/Programs/Horde/Horde.Build/Fleet/IPoolCollection.cs b/Engine/Source/Programs/Horde/Horde.Build/Fleet/IPoolCollection.cs
--- a/Engine/Source/Programs/Horde/Horde.Build/Fleet/IPoolCollection.cs
+++ b/Engine/Source/Programs/Horde/Horde.Build/Fleet/IPoolCollection.cs
@@ -13,6 +13,7 @@ using System.Collections.Generic;
 using System.Linq;
 using System.Linq.Expressions;
 using System.Threading.Tasks;
+using Horde.Build.Fleet.Autoscale;
 
 namespace HordeServer.Collections
 {
@@ -32,9 +33,10 @@ namespace HordeServer.Collections
 		/// <param name="EnableAutoscaling">Whether to enable autoscaling for this pool</param>
 		/// <param name="MinAgents">Minimum number of agents in the pool</param>
 		/// <param name="NumReserveAgents">Minimum number of idle agents to maintain</param>
+		/// <param name="SizeStrategy">Pool sizing strategy</param>
 		/// <param name="Properties">Properties for the pool</param>
 		/// <returns>The new pool document</returns>
-		Task<IPool> AddAsync(PoolId Id, string Name, Condition? Condition = null, bool? EnableAutoscaling = null, int? MinAgents = null, int? NumReserveAgents = null, IEnumerable<KeyValuePair<string, string>>? Properties = null);
+		Task<IPool> AddAsync(PoolId Id, string Name, Condition? Condition = null, bool? EnableAutoscaling = null, int? MinAgents = null, int? NumReserveAgents = null, PoolSizeStrategy? SizeStrategy = null, IEnumerable<KeyValuePair<string, string>>? Properties = null);
 
 		/// <summary>
 		/// Enumerates all the pools
@@ -76,7 +78,8 @@ namespace HordeServer.Collections
 		/// <param name="NewProperties">New properties for the pool</param>
 		/// <param name="LastScaleUpTime">New time for last (auto) scale up</param>
 		/// <param name="LastScaleDownTime">New time for last (auto) scale down</param>
+		/// <param name="SizeStrategy">Pool sizing strategy</param>
 		/// <returns>Async task</returns>
-		Task<IPool?> TryUpdateAsync(IPool Pool, string? NewName = null, Condition? NewCondition = null, bool? NewEnableAutoscaling = null, int? NewMinAgents = null, int? NewNumReserveAgents = null, List<AgentWorkspace>? NewWorkspaces = null, bool? NewUseAutoSdk = null, Dictionary<string, string?>? NewProperties = null, DateTime? LastScaleUpTime = null, DateTime? LastScaleDownTime = null);
+		Task<IPool?> TryUpdateAsync(IPool Pool, string? NewName = null, Condition? NewCondition = null, bool? NewEnableAutoscaling = null, int? NewMinAgents = null, int? NewNumReserveAgents = null, List<AgentWorkspace>? NewWorkspaces = null, bool? NewUseAutoSdk = null, Dictionary<string, string?>? NewProperties = null, DateTime? LastScaleUpTime = null, DateTime? LastScaleDownTime = null, PoolSizeStrategy? SizeStrategy = null);
 	}
 }
diff --git a/Engine/Source/Programs/Horde/Horde.Build/Fleet/Impl/PoolCollection.cs b/Engine/Source/Programs/Horde/Horde.Build/Fleet/Impl/PoolCollection.cs
--- a/Engine/Source/Programs/Horde/Horde.Build/Fleet/Impl/PoolCollection.cs
+++ b/Engine/Source/Programs/Horde/Horde.Build/Fleet/Impl/PoolCollection.cs
@@ -13,6 +13,7 @@ using System.Collections.Generic;
 using System.Linq;
 using System.Linq.Expressions;
 using System.Threading.Tasks;
+using Horde.Build.Fleet.Autoscale;
 
 namespace HordeServer.Collections.Impl
 {
@@ -40,6 +41,7 @@ namespace HordeServer.Collections.Impl
 			public int? NumReserveAgents { get; set; }
 			public DateTime? LastScaleUpTime { get; set; }
 			public DateTime? LastScaleDownTime { get; set; }
+			public PoolSizeStrategy? SizeStrategy { get; set; }
 			public int UpdateIndex { get; set; }
 
 			// Read-only wrappers
@@ -63,6 +65,7 @@ namespace HordeServer.Collections.Impl
 				NumReserveAgents = Other.NumReserveAgents;
 				LastScaleUpTime = Other.LastScaleUpTime;
 				LastScaleDownTime = Other.LastScaleDownTime;
+				SizeStrategy = Other.SizeStrategy;
 				UpdateIndex = Other.UpdateIndex;
 			}
 		}
@@ -82,7 +85,7 @@ namespace HordeServer.Collections.Impl
 		}
 
 		/// <inheritdoc/>
-		public async Task<IPool> AddAsync(PoolId Id, string Name, Condition? Condition, bool? EnableAutoscaling, int? MinAgents, int? NumReserveAgents, IEnumerable<KeyValuePair<string, string>>? Properties)
+		public async Task<IPool> AddAsync(PoolId Id, string Name, Condition? Condition, bool? EnableAutoscaling, int? MinAgents, int? NumReserveAgents, PoolSizeStrategy? SizeStrategy, IEnumerable<KeyValuePair<string, string>>? Properties)
 		{
 			PoolDocument Pool = new PoolDocument();
 			Pool.Id = Id;
@@ -98,6 +101,8 @@ namespace HordeServer.Collections.Impl
 			{
 				Pool.Properties = new Dictionary<string, string>(Properties);
 			}
+
+			Pool.SizeStrategy = SizeStrategy;
 			await Pools.InsertOneAsync(Pool);
 			return Pool;
 		}
@@ -172,7 +177,7 @@ namespace HordeServer.Collections.Impl
 		}
 
 		/// <inheritdoc/>
-		public Task<IPool?> TryUpdateAsync(IPool Pool, string? NewName, Condition? NewCondition, bool? NewEnableAutoscaling, int? NewMinAgents, int? NewNumReserveAgents, List<AgentWorkspace>? NewWorkspaces, bool? NewUseAutoSdk, Dictionary<string, string?>? NewProperties, DateTime? LastScaleUpTime, DateTime? LastScaleDownTime)
+		public Task<IPool?> TryUpdateAsync(IPool Pool, string? NewName, Condition? NewCondition, bool? NewEnableAutoscaling, int? NewMinAgents, int? NewNumReserveAgents, List<AgentWorkspace>? NewWorkspaces, bool? NewUseAutoSdk, Dictionary<string, string?>? NewProperties, DateTime? LastScaleUpTime, DateTime? LastScaleDownTime, PoolSizeStrategy? SizeStrategy)
 		{
 			TransactionBuilder<PoolDocument> Transaction = new TransactionBuilder<PoolDocument>();
 			if (NewName != null)
@@ -243,6 +248,10 @@ namespace HordeServer.Collections.Impl
 			{
 				Transaction.Set(x => x.LastScaleDownTime, LastScaleDownTime);
 			}
+			if (SizeStrategy != null)
+			{
+				Transaction.Set(x => x.SizeStrategy, SizeStrategy);
+			}
 			return TryUpdateAsync(Pool, Transaction);
 		}
 	}
diff --git a/Engine/Source/Programs/Horde/Horde.Build/Fleet/PoolService.cs b/Engine/Source/Programs/Horde/Horde.Build/Fleet/PoolService.cs
--- a/Engine/Source/Programs/Horde/Horde.Build/Fleet/PoolService.cs
+++ b/Engine/Source/Programs/Horde/Horde.Build/Fleet/PoolService.cs
@@ -15,6 +15,7 @@ using System.Collections.Generic;
 using System.Linq;
 using System.Threading;
 using System.Threading.Tasks;
+using Horde.Build.Fleet.Autoscale;
 
 namespace HordeServer.Services
 {
@@ -66,11 +67,12 @@ namespace HordeServer.Services
 		/// <param name="EnableAutoscaling">Whether to enable autoscaling for this pool</param>
 		/// <param name="MinAgents">Minimum number of agents in the pool</param>
 		/// <param name="NumReserveAgents">Minimum number of idle agents to maintain</param>
+		/// <param name="SizeStrategy">Pool sizing strategy</param>
 		/// <param name="Properties">Properties for the new pool</param>
 		/// <returns>The new pool document</returns>
-		public Task<IPool> CreatePoolAsync(string Name, Condition? Condition = null, bool? EnableAutoscaling = null, int? MinAgents = null, int? NumReserveAgents = null, Dictionary<string, string>? Properties = null)
+		public Task<IPool> CreatePoolAsync(string Name, Condition? Condition = null, bool? EnableAutoscaling = null, int? MinAgents = null, int? NumReserveAgents = null, PoolSizeStrategy? SizeStrategy = null, Dictionary<string, string>? Properties = null)
 		{
-			return Pools.AddAsync(PoolId.Sanitize(Name), Name, Condition, EnableAutoscaling, MinAgents, NumReserveAgents, Properties);
+			return Pools.AddAsync(PoolId.Sanitize(Name), Name, Condition, EnableAutoscaling, MinAgents, NumReserveAgents, SizeStrategy, Properties);
 		}
 
 		/// <summary>
diff --git a/Engine/Source/Programs/Horde/Horde.Build/ServerSettings.cs b/Engine/Source/Programs/Horde/Horde.Build/ServerSettings.cs
--- a/Engine/Source/Programs/Horde/Horde.Build/ServerSettings.cs
+++ b/Engine/Source/Programs/Horde/Horde.Build/ServerSettings.cs
@@ -14,6 +14,7 @@ using System.ComponentModel.DataAnnotations;
 using System.Linq;
 using System.Runtime.InteropServices;
 using System.Threading.Tasks;
+using Horde.Build.Fleet.Autoscale;
 using TimeZoneConverter;
 
 namespace HordeServer
@@ -138,6 +139,34 @@ namespace HordeServer
 		/// </summary>
 		public int MaxConcurrentLeasesPerAgent { get; set; } = 1;
 	}
+	
+	/// <summary>
+	/// Feature flags to aid rollout of new features
+	///
+	/// Once a feature is running in its intended state and is stable, the flag should be removed.
+	/// A name and date of when the flag was created is noted next to it to help encourage this behavior.
+	/// Try having them be just a flag, a boolean.
+	/// </summary>
+	public class FeatureFlagSettings
+	{
+		/// <summary>
+		/// Registers the old v1 autoscale service as a background service
+		/// Added by carl.bystrom on Jan 26, 2022
+		/// </summary>
+		public bool AutoscaleServiceV1Enabled { get; set; } = true;
+		
+		/// <summary>
+		/// Registers the new autoscale service v2 as a background service (can run in parallel with v1 if shadow mode is on)
+		/// Added by carl.bystrom on Jan 26, 2022
+		/// </summary>
+		public bool AutoscaleServiceV2Enabled { get; set; } = false;
+		
+		/// <summary>
+		/// Shadow mode enabled on new AutoscaleServiceV2 will only log pool size changes but not act on them.
+		/// Added by carl.bystrom on Jan 26, 2022
+		/// </summary>
+		public bool AutoscaleServiceV2ShadowMode { get; set; } = true;
+	}
 
 	/// <summary>
 	/// Global settings for the application
@@ -504,11 +533,11 @@ namespace HordeServer
 		/// The Uri for the Jira installation
 		/// </summary>
 		public Uri? JiraUrl { get; set; }
-		
+
 		/// <summary>
-		/// Type of auto-scaling strategy to be used for sizing agent pools (if left empty, legacy auto scaling impl is used)
+		/// Default agent pool sizing strategy for pools that doesn't have one explicitly configured
 		/// </summary>
-		public string? AgentPoolSizeStrategy { get; set; }
+		public PoolSizeStrategy DefaultAgentPoolSizeStrategy { get; set; } = PoolSizeStrategy.LeaseUtilization;
 
 		/// <summary>
 		/// Set the minimum size of the global thread pool
@@ -554,5 +583,8 @@ namespace HordeServer
 		/// Whether to open a browser on startup
 		/// </summary>
 		public bool OpenBrowser { get; set; } = false;
+
+		/// <inheritdoc cref="FeatureFlags" />
+		public FeatureFlagSettings FeatureFlags { get; set; } = new ();
 	}
 }
diff --git a/Engine/Source/Programs/Horde/Horde.Build/Startup.cs b/Engine/Source/Programs/Horde/Horde.Build/Startup.cs
--- a/Engine/Source/Programs/Horde/Horde.Build/Startup.cs
+++ b/Engine/Source/Programs/Horde/Horde.Build/Startup.cs
@@ -336,20 +336,9 @@ namespace HordeServer
 
 			Services.AddSingleton<AutoscaleService>();
 			Services.AddSingleton<AutoscaleServiceV2>();
-
-			if (Enum.TryParse(Settings.AgentPoolSizeStrategy, out PoolSizeStrategy PoolSizeStrategy))
-			{
-				switch (PoolSizeStrategy)
-				{
-					case PoolSizeStrategy.LeaseUtilization: Services.AddSingleton<IPoolSizeStrategy, LeaseUtilizationStrategy>(); break;
-					case PoolSizeStrategy.JobQueue: Services.AddSingleton<IPoolSizeStrategy, JobQueueStrategy>(); break;
-					default: Services.AddSingleton<IPoolSizeStrategy, NoOpPoolSizeStrategy>(); break;
-				}
-			}
-			else
-			{
-				Services.AddSingleton<IPoolSizeStrategy, NoOpPoolSizeStrategy>();
-			}
+			Services.AddSingleton<LeaseUtilizationStrategy>();
+			Services.AddSingleton<JobQueueStrategy>();
+			Services.AddSingleton<NoOpPoolSizeStrategy>();
 			
 			switch (Settings.FleetManager)
 			{
@@ -574,14 +563,14 @@ namespace HordeServer
 
 			if (Settings.EnableBackgroundServices)
 			{
-				// Allow both legacy and v2 impl to co-exist during rollout phase
-				if (Settings.AgentPoolSizeStrategy != null)
+				if (Settings.FeatureFlags.AutoscaleServiceV1Enabled)
 				{
-					Services.AddHostedService(Provider => Provider.GetRequiredService<AutoscaleServiceV2>());
+					Services.AddHostedService(Provider => Provider.GetRequiredService<AutoscaleService>());	
 				}
-				else
+				
+				if (Settings.FeatureFlags.AutoscaleServiceV2Enabled)
 				{
-					Services.AddHostedService(Provider => Provider.GetRequiredService<AutoscaleService>());	
+					Services.AddHostedService(Provider => Provider.GetRequiredService<AutoscaleServiceV2>());
 				}
 				
 				Services.AddHostedService(Provider => Provider.GetRequiredService<AgentService>());
