Horde Storage - Wait a while longer before running tests to give scylla time to start up.

diff --git a/Engine/Source/Programs/Horde/BuildHorde.xml b/Engine/Source/Programs/Horde/BuildHorde.xml
--- a/Engine/Source/Programs/Horde/BuildHorde.xml
+++ b/Engine/Source/Programs/Horde/BuildHorde.xml
@@ -39,7 +39,7 @@
 			<Docker-Compose-Up File="$(ComposeFile)"/>
 
 			<!-- Sleep a while to let scylla finish starting up before we execute the tests -->
-			<Spawn Exe="/bin/bash" Arguments="-c 'sleep 30'" />
+			<Spawn Exe="/bin/bash" Arguments="-c 'sleep 60'" />
 
 			<!-- Run the unit tests, and override the dotnet sdk used to the global one for the machine rather then what is in the sync, to allow us to use newer versions then the rest of the engine has migrated to -->
 			<DotNet Arguments="test Engine/Source/Programs/Horde/HordeStorage/Horde.Storage.sln" DotNetPath="/usr/share/dotnet/dotnet"/>
