Horde: Fix issue with suppressed 404 fetches not resolving
#rnx
#skipci

diff --git a/Engine/Source/Programs/Horde/HordeDashboard/src/backend/Fetch.ts b/Engine/Source/Programs/Horde/HordeDashboard/src/backend/Fetch.ts
--- a/Engine/Source/Programs/Horde/HordeDashboard/src/backend/Fetch.ts
+++ b/Engine/Source/Programs/Horde/HordeDashboard/src/backend/Fetch.ts
@@ -49,6 +49,7 @@ export class Fetch {
 
                 if (response.status === 404) {
                     if (config?.suppress404) {
+                        reject(`Received suppressed 404 on ${url}`);
                         return;
                     }
                 }
