Horde: Ensure job queue-based stategy can shrink pools

Will scale-in a pool when the queue is empty.

diff --git a/Engine/Source/Programs/Horde/Horde.Build.Tests/Fleet/JobQueueStrategyTest.cs b/Engine/Source/Programs/Horde/Horde.Build.Tests/Fleet/JobQueueStrategyTest.cs
--- a/Engine/Source/Programs/Horde/Horde.Build.Tests/Fleet/JobQueueStrategyTest.cs
+++ b/Engine/Source/Programs/Horde/Horde.Build.Tests/Fleet/JobQueueStrategyTest.cs
@@ -18,7 +18,7 @@ namespace Horde.Build.Tests.Fleet
 	using StreamId = StringId<IStream>;
 	
 	[TestClass]
-	public class JobQueueTest : TestSetup
+	public class JobQueueStrategyTest : TestSetup
 	{
 		[TestMethod]
 		public async Task GetPoolQueueSizes()
@@ -30,30 +30,41 @@ namespace Horde.Build.Tests.Fleet
 		}
 		
 		[TestMethod]
-		public async Task JobQueueWith5Jobs()
+		public async Task EmptyJobQueue()
 		{
-			(JobQueueStrategy Strategy, PoolSizeData PoolSizeData) = await SetUpJobsAsync(1, 5);
-			List<PoolSizeData> Result = await Strategy.CalcDesiredPoolSizesAsync(new() { PoolSizeData });
-			Assert.AreEqual(1, Result.Count);
-			Assert.AreEqual(Result[0].Agents.Count + 0, Result[0].DesiredAgentCount);
+			await AssertAgentCount(0, -1);
 		}
 		
 		[TestMethod]
-		public async Task JobQueueWith15Jobs()
+		public async Task BelowMinQueueSizeForScaleOut()
 		{
-			(JobQueueStrategy Strategy, PoolSizeData PoolSizeData) = await SetUpJobsAsync(1, 15);
-			List<PoolSizeData> Result = await Strategy.CalcDesiredPoolSizesAsync(new() { PoolSizeData });
-			Assert.AreEqual(1, Result.Count);
-			Assert.AreEqual(Result[0].Agents.Count + 1, Result[0].DesiredAgentCount);
+			await AssertAgentCount(2, 0);
 		}
 		
 		[TestMethod]
-		public async Task JobQueueWith25Jobs()
+		public async Task NumQueuedJobs3()
 		{
-			(JobQueueStrategy Strategy, PoolSizeData PoolSizeData) = await SetUpJobsAsync(1, 25);
+			await AssertAgentCount(3, 1);
+		}
+		
+		[TestMethod]
+		public async Task NumQueuedJobs6()
+		{
+			await AssertAgentCount(6, 2);
+		}
+		
+		[TestMethod]
+		public async Task NumQueuedJobs25()
+		{
+			await AssertAgentCount(25, 6);
+		}
+
+		public async Task AssertAgentCount(int NumBatchesWaiting, int ExpectedAgentDelta)
+		{
+			(JobQueueStrategy Strategy, PoolSizeData PoolSizeData) = await SetUpJobsAsync(1, NumBatchesWaiting);
 			List<PoolSizeData> Result = await Strategy.CalcDesiredPoolSizesAsync(new() { PoolSizeData });
 			Assert.AreEqual(1, Result.Count);
-			Assert.AreEqual(Result[0].Agents.Count + 2, Result[0].DesiredAgentCount);
+			Assert.AreEqual(Result[0].Agents.Count + ExpectedAgentDelta, Result[0].DesiredAgentCount);
 		}
 	
 		/// <summary>
@@ -64,9 +75,13 @@ namespace Horde.Build.Tests.Fleet
 		private async Task<(JobQueueStrategy, PoolSizeData)> SetUpJobsAsync(int NumBatchesRunning, int NumBatchesWaiting)
 		{
 			IPool Pool1 = await PoolService.CreatePoolAsync("bogusPool1", null, true, 0, 0);
-			IAgent Agent1 = await CreateAgentAsync(Pool1);
-			IAgent Agent2 = await CreateAgentAsync(Pool1);
-			PoolSizeData PoolSize = new (Pool1, new List<IAgent>() { Agent1, Agent2 }, null);
+			List<IAgent> Agents = new();
+			for (int i = 0; i < 10; i++)
+			{
+				Agents.Add(await CreateAgentAsync(Pool1));
+			}
+			
+			PoolSizeData PoolSize = new (Pool1, Agents, null);
 			
 			string AgentTypeName1 = "bogusAgentType1";
 			Dictionary<string, CreateAgentTypeRequest> AgentTypes = new() { {AgentTypeName1, new() { Pool = Pool1.Name} }, };
diff --git a/Engine/Source/Programs/Horde/Horde.Build/Fleet/Autoscale/JobQueueStrategy.cs b/Engine/Source/Programs/Horde/Horde.Build/Fleet/Autoscale/JobQueueStrategy.cs
--- a/Engine/Source/Programs/Horde/Horde.Build/Fleet/Autoscale/JobQueueStrategy.cs
+++ b/Engine/Source/Programs/Horde/Horde.Build/Fleet/Autoscale/JobQueueStrategy.cs
@@ -29,7 +29,32 @@ namespace Horde.Build.Fleet.Autoscale
 		private readonly IGraphCollection Graphs;
 		private readonly StreamService StreamService;
 		private readonly IClock Clock;
-		private readonly TimeSpan LookbackPeriod = TimeSpan.FromHours(2.0);
+			
+		/// <summary>
+		/// How far back in time to look for job batches (that potentially are in the queue)
+		/// </summary>
+		private readonly TimeSpan SamplePeriod = TimeSpan.FromHours(2.0);
+		
+		/// <summary>
+		/// Minimum number of jobs in queue for scale-out logic to activate
+		///
+		/// Useful to avoid a very small queue size triggering any scaling. 
+		/// </summary>
+		private readonly int MinQueueSizeForScaleOut = 3;
+
+		/// <summary>
+		/// Factor translating queue size to additional agents to grow the pool with
+		///
+		/// Example: if there are 20 jobs in queue, a factor 0.25 will result in 5 new agents being added (20 * 0.25)
+		/// </summary>
+		private readonly double ScaleOutFactor = 0.25;
+		
+		/// <summary>
+		/// Factor by which to shrink the pool size with when queue is empty
+		///
+		/// Example: when the queue size is zero, a default value of 0.9 will shrink the pool by 10% (current agent count * 0.9)
+		/// </summary>
+		private readonly double ScaleInFactor = 0.9;
 
 		/// <summary>
 		/// Constructor
@@ -38,14 +63,14 @@ namespace Horde.Build.Fleet.Autoscale
 		/// <param name="Graphs"></param>
 		/// <param name="StreamService"></param>
 		/// <param name="Clock"></param>
-		/// <param name="LookbackPeriod">Time period for each sample</param>
-		public JobQueueStrategy(IJobCollection Jobs, IGraphCollection Graphs, StreamService StreamService, IClock Clock, TimeSpan? LookbackPeriod = null)
+		/// <param name="SamplePeriod">Time period for each sample</param>
+		public JobQueueStrategy(IJobCollection Jobs, IGraphCollection Graphs, StreamService StreamService, IClock Clock, TimeSpan? SamplePeriod = null)
 		{
 			this.Jobs = Jobs;
 			this.Graphs = Graphs;
 			this.StreamService = StreamService;
 			this.Clock = Clock;
-			this.LookbackPeriod = LookbackPeriod ?? this.LookbackPeriod;
+			this.SamplePeriod = SamplePeriod ?? this.SamplePeriod;
 		}
 
 		/// <inheritdoc/>
@@ -102,33 +127,27 @@ namespace Horde.Build.Fleet.Autoscale
 		/// <inheritdoc/>
 		public async Task<List<PoolSizeData>> CalcDesiredPoolSizesAsync(List<PoolSizeData> Pools)
 		{
-			DateTimeOffset MinCreateTime = Clock.UtcNow - LookbackPeriod;
+			DateTimeOffset MinCreateTime = Clock.UtcNow - SamplePeriod;
 			Dictionary<PoolId, int> PoolQueueSizes = await GetPoolQueueSizesAsync(MinCreateTime);
-			List<PoolSizeData> Result = new();
-			
-			foreach ((PoolId PoolId, int QueueSize) in PoolQueueSizes.OrderBy(x => x.Value))
+
+			return Pools.Select(Current =>
 			{
-				PoolSizeData? PoolSize = Pools.Find(x => x.Pool.Id == PoolId);
-				if (PoolSize != null)
+				PoolQueueSizes.TryGetValue(Current.Pool.Id, out var QueueSize);
+				if (QueueSize > 0)
 				{
-					// TODO: Tweak these values once in production
-					int AdditionalAgentCount = QueueSize switch
-					{
-						< 10 => 0,
-						<= 20 => 1,
-						<= 30 => 2,
-						<= 40 => 3,
-						<= 50 => 4,
-						<= 60 => 5,
-						_ => 4
-					};
-
-					int DesiredAgentCount = PoolSize.Agents.Count + AdditionalAgentCount;
-					Result.Add(new(PoolSize.Pool, PoolSize.Agents, DesiredAgentCount, $"QueueSize={QueueSize}"));
-				}
-			}
+					int AdditionalAgentCount = (int)Math.Round(QueueSize * ScaleOutFactor);
+					if (QueueSize < MinQueueSizeForScaleOut)
+						AdditionalAgentCount = 0;
 
-			return Result;
+					int DesiredAgentCount = Current.Agents.Count + AdditionalAgentCount;
+					return new PoolSizeData(Current.Pool, Current.Agents, DesiredAgentCount, $"QueueSize={QueueSize}");
+				}
+				else
+				{
+					int DesiredAgentCount = (int)(Current.Agents.Count * ScaleInFactor);
+					return new PoolSizeData(Current.Pool, Current.Agents, DesiredAgentCount, "Empty job queue");
+				}
+			}).ToList();
 		}
 	}
 }
