Horde Storage - Made the consistency check service delete incorrect blobs that it finds, seen a few examples of 0 byte blobs getting created somehow which I will need to investigate.

diff --git a/Engine/Source/Programs/Horde/HordeStorage/Horde.Storage/Implementation/Blob/ConsistencyCheckService.cs b/Engine/Source/Programs/Horde/HordeStorage/Horde.Storage/Implementation/Blob/ConsistencyCheckService.cs
--- a/Engine/Source/Programs/Horde/HordeStorage/Horde.Storage/Implementation/Blob/ConsistencyCheckService.cs
+++ b/Engine/Source/Programs/Horde/HordeStorage/Horde.Storage/Implementation/Blob/ConsistencyCheckService.cs
@@ -107,7 +107,9 @@ namespace Horde.Storage.Implementation
                     BlobIdentifier newHash = await BlobIdentifier.FromStream(s);
                     if (!blob.Equals(newHash))
                     {
-                        _logger.Error("Mismatching hash for S3 object {Blob} in {Namespace}, new hash has {NewHash}", blob, ns, newHash);
+                        _logger.Error("Mismatching hash for S3 object {Blob} in {Namespace}, new hash has {NewHash}. Deleting incorrect blob.", blob, ns, newHash);
+
+                        await s3Store.DeleteObject(ns, blob);
                     }
                 }
                 
