Horde Storage - Added flag to control if blob cleanup should run, and disabling the runs when not elected leader.

This will remove a bunch of empty spans of runs that didnt actually run which looks odd in the monitoring.

diff --git a/Engine/Source/Programs/Horde/HordeStorage/Horde.Storage/Implementation/GC/BlobCleanupService.cs b/Engine/Source/Programs/Horde/HordeStorage/Horde.Storage/Implementation/GC/BlobCleanupService.cs
--- a/Engine/Source/Programs/Horde/HordeStorage/Horde.Storage/Implementation/GC/BlobCleanupService.cs
+++ b/Engine/Source/Programs/Horde/HordeStorage/Horde.Storage/Implementation/GC/BlobCleanupService.cs
@@ -83,6 +83,9 @@ namespace Horde.Storage.Implementation
             
             foreach (IBlobCleanup blobCleanup in state.BlobCleanups)
             {
+                if (!blobCleanup.ShouldRun())
+                    continue;
+
                 string type = blobCleanup.GetType().ToString();
                 _logger.Information("Blob cleanup running for {BlobCleanup}", type);
                 using IScope scope = Tracer.Instance.StartActive("gc.blob");
diff --git a/Engine/Source/Programs/Horde/HordeStorage/Horde.Storage/Implementation/GC/OrphanBlobCleanup.cs b/Engine/Source/Programs/Horde/HordeStorage/Horde.Storage/Implementation/GC/OrphanBlobCleanup.cs
--- a/Engine/Source/Programs/Horde/HordeStorage/Horde.Storage/Implementation/GC/OrphanBlobCleanup.cs
+++ b/Engine/Source/Programs/Horde/HordeStorage/Horde.Storage/Implementation/GC/OrphanBlobCleanup.cs
@@ -19,6 +19,7 @@ namespace Horde.Storage.Implementation
 {
     public interface IBlobCleanup
     {
+        bool ShouldRun();
         Task<List<BlobIdentifier>> Cleanup(CancellationToken none);
     }
 
@@ -66,6 +67,16 @@ namespace Horde.Storage.Implementation
 
         }
 
+        public bool ShouldRun()
+        {
+            if (!_leaderElection.IsThisInstanceLeader())
+            {
+                return false;
+            }
+
+            return true;
+        }
+
         public async Task<List<BlobIdentifier>> Cleanup(CancellationToken cancellationToken)
         {
             if (!_leaderElection.IsThisInstanceLeader())
diff --git a/Engine/Source/Programs/Horde/HordeStorage/Horde.Storage/Implementation/GC/OrphanBlobCleanupRefs.cs b/Engine/Source/Programs/Horde/HordeStorage/Horde.Storage/Implementation/GC/OrphanBlobCleanupRefs.cs
--- a/Engine/Source/Programs/Horde/HordeStorage/Horde.Storage/Implementation/GC/OrphanBlobCleanupRefs.cs
+++ b/Engine/Source/Programs/Horde/HordeStorage/Horde.Storage/Implementation/GC/OrphanBlobCleanupRefs.cs
@@ -39,6 +39,16 @@ namespace Horde.Storage.Implementation
             _namespacePolicyResolver = namespacePolicyResolver;
         }
 
+        public bool ShouldRun()
+        {
+            if (!_leaderElection.IsThisInstanceLeader())
+            {
+                return false;
+            }
+
+            return true;
+        }
+
         public async Task<List<BlobIdentifier>> Cleanup(CancellationToken cancellationToken)
         {
             if (!_leaderElection.IsThisInstanceLeader())
