Horde Storage - Replicaton snapshot service will now only strive to build a snapshot per day more reliably by checking when the previous snapshot was built, but polling much more frequently.

This fixes issue were no snapshots were built if we frequently redeployed as snapshots would only run 24 hrs after the deploy (and then every 24hrs from that).

Also added a setting to control how wide we go with the snapshot building and limited it to 2 at the same time, in attempt to spread out the load on the worker a bit more.

diff --git a/Engine/Source/Programs/Horde/HordeStorage/Horde.Storage/Implementation/TransactionLog/ReplicationSnapshotService.cs b/Engine/Source/Programs/Horde/HordeStorage/Horde.Storage/Implementation/TransactionLog/ReplicationSnapshotService.cs
--- a/Engine/Source/Programs/Horde/HordeStorage/Horde.Storage/Implementation/TransactionLog/ReplicationSnapshotService.cs
+++ b/Engine/Source/Programs/Horde/HordeStorage/Horde.Storage/Implementation/TransactionLog/ReplicationSnapshotService.cs
@@ -1,7 +1,6 @@
 ﻿// Copyright Epic Games, Inc. All Rights Reserved.
 
 using System;
-using System.Collections.Generic;
 using System.ComponentModel.DataAnnotations;
 using System.Threading;
 using System.Threading.Tasks;
@@ -37,7 +36,7 @@ namespace Horde.Storage.Implementation
         }
 
         public ReplicationSnapshotService(IServiceProvider provider, IOptionsMonitor<SnapshotSettings> settings, IReplicationLog replicationLog, ILeaderElection leaderElection) :
-            base(serviceName: nameof(ReplicationSnapshotService), TimeSpan.FromSeconds(settings.CurrentValue.SnapshotFrequencySeconds), new SnapshotState())
+            base(serviceName: nameof(ReplicationSnapshotService), TimeSpan.FromMinutes(15), new SnapshotState())
         {
             _provider = provider;
             _settings = settings;
@@ -62,12 +61,23 @@ namespace Horde.Storage.Implementation
             bool ran = false;
             _snapshotBuildTask = _replicationLog.GetNamespaces().ParallelForEachAsync(async ns =>
             {
+                SnapshotInfo? latestSnapshot = await _replicationLog.GetLatestSnapshot(ns);
+                if (latestSnapshot != null)
+                {
+                    DateTime lastSnapshot = latestSnapshot.Timestamp;
+                    DateTime nextSnapshot = lastSnapshot.AddDays(1);
+                    if (DateTime.Now < nextSnapshot)
+                    {
+                        _logger.Information("Skipped building snapshot for namespace {Namespace} as the previous snapshot ({PreviousSnapshot}) was not a day old.", ns, lastSnapshot);
+                        return;
+                    }
+                }
                 ReplicationLogSnapshotBuilder builder = ActivatorUtilities.CreateInstance<ReplicationLogSnapshotBuilder>(_provider);
                 BlobIdentifier snapshotBlob = await builder.BuildSnapshot(ns, _settings.CurrentValue.SnapshotStorageNamespace, _cancellationTokenSource.Token);
                 _logger.Information("Snapshot built for {Namespace} with id {Id}", ns, snapshotBlob);
 
                 ran = true;
-            }, _cancellationTokenSource.Token);
+            }, _settings.CurrentValue.MaxCountOfNamespacesToSnapshotInParallel, _cancellationTokenSource.Token);
             await _snapshotBuildTask;
             _snapshotBuildTask = null;
             return ran;
@@ -89,12 +99,7 @@ namespace Horde.Storage.Implementation
         /// </summary>
         public bool Enabled { get; set; } = false;
 
-        /// <summary>
-        /// The frequency at which to poll for new replication events
-        /// </summary>
-        [Required]
-        public int SnapshotFrequencySeconds { get; set; } = (int)TimeSpan.FromDays(1.0).TotalSeconds; // default to once a day
-
         public NamespaceId SnapshotStorageNamespace { get; set; } = INamespacePolicyResolver.JupiterInternalNamespace;
+        public int MaxCountOfNamespacesToSnapshotInParallel { get; set; } = 2;
     }
 }
