Horde Storage - Missing file from change 18770087 - filesystem store cleanup should always run

diff --git a/Engine/Source/Programs/Horde/HordeStorage/Horde.Storage/Implementation/Blob/FileSystemStore.cs b/Engine/Source/Programs/Horde/HordeStorage/Horde.Storage/Implementation/Blob/FileSystemStore.cs
--- a/Engine/Source/Programs/Horde/HordeStorage/Horde.Storage/Implementation/Blob/FileSystemStore.cs
+++ b/Engine/Source/Programs/Horde/HordeStorage/Horde.Storage/Implementation/Blob/FileSystemStore.cs
@@ -187,7 +187,12 @@ namespace Horde.Storage.Implementation
                 yield return (new BlobIdentifier(file.Name), file.LastWriteTime);
             }
         }
-        
+
+        public bool ShouldRun()
+        {
+            return true;
+        }
+
         public async Task<List<BlobIdentifier>> Cleanup(CancellationToken cancellationToken)
         {
             return await CleanupInternal(cancellationToken);
