Horde Storage - Resolves the filesystem store path for the healthcheck

diff --git a/Engine/Source/Programs/Horde/HordeStorage/Horde.Storage/Horde.Storage.Startup.cs b/Engine/Source/Programs/Horde/HordeStorage/Horde.Storage/Horde.Storage.Startup.cs
--- a/Engine/Source/Programs/Horde/HordeStorage/Horde.Storage/Horde.Storage.Startup.cs
+++ b/Engine/Source/Programs/Horde/HordeStorage/Horde.Storage/Horde.Storage.Startup.cs
@@ -623,7 +623,7 @@ namespace Horde.Storage
                         FilesystemSettings filesystemSettings = provider.GetService<IOptionsMonitor<FilesystemSettings>>()!.CurrentValue;
                         healthChecks.AddDiskStorageHealthCheck(options =>
                         {
-                            string? driveRoot = Path.GetPathRoot(filesystemSettings.RootDir);
+                            string? driveRoot = Path.GetPathRoot(PathUtil.ResolvePath(filesystemSettings.RootDir));
                             options.AddDrive(driveRoot);
                         });
                         break;
