Horde Storage - Corrected settings for development mode, is supposed to be running with the mongo service.

diff --git a/Engine/Source/Programs/Horde/HordeStorage/Horde.Storage/appsettings.Development.json b/Engine/Source/Programs/Horde/HordeStorage/Horde.Storage/appsettings.Development.json
--- a/Engine/Source/Programs/Horde/HordeStorage/Horde.Storage/appsettings.Development.json
+++ b/Engine/Source/Programs/Horde/HordeStorage/Horde.Storage/appsettings.Development.json
@@ -34,8 +34,8 @@
     "LeaderElectionImplementation": "Static",
 
     // use mongo for any database needs
-    "ReferencesDbImplementations": "Mongo",
-    "ContentIdStoreImplementations": "Mongo",
+    "ReferencesDbImplementation": "Mongo",
+    "ContentIdStoreImplementation": "Mongo",
     "BlobIndexImplementation": "Mongo",
 
     // store blobs in the filesystem
