Horde: Build Horde.Build with Horde.Storage for local development. Can be disabled with the /p:WithHordeStorage argument for containerized builds.

diff --git a/Engine/Source/Programs/Horde/Horde.Build/Horde.Build.csproj b/Engine/Source/Programs/Horde/Horde.Build/Horde.Build.csproj
--- a/Engine/Source/Programs/Horde/Horde.Build/Horde.Build.csproj
+++ b/Engine/Source/Programs/Horde/Horde.Build/Horde.Build.csproj
@@ -10,6 +10,14 @@
     <RunAnalyzersDuringBuild>False</RunAnalyzersDuringBuild>
   </PropertyGroup>
 
+  <PropertyGroup Condition="'$(WithHordeStorage)' == ''">
+    <WithHordeStorage>true</WithHordeStorage>
+  </PropertyGroup>
+
+  <PropertyGroup Condition="$(WithHordeStorage)">
+    <DefineConstants>$(DefineConstants);WITH_HORDE_STORAGE</DefineConstants>
+  </PropertyGroup>
+
   <ItemGroup>
     <PackageReference Include="AWSSDK.EC2" Version="3.7.0" />
     <PackageReference Include="AWSSDK.Extensions.NETCore.Setup" Version="3.7.0.1" />
@@ -30,13 +38,13 @@
     <PackageReference Include="Microsoft.AspNetCore.Authentication.OpenIdConnect" Version="3.1.5" />
     <PackageReference Include="Microsoft.AspNetCore.DataProtection.StackExchangeRedis" Version="3.1.7" />
     <PackageReference Include="Microsoft.AspNetCore.SpaServices.Extensions" Version="3.1.17" />
-    <PackageReference Include="Microsoft.Extensions.Caching.Memory" Version="5.0.0" />
+    <PackageReference Include="Microsoft.Extensions.Caching.Memory" Version="6.0.0" />
     <PackageReference Include="MongoDB.Driver" Version="2.11.6" />
     <PackageReference Include="OpenTracing" Version="0.12.1" />
     <PackageReference Include="OpenTracing.Contrib.Grpc" Version="0.4.0" />
     <PackageReference Include="Polly" Version="7.2.1" />
     <PackageReference Include="Polly.Extensions.Http" Version="3.0.0" />
-    <PackageReference Include="Serilog.AspNetCore" Version="3.2.0" />
+    <PackageReference Include="Serilog.AspNetCore" Version="3.4.0" />
     <PackageReference Include="Serilog.Enrichers.OpenTracing" Version="0.0.2" />
     <PackageReference Include="Serilog.Extensions.Hosting" Version="3.1.0" />
     <PackageReference Include="Serilog.Sinks.Console" Version="3.1.1" />
@@ -44,7 +52,7 @@
     <PackageReference Include="SerilogAnalyzer" Version="0.15.0" />
     <PackageReference Include="SharpZipLib" Version="1.2.0" />
     <PackageReference Include="StackExchange.Redis" Version="2.2.4" />
-    <PackageReference Include="Swashbuckle.AspNetCore" Version="5.5.0" />
+    <PackageReference Include="Swashbuckle.AspNetCore" Version="5.6.3" />
     <PackageReference Include="System.Linq.Async" Version="5.0.0" />
     <PackageReference Include="TimeZoneConverter" Version="3.2.0" />
   </ItemGroup>
@@ -98,6 +106,10 @@
     <ProjectReference Include="..\..\..\ThirdParty\Perforce\P4Api.Net\p4api.net\p4api.net.csproj" />
   </ItemGroup>
 
+  <ItemGroup Condition="$(WithHordeStorage)">
+    <ProjectReference Include="..\HordeStorage\Horde.Storage\Horde.Storage.csproj" />
+  </ItemGroup>
+
   <ItemGroup>
     <Content Include="..\..\..\ThirdParty\Perforce\P4Api.Net\native\windows\release\p4bridge.*" Condition=" '$(OS)' == 'Windows_NT' ">
       <CopyToOutputDirectory>PreserveNewest</CopyToOutputDirectory>
diff --git a/Engine/Source/Programs/Horde/Horde.Build/Issues/IIssueCollection.cs b/Engine/Source/Programs/Horde/Horde.Build/Issues/IIssueCollection.cs
--- a/Engine/Source/Programs/Horde/Horde.Build/Issues/IIssueCollection.cs
+++ b/Engine/Source/Programs/Horde/Horde.Build/Issues/IIssueCollection.cs
@@ -4,7 +4,6 @@ using HordeCommon;
 using HordeServer.Models;
 using HordeServer.Services;
 using HordeServer.Utilities;
-using MessagePack.Formatters;
 using Microsoft.AspNetCore.Connections.Features;
 using Microsoft.Extensions.Logging;
 using MongoDB.Bson;
diff --git a/Engine/Source/Programs/Horde/Horde.Build/Program.cs b/Engine/Source/Programs/Horde/Horde.Build/Program.cs
--- a/Engine/Source/Programs/Horde/Horde.Build/Program.cs
+++ b/Engine/Source/Programs/Horde/Horde.Build/Program.cs
@@ -145,7 +145,7 @@ namespace HordeServer
 			return SchemaTypes.ToArray();
 		}
 
-		public static void Main(string[] Args)
+		public static async Task Main(string[] Args)
 		{
 			CommandLineArguments Arguments = new CommandLineArguments(Args);
 
@@ -213,7 +213,12 @@ namespace HordeServer
 			{
 				using (X509Certificate2? GrpcCertificate = ReadGrpcCertificate(HordeSettings))
 				{
-					CreateHostBuilderWithCert(Args, Config, HordeSettings, GrpcCertificate).Build().Run();
+					List<IHost> Hosts = new List<IHost>();
+					Hosts.Add(CreateHostBuilderWithCert(Args, Config, HordeSettings, GrpcCertificate).Build());
+#if WITH_HORDE_STORAGE
+					Hosts.Add(Horde.Storage.Program.CreateHostBuilder(Args).Build());
+#endif
+					await Task.WhenAll(Hosts.Select(x => x.RunAsync()));
 				}
 			}
 #pragma warning disable CA1031 // Do not catch general exception types
diff --git a/Engine/Source/Programs/Horde/Horde.Build/appsettings.json b/Engine/Source/Programs/Horde/Horde.Build/appsettings.json
--- a/Engine/Source/Programs/Horde/Horde.Build/appsettings.json
+++ b/Engine/Source/Programs/Horde/Horde.Build/appsettings.json
@@ -6,10 +6,8 @@
 		"CorsOrigin": "https://localhost:5003;http://localhost:3000",
 		"LogFileProviderType": "Local"
 	},
-	"Serilog": 
-	{
-		"MinimumLevel": 
-		{
+	"Serilog": {
+		"MinimumLevel": {
 			"Default": "Information",
 			"Override": {
 				"MongoDB": "Warning", // For bundled MongoDB output
@@ -24,5 +22,71 @@
 			}
 		}
 	},
-	"AllowedHosts": "*"
+	"AllowedHosts": "*",
+
+	"Mongo": {
+		// this assumes a mongo db were we have localhost exemption, e.g. we do not need to authenticate
+		// if using mongo via a docker container you will need to update this
+		"ConnectionString": "mongodb://localhost:27017",
+		"RequireTls12": true
+	},
+
+	// scylla is normally not enabled in development but we still configure it so we can easily switch
+	"Scylla": {
+		"LocalDatacenterName": "datacenter1",
+		"LocalKeyspaceSuffix": "local"
+	},
+
+	"Horde_Storage": {
+		// assume a single instance is running
+		"LeaderElectionImplementation": "Static",
+
+		// use mongo for any database needs
+		"ReferencesDbImplementation": "Mongo",
+		"ContentIdStoreImplementation": "Mongo",
+		"BlobIndexImplementation": "Mongo",
+
+		// store blobs in the filesystem
+		"StorageImplementations": [
+			"Filesystem"
+		],
+
+		// disable old api that we do not want to use
+		"TransactionLogWriterImplementation": "Memory",
+		"RefDbImplementation": "Memory"
+	},
+
+	"Filesystem": {
+		"RootDir": "$(ExecutableLocation)/BlobStorage"
+	},
+
+	"Replication": {
+		"StateRoot": "$(ExecutableLocation)/ganymede",
+		"ReplicationPollFrequencySeconds": 15,
+		"CurrentSite": "local-test",
+		"Enabled": false
+	},
+
+	"GC": {
+		"BlobCleanupServiceEnabled": false,
+		"CleanOldBlobs": false,
+		"CleanOldRefRecords": false
+	},
+
+	"Jupiter": {
+		"CurrentSite": "local"
+	},
+
+	// disable auth and allow any namespace
+	"Auth": {
+		"Method": "Disabled"
+	},
+
+	"Namespaces": {
+		"Policies": {
+			"*": {
+				"Claims": [ "*" ]
+			}
+		}
+	}
 }
diff --git a/Engine/Source/Programs/Horde/HordeStorage/Horde.Storage/Implementation/MongoStore.cs b/Engine/Source/Programs/Horde/HordeStorage/Horde.Storage/Implementation/MongoStore.cs
--- a/Engine/Source/Programs/Horde/HordeStorage/Horde.Storage/Implementation/MongoStore.cs
+++ b/Engine/Source/Programs/Horde/HordeStorage/Horde.Storage/Implementation/MongoStore.cs
@@ -5,6 +5,7 @@ using System.Collections.Generic;
 using System.Security.Authentication;
 using System.Threading.Tasks;
 using Microsoft.Extensions.Options;
+using MongoDB.Bson;
 using MongoDB.Bson.Serialization;
 using MongoDB.Driver;
 
@@ -31,9 +32,18 @@ public class MongoStore
     {
         string collectionName = GetCollectionName<T>();
 
+        IMongoDatabase database = _client.GetDatabase(GetDatabaseName());
+
+        // Try to avoid exceptions breaking in the debugger unnecessarily by checking for the existence of the collection beforehand.
+        FilterDefinition<BsonDocument> filter = new BsonDocument("name", collectionName);
+        if (await database.ListCollectionNames(new ListCollectionNamesOptions { Filter = filter }).AnyAsync())
+        {
+            return;
+        }
+
         try
         {
-            await _client.GetDatabase(GetDatabaseName()).CreateCollectionAsync(collectionName);
+            await database.CreateCollectionAsync(collectionName);
         }
         catch (MongoCommandException e)
         {
diff --git a/Engine/Source/Programs/Horde/HordeStorage/Horde.Storage/Program.cs b/Engine/Source/Programs/Horde/HordeStorage/Horde.Storage/Program.cs
--- a/Engine/Source/Programs/Horde/HordeStorage/Horde.Storage/Program.cs
+++ b/Engine/Source/Programs/Horde/HordeStorage/Horde.Storage/Program.cs
@@ -5,6 +5,7 @@ using System.IO.Pipelines;
 using System.Net.Sockets;
 using Horde.Storage.Implementation.Kestrel;
 using Jupiter;
+using Microsoft.Extensions.Hosting;
 
 namespace Horde.Storage
 {
@@ -20,5 +21,10 @@ namespace Horde.Storage
         {
             return new KestrelDispatcher(sp, reader, writer, socket);
         }
+
+        public static IHostBuilder CreateHostBuilder(string[] Args)
+        {
+            return BaseProgram<HordeStorageStartup>.CreateHostBuilder(Args, HttpConnectionFactory);
+        }
     }
 }
diff --git a/Engine/Source/Programs/Horde/HordeStorage/Jupiter.Common/BaseProgram.cs b/Engine/Source/Programs/Horde/HordeStorage/Jupiter.Common/BaseProgram.cs
--- a/Engine/Source/Programs/Horde/HordeStorage/Jupiter.Common/BaseProgram.cs
+++ b/Engine/Source/Programs/Horde/HordeStorage/Jupiter.Common/BaseProgram.cs
@@ -90,7 +90,7 @@ namespace Jupiter
             }
         }
 
-        private static IHostBuilder CreateHostBuilder(string[] args, HttpConnectionFactory? httpConnectionFactory)
+        internal static IHostBuilder CreateHostBuilder(string[] args, HttpConnectionFactory? httpConnectionFactory)
         {
             return Host.CreateDefaultBuilder(args)
                 .ConfigureWebHostDefaults(webBuilder =>
