Horde: Make BaseProgram<>.CreateHostBuilder public so it can be accessed from Horde.Storage.

diff --git a/Engine/Source/Programs/Horde/HordeStorage/Jupiter.Common/BaseProgram.cs b/Engine/Source/Programs/Horde/HordeStorage/Jupiter.Common/BaseProgram.cs
--- a/Engine/Source/Programs/Horde/HordeStorage/Jupiter.Common/BaseProgram.cs
+++ b/Engine/Source/Programs/Horde/HordeStorage/Jupiter.Common/BaseProgram.cs
@@ -90,7 +90,7 @@ namespace Jupiter
             }
         }
 
-        internal static IHostBuilder CreateHostBuilder(string[] args, HttpConnectionFactory? httpConnectionFactory)
+        public static IHostBuilder CreateHostBuilder(string[] args, HttpConnectionFactory? httpConnectionFactory)
         {
             return Host.CreateDefaultBuilder(args)
                 .ConfigureWebHostDefaults(webBuilder =>
