Horde: Configure Horde storage to run on different ports (57000, 57001) in shared instances with Horde.Build.

diff --git a/Engine/Source/Programs/Horde/Horde.Build/Program.cs b/Engine/Source/Programs/Horde/Horde.Build/Program.cs
--- a/Engine/Source/Programs/Horde/Horde.Build/Program.cs
+++ b/Engine/Source/Programs/Horde/Horde.Build/Program.cs
@@ -216,7 +216,16 @@ namespace HordeServer
 					List<IHost> Hosts = new List<IHost>();
 					Hosts.Add(CreateHostBuilderWithCert(Args, Config, HordeSettings, GrpcCertificate).Build());
 #if WITH_HORDE_STORAGE
-					Hosts.Add(Horde.Storage.Program.CreateHostBuilder(Args).Build());
+					IHostBuilder StorageHostBuilder = Horde.Storage.Program.CreateHostBuilder(Args);
+					StorageHostBuilder.ConfigureWebHostDefaults(Builder =>
+					{
+						Builder.ConfigureKestrel(Options =>
+						{
+							Options.ListenAnyIP(57000);
+							Options.ListenAnyIP(57001, Configure => Configure.UseHttps());
+						});
+					});
+					Hosts.Add(StorageHostBuilder.Build());
 #endif
 					await Task.WhenAll(Hosts.Select(x => x.RunAsync()));
 				}
