Horde: Use Horde.Storage for remote execution storage, rather than the legacy storage implementation in Horde.Build.

- Agent and server both use the IStorageClient interface in EpicGames.Horde to communicate with the storage backend.
- An instance of Horde.Storage is run in the same process as Horde.Build for local development, but is compiled out for containerized builds.
- Compute clusters are defined in global config file, and control permissions for who can create remote execution jobs. In the future it will likely be used to control access for different machines, prioritization, etc...
- Each compute cluster has an identifier, and the API for starting/monitoring jobs includes that identifier as part of the route (ie. /api/v1/compute/<cluster-id>).
- Removed gRPC interface for storage and compute services. May be resurrected in the future.

diff --git a/Engine/Source/Programs/Horde/Horde.Agent.Tests/WorkerServiceTest.cs b/Engine/Source/Programs/Horde/Horde.Agent.Tests/WorkerServiceTest.cs
--- a/Engine/Source/Programs/Horde/Horde.Agent.Tests/WorkerServiceTest.cs
+++ b/Engine/Source/Programs/Horde/Horde.Agent.Tests/WorkerServiceTest.cs
@@ -50,7 +50,7 @@ namespace HordeAgentTests
 			Settings.Executor = ExecutorType.Test; // Not really used since the executor is overridden in the tests
 			IOptionsMonitor<AgentSettings> SettingsMonitor = new TestOptionsMonitor<AgentSettings>(Settings);
 
-			return new WorkerService(WorkerLogger, SettingsMonitor, null!, CreateExecutor);
+			return new WorkerService(WorkerLogger, SettingsMonitor, null!, null!, CreateExecutor);
 		}
 
 		[TestMethod]
diff --git a/Engine/Source/Programs/Horde/Horde.Agent/AgentSettings.cs b/Engine/Source/Programs/Horde/Horde.Agent/AgentSettings.cs
--- a/Engine/Source/Programs/Horde/Horde.Agent/AgentSettings.cs
+++ b/Engine/Source/Programs/Horde/Horde.Agent/AgentSettings.cs
@@ -1,7 +1,9 @@
 // Copyright Epic Games, Inc. All Rights Reserved.
 
 using EpicGames.Core;
+using EpicGames.Horde.Storage;
 using HordeCommon.Rpc;
+using Microsoft.Extensions.Configuration;
 using System;
 using System.Collections.Generic;
 using System.ComponentModel.DataAnnotations;
@@ -90,6 +92,11 @@ namespace HordeAgent
 		/// </summary>
 		public List<string> Thumbprints { get; set; } = new List<string>();
 
+		/// <summary>
+		/// Storage settings for using this server
+		/// </summary>
+		public StorageSettings Storage { get; set; } = new StorageSettings();
+
 		/// <summary>
 		/// Checks whether the given certificate thumbprint should be trusted
 		/// </summary>
@@ -242,4 +249,21 @@ namespace HordeAgent
 			return Name ?? Environment.MachineName;
 		}
 	}
+
+	/// <summary>
+	/// Extension methods for retrieving config settings
+	/// </summary>
+	public static class AgentSettingsExtensions
+	{
+		/// <summary>
+		/// Gets the configuration section for the active server profile
+		/// </summary>
+		/// <param name="ConfigSection"></param>
+		/// <returns></returns>
+		public static IConfigurationSection GetCurrentServerProfile(this IConfigurationSection ConfigSection)
+		{
+			string ProfileName = ConfigSection[nameof(AgentSettings.Server)];
+			return ConfigSection.GetSection(nameof(AgentSettings.ServerProfiles)).GetChildren().First(x => x[nameof(ServerProfile.Name)] == ProfileName);
+		}
+	}
 }
diff --git a/Engine/Source/Programs/Horde/Horde.Agent/Commands/Execution/ExecuteCommandV2.cs b/Engine/Source/Programs/Horde/Horde.Agent/Commands/Execution/ExecuteCommandV2.cs
--- a/Engine/Source/Programs/Horde/Horde.Agent/Commands/Execution/ExecuteCommandV2.cs
+++ b/Engine/Source/Programs/Horde/Horde.Agent/Commands/Execution/ExecuteCommandV2.cs
@@ -23,6 +23,8 @@ using EpicGames.Serialization;
 using EpicGames.Horde.Compute;
 using EpicGames.Horde.Storage;
 using System.Net.Http;
+using EpicGames.Horde.Compute.Impl;
+using EpicGames.Horde.Storage.Impl;
 
 namespace HordeAgent.Commands
 {
@@ -32,6 +34,8 @@ namespace HordeAgent.Commands
 	[Command("ExecuteV2", "Executes a command through the remote execution API")]
 	class ExecuteCommandV2 : Command
 	{
+		static ClusterId ClusterId { get; } = new ClusterId("default");
+
 		class JsonRequirements
 		{
 			public string? Condition { get; set; }
@@ -109,182 +113,6 @@ namespace HordeAgent.Commands
 		[CommandLine("-Token=")]
 		public string? Token = Environment.GetEnvironmentVariable("HORDE_TOKEN");
 
-		interface IComputeClient : IDisposable
-		{
-			Task AddBlobAsync(string NamespaceId, IoHash Hash, byte[] Data);
-			Task<byte[]> GetBlobAsync(string NamespaceId, IoHash Hash);
-
-			Task AddObjectAsync(string NamespaceId, IoHash Hash, byte[] Data);
-			Task<byte[]> GetObjectAsync(string NamespaceId, IoHash Hash);
-
-			Task AddTaskAsync(string ChannelId, string NamespaceId, IoHash RequirementsHash, IoHash TaskHash, bool SkipTaskLookup);
-			Task GetTaskUpdatesAsync(string ChannelId, string NamespaceId, Func<GetTaskUpdateResponse, Task<bool>> OnUpdate);
-		}
-
-		class HttpComputeClient : IComputeClient
-		{
-			HttpClient HttpClient;
-
-			public HttpComputeClient(Uri Server, string? BearerToken)
-			{
-				HttpClient = new HttpClient();
-				HttpClient.BaseAddress = Server;
-				if (BearerToken != null)
-				{
-					HttpClient.DefaultRequestHeaders.Authorization = new AuthenticationHeaderValue("Bearer", BearerToken);
-				}
-			}
-
-			public void Dispose()
-			{
-				HttpClient.Dispose();
-			}
-
-			public async Task AddBlobAsync(string NamespaceId, IoHash Hash, byte[] Data)
-			{
-				HttpRequestMessage Request = new HttpRequestMessage(HttpMethod.Put, $"api/v1/blobs/{NamespaceId}/{Hash}");
-				Request.Content = new ByteArrayContent(Data);
-
-				HttpResponseMessage Response = await HttpClient.SendAsync(Request);
-				Response.EnsureSuccessStatusCode();
-			}
-
-			public async Task<byte[]> GetBlobAsync(string NamespaceId, IoHash Hash)
-			{
-				HttpRequestMessage Request = new HttpRequestMessage(HttpMethod.Get, $"api/v1/blobs/{NamespaceId}/{Hash}");
-				HttpResponseMessage Response = await HttpClient.SendAsync(Request);
-				Response.EnsureSuccessStatusCode();
-				return await Response.Content.ReadAsByteArrayAsync();
-			}
-
-			public Task AddObjectAsync(string NamespaceId, IoHash Hash, byte[] Data)
-			{
-				return AddBlobAsync(NamespaceId, Hash, Data);
-			}
-
-			public Task<byte[]> GetObjectAsync(string NamespaceId, IoHash Hash)
-			{
-				return GetBlobAsync(NamespaceId, Hash);
-			}
-
-			public async Task AddTaskAsync(string ChannelId, string NamespaceId, IoHash RequirementsHash, IoHash TaskHash, bool SkipCacheLookup)
-			{
-				AddTasksRequest AddTasks = new AddTasksRequest();
-				AddTasks.RequirementsHash = RequirementsHash;
-				AddTasks.TaskHashes.Add(TaskHash);
-				AddTasks.DoNotCache = SkipCacheLookup;
-
-				ReadOnlyMemoryContent Content = new ReadOnlyMemoryContent(CbSerializer.Serialize(AddTasks).GetView());
-				Content.Headers.Add("Content-Type", "application/x-ue-cb");
-
-				HttpRequestMessage Request = new HttpRequestMessage(HttpMethod.Post, $"api/v1/compute/{ChannelId}");
-				Request.Content = Content;
-
-				HttpResponseMessage Response = await HttpClient.SendAsync(Request);
-				Response.EnsureSuccessStatusCode();
-			}
-
-			public async Task GetTaskUpdatesAsync(string ChannelId, string NamespaceId, Func<GetTaskUpdateResponse, Task<bool>> OnUpdateAsync)
-			{
-				for (; ; )
-				{
-					HttpRequestMessage Request = new HttpRequestMessage(HttpMethod.Post, $"api/v1/compute/{ChannelId}/updates?wait=10");
-					Request.Headers.Add("Accept", "application/x-ue-cb");
-
-					HttpResponseMessage Response = await HttpClient.SendAsync(Request);
-					Response.EnsureSuccessStatusCode();
-
-					byte[] Data = await Response.Content.ReadAsByteArrayAsync();
-					GetTaskUpdatesResponse ParsedResponse = CbSerializer.Deserialize<GetTaskUpdatesResponse>(new CbField(Data));
-					
-					foreach (GetTaskUpdateResponse Update in ParsedResponse.Updates)
-					{
-						if (!await OnUpdateAsync(Update))
-						{
-							return;
-						}
-					}
-				}
-			}
-		}
-
-		class GrpcComputeClient : IComputeClient
-		{
-			GrpcChannel Channel;
-			BlobStore.BlobStoreClient BlobStore;
-			ComputeRpc.ComputeRpcClient ComputeClient;
-
-			public GrpcComputeClient(GrpcService GrpcService)
-			{
-				Channel = GrpcService.CreateGrpcChannel();
-				BlobStore = new BlobStore.BlobStoreClient(Channel);
-				ComputeClient = new ComputeRpc.ComputeRpcClient(Channel);
-			}
-
-			public void Dispose()
-			{
-				Channel.Dispose();
-			}
-
-			public async Task AddBlobAsync(string NamespaceId, IoHash Hash, byte[] Data)
-			{
-				AddBlobsRequest Request = new AddBlobsRequest();
-				Request.NamespaceId = NamespaceId;
-				Request.Blobs.Add(new AddBlobRequest(Hash, Data));
-				await BlobStore.AddAsync(Request);
-			}
-
-			public async Task<byte[]> GetBlobAsync(string NamespaceId, IoHash Hash)
-			{
-				return await BlobStore.GetBlobAsync(NamespaceId, Hash);
-			}
-
-			public Task AddObjectAsync(string NamespaceId, IoHash Hash, byte[] Data)
-			{
-				return AddBlobAsync(NamespaceId, Hash, Data);
-			}
-
-			public Task<byte[]> GetObjectAsync(string NamespaceId, IoHash Hash)
-			{
-				return BlobStore.GetBlobAsync(NamespaceId, Hash);
-			}
-
-			public async Task AddTaskAsync(string ChannelId, string NamespaceId, IoHash RequirementsHash, IoHash TaskHash, bool SkipCacheLookup)
-			{
-				AddTasksRpcRequest AddTasksRequest = new AddTasksRpcRequest(ChannelId, NamespaceId, RequirementsHash, new List<CbObjectAttachment> { TaskHash }, SkipCacheLookup);
-				await ComputeClient.AddTasksAsync(AddTasksRequest);
-			}
-
-			public async Task GetTaskUpdatesAsync(string ChannelId, string NamespaceId, Func<GetTaskUpdateResponse, Task<bool>> OnUpdateAsync)
-			{
-				using (AsyncDuplexStreamingCall<GetTaskUpdatesRpcRequest, GetTaskUpdatesRpcResponse> Call = ComputeClient.GetTaskUpdates())
-				{
-					await Call.RequestStream.WriteAsync(new GetTaskUpdatesRpcRequest(ChannelId));
-
-					//					TaskCompletionSource<bool> CompleteNowTask = new TaskCompletionSource<bool>();
-					//					Task CompleteTask = Task.Run(async () => { await CompleteNowTask.Task; await Call.RequestStream.CompleteAsync(); });
-
-					while (await Call.ResponseStream.MoveNext())
-					{
-						GetTaskUpdatesRpcResponse RpcResponse = Call.ResponseStream.Current;
-
-						GetTaskUpdateResponse Response = new GetTaskUpdateResponse();
-						Response.AgentId = RpcResponse.AgentId;
-						Response.LeaseId = RpcResponse.LeaseId;
-						Response.Result = RpcResponse.Result?.Hash;
-						Response.State = RpcResponse.State;
-						Response.TaskHash = RpcResponse.Task;
-						Response.Time = RpcResponse.Time.ToDateTime();
-
-						if (!await OnUpdateAsync(Response))
-						{
-							await Call.RequestStream.CompleteAsync();
-						}
-					}
-				}
-			}
-		}
-
 		static (DirectoryTree, IoHash) CreateSandbox(DirectoryInfo BaseDirInfo, Dictionary<IoHash, byte[]> UploadList)
 		{
 			DirectoryTree Tree = new DirectoryTree();
@@ -308,18 +136,6 @@ namespace HordeAgent.Commands
 			return (Tree, AddCbObject(UploadList, Tree));
 		}
 
-		IComputeClient CreateComputeClient(IServiceProvider Services)
-		{
-			if (UseHttp)
-			{
-				return new HttpComputeClient(new Uri(Server), Token);
-			}
-			else
-			{
-				return new GrpcComputeClient(Services.GetRequiredService<GrpcService>());
-			}
-		}
-
 		/// <inheritdoc/>
 		public override async Task<int> ExecuteAsync(ILogger Logger)
 		{
@@ -350,13 +166,19 @@ namespace HordeAgent.Commands
 				{
 					Services.AddOptions<AgentSettings>().Configure(Options => ConfigSection.Bind(Options)).ValidateDataAnnotations();
 					Services.AddLogging();
+
+					IConfigurationSection StorageSettings = ConfigSection.GetCurrentServerProfile().GetSection(nameof(ServerProfile.Storage));
+					Services.AddHordeStorage(Settings => StorageSettings.Bind(Settings));
+
 					Services.AddSingleton<GrpcService>();
+					Services.AddSingleton<IComputeClient, HttpComputeClient>(SP => new HttpComputeClient(new HttpClient() { BaseAddress = new Uri(Server) }));
 				});
 
 			using (IHost Host = HostBuilder.Build())
 			{
-				using IComputeClient ComputeClient = CreateComputeClient(Host.Services);
-				
+				IStorageClient StorageClient = Host.Services.GetRequiredService<IStorageClient>();
+				IComputeClient ComputeClient = Host.Services.GetRequiredService<IComputeClient>();
+
 				byte[] Json = await FileReference.ReadAllBytesAsync(Input);
 				JsonComputeTask JsonComputeTask = JsonSerializer.Deserialize<JsonComputeTask>(Json, new JsonSerializerOptions { PropertyNameCaseInsensitive = true })!;
 
@@ -378,18 +200,18 @@ namespace HordeAgent.Commands
 				Dictionary<IoHash, byte[]> Blobs = new Dictionary<IoHash, byte[]>();
 				(_, IoHash SandboxHash) = CreateSandbox(InputDir.ToDirectoryInfo(), Blobs);
 
-				ComputeTask Task = new ComputeTask(JsonComputeTask.Executable, JsonComputeTask.Arguments.ConvertAll<Utf8String>(x => x), JsonComputeTask.WorkingDirectory, SandboxHash);
-				Task.EnvVars = JsonComputeTask.EnvVars.ToDictionary(x => (Utf8String)x.Key, x => (Utf8String)x.Value);
-				Task.OutputPaths.AddRange(JsonComputeTask.OutputPaths.Select(x => (Utf8String)x));
-				IoHash TaskHash = AddCbObject(Blobs, Task);
-
 				JsonRequirements JsonRequirements = JsonComputeTask.Requirements;
 				Requirements Requirements = new Requirements(JsonRequirements.Condition ?? String.Empty);
 				Requirements.Resources = JsonRequirements.Resources.ToDictionary(x => x.Key, x => x.Value);
 				Requirements.Exclusive = JsonRequirements.Exclusive;
 				IoHash RequirementsHash = AddCbObject(Blobs, Requirements);
 
-				await ExecuteAction(Logger, ComputeClient, TaskHash, RequirementsHash, Blobs);	
+				ComputeTask Task = new ComputeTask(JsonComputeTask.Executable, JsonComputeTask.Arguments.ConvertAll<Utf8String>(x => x), JsonComputeTask.WorkingDirectory, SandboxHash);
+				Task.EnvVars = JsonComputeTask.EnvVars.ToDictionary(x => (Utf8String)x.Key, x => (Utf8String)x.Value);
+				Task.OutputPaths.AddRange(JsonComputeTask.OutputPaths.Select(x => (Utf8String)x));
+				Task.RequirementsHash = AddCbObject(Blobs, Requirements);
+
+				await ExecuteAction(Logger, ComputeClient, StorageClient, Task, Blobs);	
 			}
 			return 0;
 		}
@@ -402,58 +224,67 @@ namespace HordeAgent.Commands
 			return Hash;
 		}
 
-		private async Task ExecuteAction(ILogger Logger, IComputeClient ComputeClient, CbObjectAttachment TaskHash, CbObjectAttachment RequirementsHash, Dictionary<IoHash, byte[]> UploadList)
+		private async Task ExecuteAction(ILogger Logger, IComputeClient ComputeClient, IStorageClient StorageClient, ComputeTask Task, Dictionary<IoHash, byte[]> UploadList)
 		{
-			Logger.LogInformation("task: {TaskHash}", TaskHash);
-			Logger.LogInformation("requirements: {RequirementsHash}", RequirementsHash);
+			IComputeClusterInfo Cluster = await ComputeClient.GetClusterInfoAsync(ClusterId);
 
-			const string NamespaceId = "default";
-
-			foreach (KeyValuePair<IoHash, byte[]> Pair in UploadList)
+			List<KeyValuePair<IoHash, byte[]>> UploadBlobs = UploadList.ToList();
+			for(int Idx = 0; Idx < UploadBlobs.Count; Idx++)
 			{
-				await ComputeClient.AddBlobAsync(NamespaceId, Pair.Key, Pair.Value);
+				KeyValuePair<IoHash, byte[]> Pair = UploadBlobs[Idx];
+				Logger.LogInformation("Uploading blob {Idx}/{Count}: {Hash}", Idx + 1, UploadBlobs.Count, Pair.Key);
+				await StorageClient.WriteBlobFromMemoryAsync(Cluster.NamespaceId, Pair.Key, Pair.Value);
 			}
 
+			CbObject TaskObject = CbSerializer.Serialize(Task);
+			IoHash TaskHash = IoHash.Compute(TaskObject.GetView().Span);
+			RefId TaskRefId = new RefId(TaskHash);
+			await StorageClient.SetRefAsync(Cluster.NamespaceId, Cluster.RequestBucketId, TaskRefId, TaskObject);
+
+			Logger.LogInformation("task: {TaskHash}", TaskHash);
+			Logger.LogInformation("requirements: {RequirementsHash}", Task.RequirementsHash);
+
 			// Execute the action
-			string ChannelId = Guid.NewGuid().ToString();
-			await ComputeClient.AddTaskAsync(ChannelId, NamespaceId, RequirementsHash, TaskHash, SkipCacheLookup);
+			ChannelId ChannelId = new ChannelId(Guid.NewGuid().ToString());
+			await ComputeClient.AddTaskAsync(ClusterId, ChannelId, TaskRefId, Task.RequirementsHash, SkipCacheLookup);
 
-			await ComputeClient.GetTaskUpdatesAsync(ChannelId, NamespaceId, async Response =>
+			await foreach(IComputeTaskInfo Response in ComputeClient.GetTaskUpdatesAsync(ClusterId, ChannelId))
 			{
-				Logger.LogInformation("{OperationName}: Execution state: {State}", (IoHash)(CbObjectAttachment)Response.TaskHash, Response.State.ToString());
+				Logger.LogInformation("{OperationName}: Execution state: {State}", Response.TaskRefId, Response.State.ToString());
 				if (!String.IsNullOrEmpty(Response.AgentId) || !String.IsNullOrEmpty(Response.LeaseId))
 				{
-					Logger.LogInformation("{OperationName}: Running on agent {AgentId} under lease {LeaseId}", (IoHash)(CbObjectAttachment)Response.TaskHash, Response.AgentId, Response.LeaseId);
+					Logger.LogInformation("{OperationName}: Running on agent {AgentId} under lease {LeaseId}", Response.TaskRefId, Response.AgentId, Response.LeaseId);
 				}
-				if (Response.Result != null)
+				if (Response.ResultRefId != null)
 				{
-					await HandleCompleteTask(ComputeClient, NamespaceId, Response.Result.Value, Logger);
+					await HandleCompleteTask(StorageClient, Cluster.NamespaceId, Cluster.ResponseBucketId, Response.ResultRefId.Value, Logger);
 				}
-				return Response.State != ComputeTaskState.Complete;
-			});
+				if (Response.State == ComputeTaskState.Complete)
+				{
+					break;
+				}
+			}
 		}
 
-		async Task HandleCompleteTask(IComputeClient ComputeClient, string NamespaceId, CbObjectAttachment Hash, ILogger Logger)
+		async Task HandleCompleteTask(IStorageClient StorageClient, NamespaceId NamespaceId, BucketId OutputBucketId, RefId OutputRefId, ILogger Logger)
 		{
-			Logger.LogInformation("result: {ResultHash}", Hash);
-
-			ComputeTaskResult Result = await GetObjectAsync<ComputeTaskResult>(ComputeClient, NamespaceId, Hash);
+			ComputeTaskResult Result = await StorageClient.GetRefAsync<ComputeTaskResult>(NamespaceId, OutputBucketId, OutputRefId);
 			Logger.LogInformation("exit: {ExitCode}", Result.ExitCode);
 
-			await LogTaskOutputAsync(ComputeClient, "stdout", NamespaceId, Result.StdOutHash, Logger);
-			await LogTaskOutputAsync(ComputeClient, "stderr", NamespaceId, Result.StdErrHash, Logger);
+			await LogTaskOutputAsync(StorageClient, "stdout", NamespaceId, Result.StdOutHash, Logger);
+			await LogTaskOutputAsync(StorageClient, "stderr", NamespaceId, Result.StdErrHash, Logger);
 
 			if (Result.OutputHash != null && OutputDir != null)
 			{
-				await WriteOutputAsync(ComputeClient, NamespaceId, Result.OutputHash.Value, new DirectoryReference(OutputDir));
+				await WriteOutputAsync(StorageClient, NamespaceId, Result.OutputHash.Value, new DirectoryReference(OutputDir));
 			}
 		}
 
-		async Task LogTaskOutputAsync(IComputeClient ComputeClient, string Channel, string NamespaceId, IoHash? LogHash, ILogger Logger)
+		async Task LogTaskOutputAsync(IStorageClient StorageClient, string Channel, NamespaceId NamespaceId, IoHash? LogHash, ILogger Logger)
 		{
 			if (LogHash != null)
 			{
-				byte[] StdOutData = await ComputeClient.GetBlobAsync(NamespaceId, LogHash.Value);
+				byte[] StdOutData = await StorageClient.ReadBlobToMemoryAsync(NamespaceId, LogHash.Value);
 				if (StdOutData.Length > 0)
 				{
 					foreach (string Line in Encoding.UTF8.GetString(StdOutData).Split('\n'))
@@ -464,36 +295,30 @@ namespace HordeAgent.Commands
 			}
 		}
 
-		async Task WriteOutputAsync(IComputeClient ComputeClient, string NamespaceId, IoHash TreeHash, DirectoryReference OutputDir)
+		async Task WriteOutputAsync(IStorageClient StorageClient, NamespaceId NamespaceId, IoHash TreeHash, DirectoryReference OutputDir)
 		{
-			DirectoryTree Tree = await GetObjectAsync<DirectoryTree>(ComputeClient, NamespaceId, TreeHash);
+			DirectoryTree Tree = await StorageClient.ReadObjectAsync<DirectoryTree>(NamespaceId, TreeHash);
 
 			List<Task> Tasks = new List<Task>();
 			foreach (FileNode File in Tree.Files)
 			{
 				FileReference OutputFile = FileReference.Combine(OutputDir, File.Name.ToString());
-				Tasks.Add(WriteObjectToFileAsync(ComputeClient, NamespaceId, File.Hash, OutputFile));
+				Tasks.Add(WriteObjectToFileAsync(StorageClient, NamespaceId, File.Hash, OutputFile));
 			}
 			foreach (DirectoryNode Directory in Tree.Directories)
 			{
 				DirectoryReference NextOutputDir = DirectoryReference.Combine(OutputDir, Directory.Name.ToString());
-				Tasks.Add(WriteOutputAsync(ComputeClient, NamespaceId, Directory.Hash, NextOutputDir));
+				Tasks.Add(WriteOutputAsync(StorageClient, NamespaceId, Directory.Hash, NextOutputDir));
 			}
 
 			await Task.WhenAll(Tasks);
 		}
 
-		static async Task<T> GetObjectAsync<T>(IComputeClient ComputeClient, string NamespaceId, IoHash Hash)
-		{
-			byte[] Data = await ComputeClient.GetObjectAsync(NamespaceId, Hash);
-			return CbSerializer.Deserialize<T>(new CbField(Data));
-		}
-
-		static async Task WriteObjectToFileAsync(IComputeClient ComputeClient, string NamespaceId, IoHash Hash, FileReference OutputFile)
+		static async Task WriteObjectToFileAsync(IStorageClient StorageClient, NamespaceId NamespaceId, IoHash Hash, FileReference OutputFile)
 		{
 			DirectoryReference.CreateDirectory(OutputFile.Directory);
 
-			byte[] Data = await ComputeClient.GetBlobAsync(NamespaceId, Hash);
+			byte[] Data = await StorageClient.ReadBlobToMemoryAsync(NamespaceId, Hash);
 			await FileReference.WriteAllBytesAsync(OutputFile, Data);
 		}
 	}
diff --git a/Engine/Source/Programs/Horde/Horde.Agent/Commands/Service/RunCommand.cs b/Engine/Source/Programs/Horde/Horde.Agent/Commands/Service/RunCommand.cs
--- a/Engine/Source/Programs/Horde/Horde.Agent/Commands/Service/RunCommand.cs
+++ b/Engine/Source/Programs/Horde/Horde.Agent/Commands/Service/RunCommand.cs
@@ -23,6 +23,9 @@ using OpenTracing;
 using OpenTracing.Util;
 using Datadog.Trace.Configuration;
 using Datadog.Trace;
+using EpicGames.Horde.Storage;
+using EpicGames.Horde.Storage.Impl;
+using EpicGames.Horde.Compute;
 
 namespace HordeAgent.Modes.Service
 {
@@ -126,6 +129,8 @@ namespace HordeAgent.Modes.Service
 						return Builder.WaitAndRetryAsync(new[] { TimeSpan.FromSeconds(1), TimeSpan.FromSeconds(5), TimeSpan.FromSeconds(10) });
 					});
 
+					Services.AddHordeStorage(Settings => ConfigSection.GetCurrentServerProfile().GetSection(nameof(ServerProfile.Storage)).Bind(Settings));
+
 					Services.AddSingleton<GrpcService>();
 					Services.AddHostedService<WorkerService>();
 				});
diff --git a/Engine/Source/Programs/Horde/Horde.Agent/Execution/ComputeTaskExecutor.cs b/Engine/Source/Programs/Horde/Horde.Agent/Execution/ComputeTaskExecutor.cs
--- a/Engine/Source/Programs/Horde/Horde.Agent/Execution/ComputeTaskExecutor.cs
+++ b/Engine/Source/Programs/Horde/Horde.Agent/Execution/ComputeTaskExecutor.cs
@@ -42,7 +42,7 @@ namespace HordeAgent
 		/// <summary>
 		/// The blob store
 		/// </summary>
-		BlobStore.BlobStoreClient BlobStore;
+		IStorageClient StorageClient;
 
 		/// <summary>
 		/// Logger for internal executor output
@@ -52,11 +52,11 @@ namespace HordeAgent
 		/// <summary>
 		/// Constructor
 		/// </summary>
-		/// <param name="GrpcChannel"></param>
+		/// <param name="StorageClient"></param>
 		/// <param name="Logger"></param>
-		public ComputeTaskExecutor(GrpcChannel GrpcChannel, ILogger Logger)
+		public ComputeTaskExecutor(IStorageClient StorageClient, ILogger Logger)
 		{
-			this.BlobStore = new BlobStore.BlobStoreClient(GrpcChannel);
+			this.StorageClient = StorageClient;
 			this.Logger = Logger;
 		}
 
@@ -102,15 +102,19 @@ namespace HordeAgent
 		/// <returns>The action result</returns>
 		public async Task<ComputeTaskResultMessage> ExecuteAsync(string LeaseId, ComputeTaskMessage ComputeTaskMessage, DirectoryReference SandboxDir, CancellationToken CancellationToken)
 		{
-			ComputeTask Task = await BlobStore.GetObjectAsync<ComputeTask>(ComputeTaskMessage.NamespaceId, ComputeTaskMessage.Task);
-			Logger.LogInformation("Executing task {Hash} for lease ID {LeaseId}", ComputeTaskMessage.Task.Hash, LeaseId);
+			NamespaceId NamespaceId = new NamespaceId(ComputeTaskMessage.NamespaceId);
+			BucketId InputBucketId = new BucketId(ComputeTaskMessage.InputBucketId);
+			BucketId OutputBucketId = new BucketId(ComputeTaskMessage.OutputBucketId);
+
+			ComputeTask Task = await StorageClient.GetRefAsync<ComputeTask>(NamespaceId, InputBucketId, ComputeTaskMessage.TaskRefId.AsRefId());
+			Logger.LogInformation("Executing task {Hash} for lease ID {LeaseId}", ComputeTaskMessage.TaskRefId, LeaseId);
 
 			DirectoryReference.CreateDirectory(SandboxDir);
 			FileUtils.ForceDeleteDirectoryContents(SandboxDir);
 
 			DateTimeOffset InputFetchStart = DateTimeOffset.UtcNow;
-			DirectoryTree InputDirectory = await BlobStore.GetObjectAsync<DirectoryTree>(ComputeTaskMessage.NamespaceId, Task.SandboxHash);
-			await SetupSandboxAsync(ComputeTaskMessage.NamespaceId, InputDirectory, SandboxDir);
+			DirectoryTree InputDirectory = await StorageClient.ReadObjectAsync<DirectoryTree>(NamespaceId, Task.SandboxHash);
+			await SetupSandboxAsync(NamespaceId, InputDirectory, SandboxDir);
 			DateTimeOffset InputFetchCompleted = DateTimeOffset.UtcNow;
 
 			using (ManagedProcessGroup ProcessGroup = new ManagedProcessGroup())
@@ -149,8 +153,8 @@ namespace HordeAgent
 					Logger.LogInformation("exit: {ExitCode}", Process.ExitCode);
 
 					ComputeTaskResult Result = new ComputeTaskResult(Process.ExitCode);
-					Result.StdOutHash = await BlobStore.PutBlobAsync(ComputeTaskMessage.NamespaceId, StdOutData);
-					Result.StdErrHash = await BlobStore.PutBlobAsync(ComputeTaskMessage.NamespaceId, StdErrData);
+					Result.StdOutHash = await StorageClient.WriteBlobFromMemoryAsync(NamespaceId, StdOutData);
+					Result.StdErrHash = await StorageClient.WriteBlobFromMemoryAsync(NamespaceId, StdErrData);
 
 					FileReference[] OutputFiles = ResolveOutputPaths(SandboxDir, Task.OutputPaths.Select(x => x.ToString())).OrderBy(x => x.FullName, StringComparer.Ordinal).ToArray();
 					if (OutputFiles.Length > 0)
@@ -159,11 +163,13 @@ namespace HordeAgent
 						{
 							Logger.LogInformation("output: {File}", OutputFile.MakeRelativeTo(SandboxDir));
 						}
-						Result.OutputHash = await PutOutput(ComputeTaskMessage.NamespaceId, SandboxDir, OutputFiles);
+						Result.OutputHash = await PutOutput(NamespaceId, SandboxDir, OutputFiles);
 					}
 
-					CbObjectAttachment Attachment = await BlobStore.PutObjectAsync(ComputeTaskMessage.NamespaceId, Result);
-					return new ComputeTaskResultMessage(Attachment);
+					CbObject ResultObject = CbSerializer.Serialize(Result);
+
+					await StorageClient.SetRefAsync(NamespaceId, OutputBucketId, ComputeTaskMessage.TaskRefId, ResultObject);
+					return new ComputeTaskResultMessage(ComputeTaskMessage.TaskRefId);
 				}
 			}
 		}
@@ -175,7 +181,7 @@ namespace HordeAgent
 		/// <param name="InputDirectory">The directory spec</param>
 		/// <param name="OutputDir">Output directory on disk</param>
 		/// <returns>Async task</returns>
-		async Task SetupSandboxAsync(string NamespaceId, DirectoryTree InputDirectory, DirectoryReference OutputDir)
+		async Task SetupSandboxAsync(NamespaceId NamespaceId, DirectoryTree InputDirectory, DirectoryReference OutputDir)
 		{
 			DirectoryReference.CreateDirectory(OutputDir);
 
@@ -183,14 +189,14 @@ namespace HordeAgent
 			{
 				FileReference File = FileReference.Combine(OutputDir, FileNode.Name.ToString());
 				Logger.LogInformation("Downloading {File} (digest: {Digest})", File, FileNode.Hash);
-				byte[] Data = await BlobStore.GetBlobAsync(NamespaceId, FileNode.Hash);
+				byte[] Data = await StorageClient.ReadBlobToMemoryAsync(NamespaceId, FileNode.Hash);
 				Logger.LogInformation("Writing {File} (digest: {Digest})", File, FileNode.Hash);
 				await FileReference.WriteAllBytesAsync(File, Data);
 			}
 
 			async Task DownloadDir(DirectoryNode DirectoryNode)
 			{
-				DirectoryTree InputSubDirectory = await BlobStore.GetObjectAsync<DirectoryTree>(NamespaceId, DirectoryNode.Hash);
+				DirectoryTree InputSubDirectory = await StorageClient.ReadObjectAsync<DirectoryTree>(NamespaceId, DirectoryNode.Hash);
 				DirectoryReference OutputSubDirectory = DirectoryReference.Combine(OutputDir, DirectoryNode.Name.ToString());
 				await SetupSandboxAsync(NamespaceId, InputSubDirectory, OutputSubDirectory);
 			}
@@ -201,14 +207,14 @@ namespace HordeAgent
 			await Task.WhenAll(Tasks);
 		}
 
-		async Task<IoHash> PutOutput(string NamespaceId, DirectoryReference BaseDir, IEnumerable<FileReference> Files)
+		async Task<IoHash> PutOutput(NamespaceId NamespaceId, DirectoryReference BaseDir, IEnumerable<FileReference> Files)
 		{
 			List<FileReference> SortedFiles = Files.OrderBy(x => x.FullName, StringComparer.Ordinal).ToList();
 			(_, IoHash Hash) = await PutDirectoryTree(NamespaceId, BaseDir.FullName.Length, SortedFiles, 0, SortedFiles.Count);
 			return Hash;
 		}
 
-		async Task<(DirectoryTree, IoHash)> PutDirectoryTree(string NamespaceId, int BaseDirLen, List<FileReference> SortedFiles, int MinIdx, int MaxIdx)
+		async Task<(DirectoryTree, IoHash)> PutDirectoryTree(NamespaceId NamespaceId, int BaseDirLen, List<FileReference> SortedFiles, int MinIdx, int MaxIdx)
 		{
 			List<Task<FileNode>> Files = new List<Task<FileNode>>();
 			List<Task<DirectoryNode>> Trees = new List<Task<DirectoryNode>>();
@@ -246,19 +252,21 @@ namespace HordeAgent
 			DirectoryTree Tree = new DirectoryTree();
 			Tree.Files.AddRange(await Task.WhenAll(Files));
 			Tree.Directories.AddRange(await Task.WhenAll(Trees));
-			return (Tree, await BlobStore.PutObjectAsync(NamespaceId, Tree));
+
+			IoHash Hash = await StorageClient.WriteObjectAsync<DirectoryTree>(NamespaceId, Tree);
+			return (Tree, Hash);
 		}
 
-		async Task<DirectoryNode> CreateDirectoryNode(string NamespaceId, string Name, int BaseDirLen, List<FileReference> SortedFiles, int MinIdx, int MaxIdx)
+		async Task<DirectoryNode> CreateDirectoryNode(NamespaceId NamespaceId, string Name, int BaseDirLen, List<FileReference> SortedFiles, int MinIdx, int MaxIdx)
 		{
 			(_, IoHash Hash) = await PutDirectoryTree(NamespaceId, BaseDirLen, SortedFiles, MinIdx, MaxIdx);
 			return new DirectoryNode(Name, Hash);
 		}
 
-		async Task<FileNode> CreateFileNode(string NamespaceId, string Name, FileReference File)
+		async Task<FileNode> CreateFileNode(NamespaceId NamespaceId, string Name, FileReference File)
 		{
 			byte[] Data = await FileReference.ReadAllBytesAsync(File);
-			IoHash Hash = await BlobStore.PutBlobAsync(NamespaceId, Data);
+			IoHash Hash = await StorageClient.WriteBlobFromMemoryAsync(NamespaceId, Data);
 			return new FileNode(Name, Hash, Data.Length, (int)FileReference.GetAttributes(File));
 		}
 
diff --git a/Engine/Source/Programs/Horde/Horde.Agent/Horde.Agent.csproj b/Engine/Source/Programs/Horde/Horde.Agent/Horde.Agent.csproj
--- a/Engine/Source/Programs/Horde/Horde.Agent/Horde.Agent.csproj
+++ b/Engine/Source/Programs/Horde/Horde.Agent/Horde.Agent.csproj
@@ -34,8 +34,8 @@
     <PackageReference Include="Microsoft.Extensions.Hosting" Version="3.1.5" />
     <PackageReference Include="Microsoft.Extensions.Hosting.Abstractions" Version="3.1.5" />
     <PackageReference Include="Microsoft.Extensions.Hosting.WindowsServices" Version="3.1.5" />
-    <PackageReference Include="Microsoft.Extensions.Http" Version="3.1.5" />
-    <PackageReference Include="Microsoft.Extensions.Http.Polly" Version="3.1.5" />
+    <PackageReference Include="Microsoft.Extensions.Http" Version="3.1.22" />
+    <PackageReference Include="Microsoft.Extensions.Http.Polly" Version="3.1.22" />
     <PackageReference Include="Microsoft.Extensions.Logging.Abstractions" Version="5.0.0" />
     <PackageReference Include="Microsoft.Extensions.Options.DataAnnotations" Version="3.1.5" />
     <PackageReference Include="OpenTracing" Version="0.12.1" />
@@ -90,12 +90,14 @@
     <Compile Include="../../../../Restricted/*/Source/Programs/Horde/Horde.Agent/**/*.cs">
       <Link>$([System.Text.RegularExpressions.Regex]::Replace(%(FullPath), '^.+?[\\/]Horde.Agent[\\/]', ''))</Link>
     </Compile>
+    <Compile Include="..\Horde.Build\Utilities\ProtoExtensions.cs">
+      <Link>Utility\ProtoExtensions.cs</Link>
+    </Compile>
   </ItemGroup>
 
   <ItemGroup>
     <ProjectReference Include="..\..\Shared\EpicGames.Core\EpicGames.Core.csproj" />
     <ProjectReference Include="..\..\Shared\EpicGames.Horde\EpicGames.Horde.csproj" />
-    <ProjectReference Include="..\..\Shared\EpicGames.Horde.Client\EpicGames.Horde.Client.csproj" />
     <ProjectReference Include="..\..\Shared\EpicGames.Perforce.Managed\EpicGames.Perforce.Managed.csproj" />
     <ProjectReference Include="..\..\Shared\EpicGames.Perforce\EpicGames.Perforce.csproj" />
   </ItemGroup>
@@ -105,12 +107,14 @@
     <Protobuf Include="..\Horde.Build\Protos\horde\horde_rpc_messages.proto" ProtoRoot="..\Horde.Build\Protos" Access="Internal" />
     <Protobuf Include="..\Horde.Build\Protos\horde\common\agent_status.proto" ProtoRoot="..\Horde.Build\Protos" Access="Internal" />
     <Protobuf Include="..\Horde.Build\Protos\horde\common\event_severity.proto" ProtoRoot="..\Horde.Build\Protos" Access="Internal" />
+    <Protobuf Include="..\Horde.Build\Protos\horde\common\iohash_wrapper.proto" ProtoRoot="..\Horde.Build\Protos" />
     <Protobuf Include="..\Horde.Build\Protos\horde\common\jobstep_outcome.proto" ProtoRoot="..\Horde.Build\Protos" Access="Internal" />
     <Protobuf Include="..\Horde.Build\Protos\horde\common\jobstep_state.proto" ProtoRoot="..\Horde.Build\Protos" Access="Internal" />
     <Protobuf Include="..\Horde.Build\Protos\horde\common\label_change.proto" ProtoRoot="..\Horde.Build\Protos" Access="Internal" />
     <Protobuf Include="..\Horde.Build\Protos\horde\common\lease_outcome.proto" ProtoRoot="..\Horde.Build\Protos" Access="Internal" />
     <Protobuf Include="..\Horde.Build\Protos\horde\common\lease_state.proto" ProtoRoot="..\Horde.Build\Protos" Access="Internal" />
     <Protobuf Include="..\Horde.Build\Protos\horde\common\priority.proto" ProtoRoot="..\Horde.Build\Protos" Access="Internal" />
+    <Protobuf Include="..\Horde.Build\Protos\horde\common\refid_wrapper.proto" ProtoRoot="..\Horde.Build\Protos" />
     <Protobuf Include="..\Horde.Build\Protos\horde\messages\agent_capabilities.proto" ProtoRoot="..\Horde.Build\Protos" Access="Internal" />
     <Protobuf Include="..\Horde.Build\Protos\horde\messages\agent_requirements.proto" ProtoRoot="..\Horde.Build\Protos" Access="Internal" />
     <Protobuf Include="..\Horde.Build\Protos\horde\messages\agent_workspace.proto" ProtoRoot="..\Horde.Build\Protos" Access="Internal" />
diff --git a/Engine/Source/Programs/Horde/Horde.Agent/Services/WorkerService.cs b/Engine/Source/Programs/Horde/Horde.Agent/Services/WorkerService.cs
--- a/Engine/Source/Programs/Horde/Horde.Agent/Services/WorkerService.cs
+++ b/Engine/Source/Programs/Horde/Horde.Agent/Services/WorkerService.cs
@@ -41,6 +41,7 @@ using Status = Grpc.Core.Status;
 using EpicGames.Horde.Storage;
 using EpicGames.Horde.Compute;
 using System.Net;
+using EpicGames.Horde.Storage.Impl;
 
 namespace HordeAgent.Services
 {
@@ -160,6 +161,11 @@ namespace HordeAgent.Services
 		/// </summary>
 		GrpcService GrpcService;
 
+		/// <summary>
+		/// Client interface to the storage system
+		/// </summary>
+		IStorageClient StorageClient;
+
 		/// <summary>
 		/// The working directory
 		/// </summary>
@@ -222,13 +228,15 @@ namespace HordeAgent.Services
 		/// <param name="Logger">Log sink</param>
 		/// <param name="OptionsMonitor">The current settings</param>
 		/// <param name="GrpcService">Instance of the Grpc service</param>
+		/// <param name="StorageClient">Instance of the storage client</param>
 		/// <param name="CreateExecutor"></param>
-		public WorkerService(ILogger<WorkerService> Logger, IOptionsMonitor<AgentSettings> OptionsMonitor, GrpcService GrpcService, Func<IRpcConnection, ExecuteJobTask, BeginBatchResponse, IExecutor>? CreateExecutor = null)
+		public WorkerService(ILogger<WorkerService> Logger, IOptionsMonitor<AgentSettings> OptionsMonitor, GrpcService GrpcService, IStorageClient StorageClient, Func<IRpcConnection, ExecuteJobTask, BeginBatchResponse, IExecutor>? CreateExecutor = null)
 		{
 			this.Logger = Logger;
 			this.Settings = OptionsMonitor.CurrentValue;
 			this.ServerProfile = Settings.GetCurrentServerProfile();
 			this.GrpcService = GrpcService;
+			this.StorageClient = StorageClient;
 
 			if (Settings.WorkingDir == null)
 			{
@@ -911,7 +919,7 @@ namespace HordeAgent.Services
 					DirectoryReference LeaseDir = DirectoryReference.Combine(WorkingDir, "Compute", LeaseId);
 					DirectoryReference.CreateDirectory(LeaseDir);
 
-					ComputeTaskExecutor Executor = new ComputeTaskExecutor(Client.Channel, Logger);
+					ComputeTaskExecutor Executor = new ComputeTaskExecutor(StorageClient, Logger);
 					try
 					{
 						Result = await Executor.ExecuteAsync(LeaseId, ComputeTask, LeaseDir, CancellationToken);
diff --git a/Engine/Source/Programs/Horde/Horde.Agent/Utilities/ProtoMethods.cs b/Engine/Source/Programs/Horde/Horde.Agent/Utilities/ProtoMethods.cs
--- a/Engine/Source/Programs/Horde/Horde.Agent/Utilities/ProtoMethods.cs
+++ b/Engine/Source/Programs/Horde/Horde.Agent/Utilities/ProtoMethods.cs
@@ -1,5 +1,6 @@
 // Copyright Epic Games, Inc. All Rights Reserved.
 
+using EpicGames.Horde.Common;
 using EpicGames.Horde.Compute;
 using EpicGames.Serialization;
 using System;
@@ -16,9 +17,9 @@ namespace HordeCommon.Rpc.Tasks
 		/// <summary>
 		/// Constructor
 		/// </summary>
-		public ComputeTaskResultMessage(CbObjectAttachment Result)
+		public ComputeTaskResultMessage(RefIdWrapper ResultRefId)
 		{
-			this.Result = Result;
+			this.ResultRefId = ResultRefId;
 		}
 
 		/// <summary>
@@ -26,7 +27,7 @@ namespace HordeCommon.Rpc.Tasks
 		/// </summary>
 		public ComputeTaskResultMessage(ComputeTaskOutcome Outcome, string? Detail = null)
 		{
-			this.Outcome = Outcome;
+			this.Outcome = (int)Outcome;
 			this.Detail = Detail;
 		}
 	}
diff --git a/Engine/Source/Programs/Horde/Horde.Agent/appsettings.json b/Engine/Source/Programs/Horde/Horde.Agent/appsettings.json
--- a/Engine/Source/Programs/Horde/Horde.Agent/appsettings.json
+++ b/Engine/Source/Programs/Horde/Horde.Agent/appsettings.json
@@ -7,7 +7,10 @@
 			{
 				"Name": "Local",
 				"Environment": "dev",
-				"Url": "https://localhost:5001"
+				"Url": "https://localhost:5001",
+				"Storage": {
+					"Url": "https://localhost:57001"
+				}
 			}
 		]
 	}
diff --git a/Engine/Source/Programs/Horde/Horde.Build.Tests/KeyValueTreeTests.cs b/Engine/Source/Programs/Horde/Horde.Build.Tests/KeyValueTreeTests.cs
deleted file mode 100644
diff --git a/Engine/Source/Programs/Horde/Horde.Build.Tests/StorageTests.cs b/Engine/Source/Programs/Horde/Horde.Build.Tests/StorageTests.cs
--- a/Engine/Source/Programs/Horde/Horde.Build.Tests/StorageTests.cs
+++ b/Engine/Source/Programs/Horde/Horde.Build.Tests/StorageTests.cs
@@ -1,9 +1,9 @@
 // Copyright Epic Games, Inc. All Rights Reserved.
 
 using EpicGames.Core;
+using EpicGames.Horde.Storage;
 using EpicGames.Serialization;
 using HordeServer.Services;
-using HordeServer.Storage;
 using HordeServer.Utilities;
 using Microsoft.Extensions.DependencyInjection;
 using Microsoft.VisualStudio.TestTools.UnitTesting;
@@ -15,9 +15,6 @@ using System.Threading.Tasks;
 
 namespace HordeServerTests
 {
-	using NamespaceId = StringId<INamespace>;
-	using BucketId = StringId<IBucket>;
-
 	[TestClass]
 	public class StorageTests : TestSetup
 	{
@@ -26,26 +23,22 @@ namespace HordeServerTests
 		{
 			NamespaceId NamespaceId = new NamespaceId("test-ns");
 
-			IBlobCollection BlobCollection = ServiceProvider.GetRequiredService<IBlobCollection>();
+			IStorageClient StorageClient = ServiceProvider.GetRequiredService<IStorageClient>();
 
 			byte[] TestData = Encoding.UTF8.GetBytes("Hello world");
 			IoHash Hash = IoHash.Compute(TestData);
 
-			Assert.IsNull(await BlobCollection.TryReadStreamAsync(NamespaceId, Hash));
-			Assert.IsFalse(await BlobCollection.ExistsAsync(NamespaceId, Hash));
+			await Assert.ThrowsExceptionAsync<BlobNotFoundException>(() => StorageClient.ReadBlobAsync(NamespaceId, Hash));
 
-			IoHash ReturnedHash = await BlobCollection.WriteBytesAsync(NamespaceId, TestData);
+			IoHash ReturnedHash = await StorageClient.WriteBlobFromMemoryAsync(NamespaceId, TestData);
 			Assert.AreEqual(Hash, ReturnedHash);
-			Assert.IsTrue(await BlobCollection.ExistsAsync(NamespaceId, Hash));
+			Assert.IsNotNull(await StorageClient.ReadBlobToMemoryAsync(NamespaceId, Hash));
 
-			ReadOnlyMemory<byte>? StoredData = await BlobCollection.TryReadBytesAsync(NamespaceId, Hash);
-			Assert.IsNotNull(StoredData);
-			Assert.IsTrue(TestData.AsSpan().SequenceEqual(StoredData.Value.Span));
-			Assert.IsTrue(await BlobCollection.ExistsAsync(NamespaceId, Hash));
+			ReadOnlyMemory<byte> StoredData = await StorageClient.ReadBlobToMemoryAsync(NamespaceId, Hash);
+			Assert.IsTrue(TestData.AsSpan().SequenceEqual(StoredData.Span));
 
 			NamespaceId OtherNamespaceId = new NamespaceId("other-ns");
-			Assert.IsNull(await BlobCollection.TryReadStreamAsync(OtherNamespaceId, Hash));
-			Assert.IsFalse(await BlobCollection.ExistsAsync(OtherNamespaceId, Hash));
+			await Assert.ThrowsExceptionAsync<BlobNotFoundException>(() => StorageClient.ReadBlobAsync(OtherNamespaceId, Hash));
 		}
 
 		[TestMethod]
@@ -53,7 +46,7 @@ namespace HordeServerTests
 		{
 			NamespaceId NamespaceId = new NamespaceId("test-ns");
 
-			IObjectCollection ObjectCollection = ServiceProvider.GetRequiredService<IObjectCollection>();
+			IStorageClient StorageClient = ServiceProvider.GetRequiredService<IStorageClient>();
 
 			CbObject ObjectD = CbObject.Build(Writer =>
 			{
@@ -61,24 +54,18 @@ namespace HordeServerTests
 			});
 			IoHash HashD = ObjectD.GetHash();
 
-			Assert.IsFalse(await ObjectCollection.ExistsAsync(NamespaceId, HashD));
-			Assert.IsNull(await ObjectCollection.GetAsync(NamespaceId, HashD));
-
-			await ObjectCollection.AddAsync(NamespaceId, HashD, ObjectD);
+			await Assert.ThrowsExceptionAsync<BlobNotFoundException>(() => StorageClient.ReadBlobToMemoryAsync(NamespaceId, HashD));
 
-			Assert.IsTrue(await ObjectCollection.ExistsAsync(NamespaceId, HashD));
+			await StorageClient.WriteBlobFromMemoryAsync(NamespaceId, HashD, ObjectD.GetView());
 
-			CbObject? ReturnedObjectD = await ObjectCollection.GetAsync(NamespaceId, HashD);
-			Assert.IsNotNull(ReturnedObjectD);
+			CbObject ReturnedObjectD = new CbObject(await StorageClient.ReadBlobToMemoryAsync(NamespaceId, HashD));
 			Assert.IsTrue(ReturnedObjectD!.GetView().Span.SequenceEqual(ObjectD.GetView().Span));
 		}
 
 		[TestMethod]
 		public async Task RefCollectionTest()
 		{
-			IBlobCollection BlobCollection = ServiceProvider.GetRequiredService<IBlobCollection>();
-			IObjectCollection ObjectCollection = ServiceProvider.GetRequiredService<IObjectCollection>();
-			IRefCollection RefCollection = ServiceProvider.GetRequiredService<IRefCollection>();
+			IStorageClient StorageClient = ServiceProvider.GetRequiredService<IStorageClient>();
 
 			byte[] BlobA = Encoding.UTF8.GetBytes("This is blob A");
 			IoHash HashA = IoHash.Compute(BlobA);
@@ -106,53 +93,46 @@ namespace HordeServerTests
 
 			NamespaceId NamespaceId = new NamespaceId("test-ns");
 			BucketId BucketId = new BucketId("test-bkt");
-			const string RefName = "refname";
+			RefId RefId = new RefId("refname");
 
 			// Check that setting the ref to E fails with the correct missing hashes
-			List<IoHash> MissingHashes = await RefCollection.SetAsync(NamespaceId, BucketId, RefName, ObjectE);
+			List<IoHash> MissingHashes = await StorageClient.TrySetRefAsync(NamespaceId, BucketId, RefId, ObjectE);
 			Assert.AreEqual(2, MissingHashes.Count);
 			Assert.IsTrue(MissingHashes.Contains(HashA));
 			Assert.IsTrue(MissingHashes.Contains(HashD));
 
-			IRef? Ref = await RefCollection.GetAsync(NamespaceId, BucketId, RefName);
-			Assert.IsNotNull(Ref);
-			Assert.IsFalse(Ref!.Finalized);
-			Assert.IsTrue(Ref!.Value.GetView().Span.SequenceEqual(ObjectE.GetView().Span));
+			await Assert.ThrowsExceptionAsync<RefNotFoundException>(() => StorageClient.GetRefAsync(NamespaceId, BucketId, RefId));
 
 			// Add object D, and check that changes the missing hashes to just the blobs
-			await ObjectCollection.AddAsync(NamespaceId, HashD, ObjectD);
+			await StorageClient.WriteBlobFromMemoryAsync(NamespaceId, HashD, ObjectD.GetView());
 
-			MissingHashes = await RefCollection.FinalizeAsync(Ref);
+			MissingHashes = await StorageClient.TryFinalizeRefAsync(NamespaceId, BucketId, RefId, HashE);
 			Assert.AreEqual(3, MissingHashes.Count);
 			Assert.IsTrue(MissingHashes.Contains(HashA));
 			Assert.IsTrue(MissingHashes.Contains(HashB));
 			Assert.IsTrue(MissingHashes.Contains(HashC));
 
-			Ref = await RefCollection.GetAsync(NamespaceId, BucketId, RefName);
-			Assert.IsNotNull(Ref);
-			Assert.IsFalse(Ref!.Finalized);
+			await Assert.ThrowsExceptionAsync<RefNotFoundException>(() => StorageClient.GetRefAsync(NamespaceId, BucketId, RefId));
 
 			// Add blobs A, B and C and check that the object can be finalized
-			await BlobCollection.WriteBytesAsync(NamespaceId, HashA, BlobA);
-			await BlobCollection.WriteBytesAsync(NamespaceId, HashB, BlobB);
-			await BlobCollection.WriteBytesAsync(NamespaceId, HashC, BlobC);
+			await StorageClient.WriteBlobFromMemoryAsync(NamespaceId, HashA, BlobA);
+			await StorageClient.WriteBlobFromMemoryAsync(NamespaceId, HashB, BlobB);
+			await StorageClient.WriteBlobFromMemoryAsync(NamespaceId, HashC, BlobC);
 
-			MissingHashes = await RefCollection.FinalizeAsync(Ref);
+			MissingHashes = await StorageClient.TryFinalizeRefAsync(NamespaceId, BucketId, RefId, HashE);
 			Assert.AreEqual(0, MissingHashes.Count);
 
-			Ref = await RefCollection.GetAsync(NamespaceId, BucketId, RefName);
+			IRef Ref = await StorageClient.GetRefAsync(NamespaceId, BucketId, RefId);
 			Assert.IsNotNull(Ref);
-			Assert.IsTrue(Ref!.Finalized);
 
 			// Add a new ref to existing objects and check it finalizes correctly 
-			const string RefName2 = "refname";
+			RefId RefId2 = new RefId("refname");
 
-			MissingHashes = await RefCollection.SetAsync(NamespaceId, BucketId, RefName2, ObjectE);
+			MissingHashes = await StorageClient.TrySetRefAsync(NamespaceId, BucketId, RefId2, ObjectE);
 			Assert.AreEqual(0, MissingHashes.Count);
 
-			Ref = await RefCollection.GetAsync(NamespaceId, BucketId, RefName2);
+			Ref = await StorageClient.GetRefAsync(NamespaceId, BucketId, RefId2);
 			Assert.IsNotNull(Ref);
-			Assert.IsTrue(Ref!.Finalized);
 		}
 	}
 }
diff --git a/Engine/Source/Programs/Horde/Horde.Build.Tests/TestSetup.cs b/Engine/Source/Programs/Horde/Horde.Build.Tests/TestSetup.cs
--- a/Engine/Source/Programs/Horde/Horde.Build.Tests/TestSetup.cs
+++ b/Engine/Source/Programs/Horde/Horde.Build.Tests/TestSetup.cs
@@ -7,7 +7,9 @@ using System.Security.Claims;
 using System.Threading.Tasks;
 using Datadog.Trace;
 using EpicGames.Core;
+using EpicGames.Horde.Storage;
 using HordeCommon;
+using Horde.Build.Storage.Services;
 using HordeServer;
 using HordeServer.Collections;
 using HordeServer.Collections.Impl;
@@ -24,8 +26,6 @@ using HordeServer.Services;
 using HordeServer.Services.Impl;
 using HordeServer.Storage;
 using HordeServer.Storage.Backends;
-using HordeServer.Storage.Collections;
-using HordeServer.Storage.Services;
 using HordeServer.Tasks.Impl;
 using HordeServer.Utilities;
 using HordeServerTests.Stubs.Collections;
@@ -200,18 +200,12 @@ namespace HordeServerTests
 
 			Services.AddSingleton<ConformTaskSource>();
 
-			Services.AddSingleton<IBlobCollection, BlobCollection>();
-			Services.AddSingleton<IObjectCollection, ObjectCollection>();
-			Services.AddSingleton<IRefCollection, RefCollection>();
-			Services.AddSingleton<INamespaceCollection, NamespaceCollection>();
-			Services.AddSingleton<IBucketCollection, BucketCollection>();
-
 			Services.AddSingleton(new Startup.StorageBackendSettings<PersistentLogStorage> { Type = StorageProviderType.Transient });
 			Services.AddSingleton(new Startup.StorageBackendSettings<ArtifactCollection> { Type = StorageProviderType.Transient });
-			Services.AddSingleton(new Startup.StorageBackendSettings<BlobCollection> { Type = StorageProviderType.Transient });
+			Services.AddSingleton(new Startup.StorageBackendSettings<BasicStorageClient> { Type = StorageProviderType.Transient });
 			Services.AddSingleton(typeof(IStorageBackend<>), typeof(Startup.StorageBackendFactory<>));
 
-			Services.AddSingleton<IStorageService, SimpleStorageService>();
+			Services.AddSingleton<IStorageClient, BasicStorageClient>();
 
 			Services.AddSingleton<ISingletonDocument<AgentSoftwareChannels>>(new SingletonDocumentStub<AgentSoftwareChannels>());
 		}
diff --git a/Engine/Source/Programs/Horde/Horde.Build/Acls/Acl.cs b/Engine/Source/Programs/Horde/Horde.Build/Acls/Acl.cs
--- a/Engine/Source/Programs/Horde/Horde.Build/Acls/Acl.cs
+++ b/Engine/Source/Programs/Horde/Horde.Build/Acls/Acl.cs
@@ -437,6 +437,27 @@ namespace HordeServer.Models
 
 
 
+
+
+		//// COMPUTE ////
+
+		/// <summary>
+		/// User can add tasks to the compute cluster
+		/// </summary>
+		AddComputeTasks,
+
+		/// <summary>
+		/// User can poll for compute results
+		/// </summary>
+		ViewComputeTasks,
+
+
+
+
+
+
+
+
 		//// STORAGE ////
 
 		/// <summary>
diff --git a/Engine/Source/Programs/Horde/Horde.Build/Compute/IComputeService.cs b/Engine/Source/Programs/Horde/Horde.Build/Compute/IComputeService.cs
--- a/Engine/Source/Programs/Horde/Horde.Build/Compute/IComputeService.cs
+++ b/Engine/Source/Programs/Horde/Horde.Build/Compute/IComputeService.cs
@@ -2,6 +2,7 @@
 
 using EpicGames.Core;
 using EpicGames.Horde.Compute;
+using EpicGames.Horde.Storage;
 using EpicGames.Serialization;
 using HordeServer.Models;
 using HordeServer.Storage;
@@ -18,17 +19,7 @@ using System.Threading.Tasks;
 
 namespace HordeServer.Compute
 {
-	using ChannelId = StringId<IComputeChannel>;
 	using LeaseId = ObjectId<ILease>;
-	using NamespaceId = StringId<INamespace>;
-
-	/// <summary>
-	/// Unused, except for ChannelId alias.
-	/// </summary>
-	[SuppressMessage("Design", "CA1040:Avoid empty interfaces")]
-	public interface IComputeChannel
-	{
-	}
 
 	/// <summary>
 	/// Status of a compute task
@@ -36,9 +27,9 @@ namespace HordeServer.Compute
 	public interface IComputeTaskStatus
 	{
 		/// <summary>
-		/// The input hash
+		/// The input task ref id
 		/// </summary>
-		public CbObjectAttachment Task { get; }
+		public RefId TaskRefId { get; }
 
 		/// <summary>
 		/// Timestamp for the event
@@ -68,7 +59,7 @@ namespace HordeServer.Compute
 		/// <summary>
 		/// Hash of the result, if complete
 		/// </summary>
-		public CbObjectAttachment? Result { get; }
+		public RefId? ResultRefId { get; }
 
 		/// <summary>
 		/// Additional information for the given outcome
@@ -81,49 +72,37 @@ namespace HordeServer.Compute
 	/// </summary>
 	public interface IComputeService : ITaskSource
 	{
+		/// <summary>
+		/// Gets information about a compute cluster
+		/// </summary>
+		/// <param name="ClusterId">Cluster to use for execution</param>
+		Task<IComputeClusterInfo> GetClusterInfoAsync(ClusterId ClusterId);
+
 		/// <summary>
 		/// Post tasks to be executed to a channel
 		/// </summary>
-		/// <param name="NamespaceId">Namespace of referenced blobs</param>
-		/// <param name="RequirementsHash">Hash of the requirements document for execution</param>
-		/// <param name="TaskHashes">List of task hashes</param>
+		/// <param name="ClusterId">Cluster to use for execution</param>
 		/// <param name="ChannelId">Unique identifier of the client</param>
+		/// <param name="TaskRefIds">List of task hashes</param>
+		/// <param name="RequirementsHash">Requirements document for execution</param>
 		/// <returns>Async task</returns>
-		Task AddTasksAsync(NamespaceId NamespaceId, CbObjectAttachment RequirementsHash, List<CbObjectAttachment> TaskHashes, ChannelId ChannelId);
+		Task AddTasksAsync(ClusterId ClusterId, ChannelId ChannelId, List<RefId> TaskRefIds, CbObjectAttachment RequirementsHash);
 
 		/// <summary>
 		/// Dequeue completed items from a queue and return immediately
 		/// </summary>
+		/// <param name="ClusterId">Cluster containing the channel</param>
 		/// <param name="ChannelId">Queue to remove items from</param>
 		/// <returns>List of status updates</returns>
-		Task<List<IComputeTaskStatus>> GetTaskUpdatesAsync(ChannelId ChannelId);
+		Task<List<IComputeTaskStatus>> GetTaskUpdatesAsync(ClusterId ClusterId, ChannelId ChannelId);
 
 		/// <summary>
 		/// Dequeue completed items from a queue
 		/// </summary>
+		/// <param name="ClusterId">Cluster containing the channel</param>
 		/// <param name="ChannelId">Queue to remove items from</param>
 		/// <param name="CancellationToken">Cancellation token to stop waiting for items</param>
 		/// <returns>List of status updates</returns>
-		Task<List<IComputeTaskStatus>> WaitForTaskUpdatesAsync(ChannelId ChannelId, CancellationToken CancellationToken);
-	}
-
-	/// <summary>
-	/// Extension methods for <see cref="IComputeService"/>
-	/// </summary>
-	static class ComputeServiceExtensions
-	{
-		/// <summary>
-		/// 
-		/// </summary>
-		/// <param name="ComputeService"></param>
-		/// <param name="NamespaceId"></param>
-		/// <param name="RequirementsHash">Hash of the requirements document for execution</param>
-		/// <param name="TaskHash">The task hash</param>
-		/// <param name="ChannelId"></param>
-		/// <returns></returns>
-		public static Task AddTaskAsync(this IComputeService ComputeService, NamespaceId NamespaceId, CbObjectAttachment RequirementsHash, CbObjectAttachment TaskHash, ChannelId ChannelId)
-		{
-			return ComputeService.AddTasksAsync(NamespaceId, RequirementsHash, new List<CbObjectAttachment> { TaskHash }, ChannelId);
-		}
+		Task<List<IComputeTaskStatus>> WaitForTaskUpdatesAsync(ClusterId ClusterId, ChannelId ChannelId, CancellationToken CancellationToken);
 	}
 }
diff --git a/Engine/Source/Programs/Horde/Horde.Build/Compute/Impl/ComputeController.cs b/Engine/Source/Programs/Horde/Horde.Build/Compute/Impl/ComputeController.cs
--- a/Engine/Source/Programs/Horde/Horde.Build/Compute/Impl/ComputeController.cs
+++ b/Engine/Source/Programs/Horde/Horde.Build/Compute/Impl/ComputeController.cs
@@ -2,6 +2,8 @@
 
 using EpicGames.Core;
 using EpicGames.Horde.Compute;
+using EpicGames.Horde.Compute.Impl;
+using HordeServer.Models;
 using HordeServer.Services;
 using HordeServer.Utilities;
 using Microsoft.AspNetCore.Authorization;
@@ -15,8 +17,6 @@ using System.Threading.Tasks;
 
 namespace HordeServer.Compute.Impl
 {
-	using ChannelId = StringId<IComputeChannel>;
-
 	/// <summary>
 	/// Controller for the /api/v1/compute endpoint
 	/// </summary>
@@ -39,46 +39,82 @@ namespace HordeServer.Compute.Impl
 			this.ComputeService = ComputeService;
 		}
 
+		/// <summary>
+		/// Gets information about a cluster
+		/// </summary>
+		/// <param name="ClusterId">The cluster to add to</param>
+		/// <returns></returns>
+		[HttpGet]
+		[Authorize]
+		[Route("/api/v1/compute/{ClusterId}")]
+		public async Task<ActionResult<GetComputeClusterInfo>> GetClusterInfoAsync([FromRoute] ClusterId ClusterId)
+		{
+			if (!await AclService.AuthorizeAsync(AclAction.ViewComputeTasks, User))
+			{
+				return Forbid();
+			}
+
+			IComputeClusterInfo ClusterInfo = await ComputeService.GetClusterInfoAsync(ClusterId);
+
+			GetComputeClusterInfo Response = new GetComputeClusterInfo();
+			Response.Id = ClusterInfo.Id;
+			Response.NamespaceId = ClusterInfo.NamespaceId;
+			Response.RequestBucketId = ClusterInfo.RequestBucketId;
+			Response.ResponseBucketId = ClusterInfo.ResponseBucketId;
+
+			return Response;
+		}
+
 		/// <summary>
 		/// Add tasks to be executed remotely
 		/// </summary>
-		/// <param name="ChannelId">The channel to add to</param>
+		/// <param name="ClusterId">The cluster to add to</param>
 		/// <param name="Request">The request parameters</param>
 		/// <returns></returns>
 		[HttpPost]
 		[Authorize]
-		[Route("/api/v1/compute/{ChannelId}")]
-		public async Task<ActionResult> AddTasksAsync([FromRoute] ChannelId ChannelId, [FromBody] AddTasksRequest Request)
+		[Route("/api/v1/compute/{ClusterId}")]
+		public async Task<ActionResult> AddTasksAsync([FromRoute] ClusterId ClusterId, [FromBody] AddTasksRequest Request)
 		{
-			if (Request.TaskHashes.Count == 0)
+			if(!await AclService.AuthorizeAsync(AclAction.AddComputeTasks, User))
+			{
+				return Forbid();
+			}
+			if (Request.TaskRefIds.Count == 0)
 			{
 				return BadRequest("No task hashes specified");
 			}
 
-			await ComputeService.AddTasksAsync(Impl.ComputeService.DefaultNamespaceId, Request.RequirementsHash, Request.TaskHashes, ChannelId);
+			await ComputeService.AddTasksAsync(ClusterId, Request.ChannelId, Request.TaskRefIds, Request.RequirementsHash);
 			return Ok();
 		}
 
 		/// <summary>
 		/// Read updates for a particular channel
 		/// </summary>
+		/// <param name="ClusterId"></param>
 		/// <param name="ChannelId">The channel to add to</param>
 		/// <param name="Wait">Amount of time to wait before responding, in seconds</param>
 		/// <returns></returns>
 		[HttpPost]
 		[Authorize]
-		[Route("/api/v1/compute/{ChannelId}/updates")]
-		public async Task<GetTaskUpdatesResponse> GetUpdatesAsync([FromRoute] ChannelId ChannelId, [FromQuery] int Wait = 0)
+		[Route("/api/v1/compute/{ClusterId}/updates/{ChannelId}")]
+		public async Task<ActionResult<GetTaskUpdatesResponse>> GetUpdatesAsync([FromRoute] ClusterId ClusterId, [FromRoute] ChannelId ChannelId, [FromQuery] int Wait = 0)
 		{
+			if (!await AclService.AuthorizeAsync(AclAction.ViewComputeTasks, User))
+			{
+				return Forbid();
+			}
+
 			List<IComputeTaskStatus> Results;
 			if (Wait == 0)
 			{
-				Results = await ComputeService.GetTaskUpdatesAsync(ChannelId);
+				Results = await ComputeService.GetTaskUpdatesAsync(ClusterId, ChannelId);
 			}
 			else
 			{
 				using CancellationTokenSource DelaySource = new CancellationTokenSource(Wait * 1000);
-				Results = await ComputeService.WaitForTaskUpdatesAsync(ChannelId, DelaySource.Token);
+				Results = await ComputeService.WaitForTaskUpdatesAsync(ClusterId, ChannelId, DelaySource.Token);
 			}
 
 			GetTaskUpdatesResponse Response = new GetTaskUpdatesResponse();
@@ -86,11 +122,11 @@ namespace HordeServer.Compute.Impl
 			{
 				GetTaskUpdateResponse Update = new GetTaskUpdateResponse();
 
-				Update.TaskHash = Result.Task;
+				Update.TaskRefId = Result.TaskRefId;
 				Update.Time = Result.Time;
 				Update.State = Result.State;
 				Update.Outcome = Result.Outcome;
-				Update.Result = Result.Result;
+				Update.ResultRefId = Result.ResultRefId;
 				Update.AgentId = Result.AgentId?.ToString();
 				Update.LeaseId = Result.LeaseId?.ToString();
 				Update.Detail = Result.Detail;
diff --git a/Engine/Source/Programs/Horde/Horde.Build/Compute/Impl/ComputeService.cs b/Engine/Source/Programs/Horde/Horde.Build/Compute/Impl/ComputeService.cs
--- a/Engine/Source/Programs/Horde/Horde.Build/Compute/Impl/ComputeService.cs
+++ b/Engine/Source/Programs/Horde/Horde.Build/Compute/Impl/ComputeService.cs
@@ -34,12 +34,12 @@ using Microsoft.Extensions.Hosting;
 using System.Threading.Channels;
 using EpicGames.Redis;
 using HordeServer.Collections;
+using EpicGames.Horde.Storage;
+using EpicGames.Horde.Common;
 
 namespace HordeServer.Compute.Impl
 {
-	using ChannelId = StringId<IComputeChannel>;
 	using LeaseId = ObjectId<ILease>;
-	using NamespaceId = StringId<INamespace>;
 
 	/// <summary>
 	/// Information about a particular task
@@ -47,19 +47,23 @@ namespace HordeServer.Compute.Impl
 	[RedisConverter(typeof(RedisCbConverter<>))]
 	class ComputeTaskInfo
 	{
+		[CbField("c")]
+		public ClusterId ClusterId { get; set; }
+
 		[CbField("h")]
-		public CbObjectAttachment TaskHash { get; set; }
+		public RefId TaskRefId { get; set; }
 
-		[CbField("c")]
+		[CbField("ch")]
 		public ChannelId ChannelId { get; set; }
 
 		private ComputeTaskInfo()
 		{
 		}
 
-		public ComputeTaskInfo(CbObjectAttachment TaskHash, ChannelId ChannelId)
+		public ComputeTaskInfo(ClusterId ClusterId, RefId TaskRefId, ChannelId ChannelId)
 		{
-			this.TaskHash = TaskHash;
+			this.ClusterId = ClusterId;
+			this.TaskRefId = TaskRefId;
 			this.ChannelId = ChannelId;
 		}
 	}
@@ -72,7 +76,7 @@ namespace HordeServer.Compute.Impl
 	{
 		/// <inheritdoc/>
 		[CbField("h")]
-		public CbObjectAttachment Task { get; set; }
+		public RefId TaskRefId { get; set; }
 
 		/// <inheritdoc/>
 		[CbField("t")]
@@ -96,7 +100,7 @@ namespace HordeServer.Compute.Impl
 
 		/// <inheritdoc/>
 		[CbField("r")]
-		public CbObjectAttachment? Result { get; set; }
+		public RefId? ResultRefId { get; set; }
 
 		/// <inheritdoc/>
 		[CbField("d")]
@@ -116,9 +120,9 @@ namespace HordeServer.Compute.Impl
 		{
 		}
 
-		public ComputeTaskStatus(CbObjectAttachment Task, ComputeTaskState State, AgentId? AgentId, LeaseId? LeaseId)
+		public ComputeTaskStatus(RefId TaskRefId, ComputeTaskState State, AgentId? AgentId, LeaseId? LeaseId)
 		{
-			this.Task = Task;
+			this.TaskRefId = TaskRefId;
 			this.Time = DateTime.UtcNow;
 			this.State = State;
 			this.AgentId = (AgentId == null)? Utf8String.Empty : AgentId.Value.ToString();
@@ -131,6 +135,49 @@ namespace HordeServer.Compute.Impl
 	/// </summary>
 	class ComputeService : TaskSourceBase<ComputeTaskMessage>, IComputeService, IHostedService, IDisposable
 	{
+		[RedisConverter(typeof(QueueKeySerializer))]
+		class QueueKey
+		{
+			public ClusterId ClusterId { get; set; }
+			public IoHash RequirementsHash { get; set; }
+
+			public QueueKey(ClusterId ClusterId, IoHash RequirementsHash)
+			{
+				this.ClusterId = ClusterId;
+				this.RequirementsHash = RequirementsHash;
+			}
+
+			public override string ToString() => $"{ClusterId}/{RequirementsHash}";
+		}
+
+		class QueueKeySerializer : IRedisConverter<QueueKey>
+		{
+			public QueueKey FromRedisValue(RedisValue Value)
+			{
+				string Str = Value.ToString();
+				int Idx = Str.LastIndexOf("/", StringComparison.Ordinal);
+				return new QueueKey(new ClusterId(Str.Substring(0, Idx)), IoHash.Parse(Str.Substring(Idx + 1)));
+			}
+
+			public RedisValue ToRedisValue(QueueKey Value) => $"{Value.ClusterId}/{Value.RequirementsHash}";
+		}
+
+		class ClusterInfo : IComputeClusterInfo
+		{
+			public ClusterId Id { get; set; }
+			public NamespaceId NamespaceId { get; set; }
+			public BucketId RequestBucketId { get; set; }
+			public BucketId ResponseBucketId { get; set; }
+
+			public ClusterInfo(ComputeClusterConfig Config)
+			{
+				Id = new ClusterId(Config.Id);
+				NamespaceId = new NamespaceId(Config.NamespaceId);
+				RequestBucketId = new BucketId(Config.RequestBucketId);
+				ResponseBucketId = new BucketId(Config.ResponseBucketId);
+			}
+		}
+
 		/// <inheritdoc/>
 		public override string Type => "Compute";
 
@@ -139,28 +186,30 @@ namespace HordeServer.Compute.Impl
 
 		public static NamespaceId DefaultNamespaceId { get; } = new NamespaceId("default");
 
-		IObjectCollection ObjectCollection;
-		ITaskScheduler<IoHash, ComputeTaskInfo> TaskScheduler;
+		IStorageClient StorageClient;
+		ITaskScheduler<QueueKey, ComputeTaskInfo> TaskScheduler;
 		RedisMessageQueue<ComputeTaskStatus> MessageQueue;
 		ITicker ExpireTasksTicker;
 		IMemoryCache RequirementsCache;
+		LazyCachedValue<Task<Globals>> Globals;
 		LazyCachedValue<Task<List<IPool>>> CachedPools;
 		ILogger Logger;
 
 		/// <summary>
 		/// Constructor
 		/// </summary>
-		public ComputeService(IDatabase Redis, IObjectCollection ObjectCollection, IPoolCollection PoolCollection, IClock Clock, ILogger<ComputeService> Logger)
+		public ComputeService(DatabaseService DatabaseService, IDatabase Redis, IStorageClient StorageClient, IPoolCollection PoolCollection, IClock Clock, ILogger<ComputeService> Logger)
 		{
-			this.ObjectCollection = ObjectCollection;
-			this.TaskScheduler = new RedisTaskScheduler<IoHash, ComputeTaskInfo>(Redis, "compute/tasks/", Logger);
+			this.StorageClient = StorageClient;
+			this.TaskScheduler = new RedisTaskScheduler<QueueKey, ComputeTaskInfo>(Redis, "compute/tasks/", Logger);
 			this.MessageQueue = new RedisMessageQueue<ComputeTaskStatus>(Redis, "compute/messages/");
 			this.ExpireTasksTicker = Clock.AddTicker(TimeSpan.FromMinutes(2.0), ExpireTasksAsync, Logger);
 			this.RequirementsCache = new MemoryCache(new MemoryCacheOptions());
+			this.Globals = new LazyCachedValue<Task<Globals>>(() => DatabaseService.GetGlobalsAsync(), TimeSpan.FromSeconds(120.0));
 			this.CachedPools = new LazyCachedValue<Task<List<IPool>>>(() => PoolCollection.GetAsync(), TimeSpan.FromSeconds(30.0));
 			this.Logger = Logger;
 
-			OnLeaseStartedProperties.Add(nameof(ComputeTaskMessage.Task), x => x.Task.Hash);
+			OnLeaseStartedProperties.Add(x => x.TaskRefId);
 		}
 
 		/// <inheritdoc/>
@@ -176,21 +225,21 @@ namespace HordeServer.Compute.Impl
 		/// <returns></returns>
 		async ValueTask ExpireTasksAsync(CancellationToken CancellationToken)
 		{
-			List<IoHash> RequirementsHashes = await TaskScheduler.GetInactiveQueuesAsync();
-			foreach (IoHash RequirementsHash in RequirementsHashes)
+			List<QueueKey> QueueKeys = await TaskScheduler.GetInactiveQueuesAsync();
+			foreach (QueueKey QueueKey in QueueKeys)
 			{
-				Logger.LogInformation("Inactive queue: {RequirementsHash}", RequirementsHash);
+				Logger.LogInformation("Inactive queue: {QueueKey}", QueueKey);
 				for (; ; )
 				{
-					ComputeTaskInfo? ComputeTask = await TaskScheduler.DequeueAsync(RequirementsHash);
+					ComputeTaskInfo? ComputeTask = await TaskScheduler.DequeueAsync(QueueKey);
 					if (ComputeTask == null)
 					{
 						break;
 					}
 
-					ComputeTaskStatus Status = new ComputeTaskStatus(ComputeTask.TaskHash, ComputeTaskState.Complete, null, null);
+					ComputeTaskStatus Status = new ComputeTaskStatus(ComputeTask.TaskRefId, ComputeTaskState.Complete, null, null);
 					Status.Outcome = ComputeTaskOutcome.Expired;
-					Logger.LogInformation("Compute task expired (queue: {RequirementsHash}, task: {TaskHash}, channel: {ChannelId})", RequirementsHash, ComputeTask.TaskHash, ComputeTask.ChannelId);
+					Logger.LogInformation("Compute task expired (queue: {RequirementsHash}, task: {TaskHash}, channel: {ChannelId})", QueueKey, ComputeTask.TaskRefId, ComputeTask.ChannelId);
 					await MessageQueue.PostAsync(ComputeTask.ChannelId.ToString(), Status);
 				}
 			}
@@ -204,19 +253,39 @@ namespace HordeServer.Compute.Impl
 			RequirementsCache.Dispose();
 		}
 
+		/// <summary>
+		/// Gets information about a compute cluster
+		/// </summary>
+		/// <param name="ClusterId">Cluster to use for execution</param>
+		public async Task<IComputeClusterInfo> GetClusterInfoAsync(ClusterId ClusterId)
+		{
+			ComputeClusterConfig? Config = await GetClusterAsync(ClusterId);
+			if (Config == null)
+			{
+				throw new KeyNotFoundException();
+			}
+			return new ClusterInfo(Config);
+		}
+
 		/// <inheritdoc/>
-		public async Task AddTasksAsync(NamespaceId NamespaceId, CbObjectAttachment RequirementsHash, List<CbObjectAttachment> TaskHashes, ChannelId ChannelId)
+		public async Task AddTasksAsync(ClusterId ClusterId, ChannelId ChannelId, List<RefId> TaskRefIds, CbObjectAttachment RequirementsHash)
 		{
 			List<Task> Tasks = new List<Task>();
-			foreach (CbObjectAttachment TaskHash in TaskHashes)
+			foreach (RefId TaskRefId in TaskRefIds)
 			{
-				ComputeTaskInfo TaskInfo = new ComputeTaskInfo(TaskHash, ChannelId);
-				Logger.LogDebug("Adding task {TaskHash} from channel {ChannelId} to queue {RequirementsHash}", TaskInfo.TaskHash.Hash, ChannelId, RequirementsHash);
-				Tasks.Add(TaskScheduler.EnqueueAsync(RequirementsHash, TaskInfo, false));
+				ComputeTaskInfo TaskInfo = new ComputeTaskInfo(ClusterId, TaskRefId, ChannelId);
+				Logger.LogDebug("Adding task {TaskHash} from channel {ChannelId} to queue {ClusterId}{RequirementsHash}", TaskRefId.Hash, ChannelId, ClusterId, RequirementsHash);
+				Tasks.Add(TaskScheduler.EnqueueAsync(new QueueKey(ClusterId, RequirementsHash), TaskInfo, false));
 			}
 			await Task.WhenAll(Tasks);
 		}
 
+		async ValueTask<ComputeClusterConfig?> GetClusterAsync(ClusterId ClusterId)
+		{
+			Globals GlobalsInstance = await Globals.GetCached();
+			return GlobalsInstance.ComputeClusters.FirstOrDefault(x => new ClusterId(x.Id) == ClusterId);
+		}
+
 		/// <inheritdoc/>
 		public override async Task<AgentLease?> AssignLeaseAsync(IAgent Agent, CancellationToken CancellationToken)
 		{
@@ -229,27 +298,44 @@ namespace HordeServer.Compute.Impl
 			}
 
 			// Find a task to execute
-			(IoHash, ComputeTaskInfo)? Entry = await TaskScheduler.DequeueAsync(RequirementsHash => CheckRequirements(Agent, RequirementsHash), CancellationToken);
+			(QueueKey, ComputeTaskInfo)? Entry = await TaskScheduler.DequeueAsync(QueueKey => CheckRequirements(Agent, QueueKey), CancellationToken);
 			if (Entry != null)
 			{
-				(IoHash RequirementsHash, ComputeTaskInfo TaskInfo) = Entry.Value;
+				(QueueKey QueueKey, ComputeTaskInfo TaskInfo) = Entry.Value;
+
+				ComputeClusterConfig? Cluster = await GetClusterAsync(TaskInfo.ClusterId);
+				if (Cluster == null)
+				{
+					Logger.LogWarning("Invalid cluster '{ClusterId}'; failing task {TaskRefId}", TaskInfo.ClusterId, TaskInfo.TaskRefId);
+					ComputeTaskStatus Status = new ComputeTaskStatus(TaskInfo.TaskRefId, ComputeTaskState.Complete, Agent.Id, null) { Detail = $"Invalid cluster '{TaskInfo.ClusterId}'" };
+					await MessageQueue.PostAsync(TaskInfo.ChannelId.ToString(), Status);
+					return null;
+				}
 
-				Requirements? Requirements = await GetCachedRequirementsAsync(RequirementsHash);
-				if (Requirements != null)
+				Requirements? Requirements = await GetCachedRequirementsAsync(QueueKey);
+				if (Requirements == null)
 				{
-					ComputeTaskMessage ComputeTask = new ComputeTaskMessage();
-					ComputeTask.ChannelId = TaskInfo.ChannelId.ToString();
-					ComputeTask.NamespaceId = DefaultNamespaceId.ToString();
-					ComputeTask.Requirements = new CbObjectAttachment(RequirementsHash);
-					ComputeTask.Task = TaskInfo.TaskHash;
-
-					string LeaseName = $"Remote action ({TaskInfo.TaskHash})";
-					byte[] Payload = Any.Pack(ComputeTask).ToByteArray();
-
-					AgentLease Lease = new AgentLease(LeaseId.GenerateNewId(), LeaseName, null, null, null, LeaseState.Pending, Requirements.Resources, Requirements.Exclusive, Payload);
-					Logger.LogDebug("Created lease {LeaseId} for channel {ChannelId} task {TaskHash} req {RequirementsHash}", Lease.Id, ComputeTask.ChannelId, ComputeTask.Task.Hash, ComputeTask.Requirements.Hash);
-					return Lease;
+					Logger.LogWarning("Unable to fetch requirements {RequirementsHash}", QueueKey);
+					ComputeTaskStatus Status = new ComputeTaskStatus(TaskInfo.TaskRefId, ComputeTaskState.Complete, Agent.Id, null) { Detail = $"Unable to retrieve requirements '{QueueKey}'" };
+					await MessageQueue.PostAsync(TaskInfo.ChannelId.ToString(), Status);
+					return null;
 				}
+
+				ComputeTaskMessage ComputeTask = new ComputeTaskMessage();
+				ComputeTask.ClusterId = TaskInfo.ClusterId.ToString();
+				ComputeTask.ChannelId = TaskInfo.ChannelId.ToString();
+				ComputeTask.NamespaceId = Cluster.NamespaceId.ToString();
+				ComputeTask.InputBucketId = Cluster.RequestBucketId.ToString();
+				ComputeTask.OutputBucketId = Cluster.ResponseBucketId.ToString();
+				ComputeTask.RequirementsHash = QueueKey.RequirementsHash;
+				ComputeTask.TaskRefId = TaskInfo.TaskRefId;
+
+				string LeaseName = $"Remote action ({TaskInfo.TaskRefId})";
+				byte[] Payload = Any.Pack(ComputeTask).ToByteArray();
+
+				AgentLease Lease = new AgentLease(LeaseId.GenerateNewId(), LeaseName, null, null, null, LeaseState.Pending, Requirements.Resources, Requirements.Exclusive, Payload);
+				Logger.LogDebug("Created lease {LeaseId} for channel {ChannelId} task {TaskHash} req {RequirementsHash}", Lease.Id, ComputeTask.ChannelId, ComputeTask.TaskRefId, ComputeTask.RequirementsHash);
+				return Lease;
 			}
 			return null;
 		}
@@ -257,20 +343,21 @@ namespace HordeServer.Compute.Impl
 		/// <inheritdoc/>
 		public override Task CancelLeaseAsync(IAgent Agent, LeaseId LeaseId, ComputeTaskMessage Message)
 		{
-			ComputeTaskInfo TaskInfo = new ComputeTaskInfo(Message.Task, new ChannelId(Message.ChannelId));
-			return TaskScheduler.EnqueueAsync(Message.Requirements.Hash, TaskInfo, true);
+			ClusterId ClusterId = new ClusterId(Message.ClusterId);
+			ComputeTaskInfo TaskInfo = new ComputeTaskInfo(ClusterId, new RefId(new IoHash(Message.TaskRefId.ToByteArray())), new ChannelId(Message.ChannelId));
+			return TaskScheduler.EnqueueAsync(new QueueKey(ClusterId, new IoHash(Message.RequirementsHash.ToByteArray())), TaskInfo, true);
 		}
 
 		/// <inheritdoc/>
-		public async Task<List<IComputeTaskStatus>> GetTaskUpdatesAsync(ChannelId ChannelId)
+		public async Task<List<IComputeTaskStatus>> GetTaskUpdatesAsync(ClusterId ClusterId, ChannelId ChannelId)
 		{
-			List<ComputeTaskStatus> Messages = await MessageQueue.ReadMessagesAsync(ChannelId.ToString());
+			List<ComputeTaskStatus> Messages = await MessageQueue.ReadMessagesAsync($"{ClusterId}/{ChannelId}");
 			return Messages.ConvertAll<IComputeTaskStatus>(x => x);
 		}
 
-		public async Task<List<IComputeTaskStatus>> WaitForTaskUpdatesAsync(ChannelId ChannelId, CancellationToken CancellationToken)
+		public async Task<List<IComputeTaskStatus>> WaitForTaskUpdatesAsync(ClusterId ClusterId, ChannelId ChannelId, CancellationToken CancellationToken)
 		{
-			List<ComputeTaskStatus> Messages = await MessageQueue.WaitForMessagesAsync(ChannelId.ToString(), CancellationToken);
+			List<ComputeTaskStatus> Messages = await MessageQueue.WaitForMessagesAsync($"{ClusterId}/{ChannelId}", CancellationToken);
 			return Messages.ConvertAll<IComputeTaskStatus>(x => x);
 		}
 
@@ -278,7 +365,7 @@ namespace HordeServer.Compute.Impl
 		{
 			base.OnLeaseStartedAsync(Agent, LeaseId, ComputeTask, Logger);
 
-			ComputeTaskStatus Status = new ComputeTaskStatus(ComputeTask.Task, ComputeTaskState.Executing, Agent.Id, LeaseId);
+			ComputeTaskStatus Status = new ComputeTaskStatus(ComputeTask.TaskRefId, ComputeTaskState.Executing, Agent.Id, LeaseId);
 			return MessageQueue.PostAsync(ComputeTask.ChannelId, Status);
 		}
 
@@ -286,14 +373,14 @@ namespace HordeServer.Compute.Impl
 		{
 			ComputeTaskResultMessage Message = ComputeTaskResultMessage.Parser.ParseFrom(Output.ToArray());
 
-			ComputeTaskStatus Status = new ComputeTaskStatus(ComputeTask.Task, ComputeTaskState.Complete, Agent.Id, LeaseId);
-			if (Message.Result != null)
+			ComputeTaskStatus Status = new ComputeTaskStatus(ComputeTask.TaskRefId, ComputeTaskState.Complete, Agent.Id, LeaseId);
+			if (Message.ResultRefId != null)
 			{
-				Status.Result = Message.Result;
+				Status.ResultRefId = Message.ResultRefId;
 			}
-			else if (Message.Outcome != ComputeTaskOutcome.Success)
+			else if ((ComputeTaskOutcome)Message.Outcome != ComputeTaskOutcome.Success)
 			{
-				(Status.Outcome, Status.Detail) = (Message.Outcome, Message.Detail);
+				(Status.Outcome, Status.Detail) = ((ComputeTaskOutcome)Message.Outcome, Message.Detail);
 			}
 			else if (Outcome == LeaseOutcome.Failed)
 			{
@@ -308,19 +395,19 @@ namespace HordeServer.Compute.Impl
 				Status.Outcome = ComputeTaskOutcome.NoResult;
 			}
 
-			Logger.LogInformation("Compute lease finished (lease: {LeaseId}, task: {TaskHash}, agent: {AgentId}, channel: {ChannelId}, result: {ResultHash}, outcome: {Outcome})", LeaseId, ComputeTask.Task.Hash, Agent.Id, ComputeTask.ChannelId, Status.Result?.Hash ?? IoHash.Zero, Status.Outcome);
-			return MessageQueue.PostAsync(ComputeTask.ChannelId, Status);
+			Logger.LogInformation("Compute lease finished (lease: {LeaseId}, task: {TaskHash}, agent: {AgentId}, channel: {ChannelId}, result: {ResultHash}, outcome: {Outcome})", LeaseId, ComputeTask.TaskRefId.AsRefId(), Agent.Id, ComputeTask.ChannelId, Status.ResultRefId?.ToString() ?? "(none)", Status.Outcome);
+			return MessageQueue.PostAsync($"{ComputeTask.ClusterId}/{ComputeTask.ChannelId}", Status);
 		}
 
 		/// <summary>
 		/// Checks that an agent matches the necessary criteria to execute a task
 		/// </summary>
 		/// <param name="Agent"></param>
-		/// <param name="RequirementsHash"></param>
+		/// <param name="QueueKey"></param>
 		/// <returns></returns>
-		async ValueTask<bool> CheckRequirements(IAgent Agent, IoHash RequirementsHash)
+		async ValueTask<bool> CheckRequirements(IAgent Agent, QueueKey QueueKey)
 		{
-			Requirements? Requirements = await GetCachedRequirementsAsync(RequirementsHash);
+			Requirements? Requirements = await GetCachedRequirementsAsync(QueueKey);
 			if (Requirements == null)
 			{
 				return false;
@@ -331,15 +418,27 @@ namespace HordeServer.Compute.Impl
 		/// <summary>
 		/// Gets the requirements object from the CAS
 		/// </summary>
-		/// <param name="RequirementsHash"></param>
+		/// <param name="QueueKey"></param>
 		/// <returns></returns>
-		async ValueTask<Requirements?> GetCachedRequirementsAsync(IoHash RequirementsHash)
+		async ValueTask<Requirements?> GetCachedRequirementsAsync(QueueKey QueueKey)
 		{
 			Requirements? Requirements;
-			if (!RequirementsCache.TryGetValue(RequirementsHash, out Requirements))
+			if (!RequirementsCache.TryGetValue(QueueKey.RequirementsHash, out Requirements))
 			{
-				Requirements = await ObjectCollection.GetAsync<Requirements>(DefaultNamespaceId, RequirementsHash);
-				using (ICacheEntry Entry = RequirementsCache.CreateEntry(RequirementsHash))
+				ComputeClusterConfig? ClusterConfig = await GetClusterAsync(QueueKey.ClusterId);
+				if (ClusterConfig != null)
+				{
+					try
+					{
+						Requirements = await StorageClient.ReadObjectAsync<Requirements>(new NamespaceId(ClusterConfig.NamespaceId), QueueKey.RequirementsHash);
+					}
+					catch (BlobNotFoundException)
+					{
+						Requirements = null;
+					}
+				}
+
+				using (ICacheEntry Entry = RequirementsCache.CreateEntry(QueueKey.RequirementsHash))
 				{
 					if (Requirements == null)
 					{
@@ -355,61 +454,4 @@ namespace HordeServer.Compute.Impl
 			return Requirements;
 		}
 	}
-
-	/// <summary>
-	/// Implementation of the gRPC compute service interface
-	/// </summary>
-	class ComputeRpcServer : ComputeRpc.ComputeRpcBase
-	{
-		IComputeService ComputeService;
-
-		public ComputeRpcServer(IComputeService ComputeService)
-		{
-			this.ComputeService = ComputeService;
-		}
-
-		public override async Task<AddTasksRpcResponse> AddTasks(AddTasksRpcRequest RpcRequest, ServerCallContext Context)
-		{
-			NamespaceId NamespaceId = new NamespaceId(RpcRequest.NamespaceId);
-			ChannelId ChannelId = new ChannelId(RpcRequest.ChannelId);
-			await ComputeService.AddTasksAsync(NamespaceId, RpcRequest.RequirementsHash, RpcRequest.TaskHashes.Select(x => (CbObjectAttachment)x).ToList(), ChannelId);
-			return new AddTasksRpcResponse();
-		}
-
-		public override async Task GetTaskUpdates(IAsyncStreamReader<GetTaskUpdatesRpcRequest> RequestStream, IServerStreamWriter<GetTaskUpdatesRpcResponse> ResponseStream, ServerCallContext Context)
-		{
-			Task<bool> MoveNextTask = RequestStream.MoveNext();
-			while(await MoveNextTask)
-			{
-				GetTaskUpdatesRpcRequest Request = RequestStream.Current;
-
-				ChannelId ChannelId = new ChannelId(Request.ChannelId);
-				using (CancellationTokenSource CancellationSource = new CancellationTokenSource())
-				{
-					MoveNextTask = MoveNextAndCancel(RequestStream, CancellationSource);
-					while (!CancellationSource.IsCancellationRequested)
-					{
-						List<IComputeTaskStatus> Updates = await ComputeService.WaitForTaskUpdatesAsync(ChannelId, CancellationSource.Token);
-						foreach (IComputeTaskStatus Update in Updates)
-						{
-							GetTaskUpdatesRpcResponse Response = new GetTaskUpdatesRpcResponse();
-							Response.Task = Update.Task;
-							Response.Time = Timestamp.FromDateTime(Update.Time);
-							Response.State = Update.State;
-							Response.Result = Update.Result;
-							await ResponseStream.WriteAsync(Response);
-						}
-					}
-					await MoveNextTask;
-				}
-			}
-		}
-
-		static async Task<bool> MoveNextAndCancel(IAsyncStreamReader<GetTaskUpdatesRpcRequest> RequestStream, CancellationTokenSource CancellationSource)
-		{
-			bool Result = await RequestStream.MoveNext();
-			CancellationSource.Cancel();
-			return Result;
-		}
-	}
 }
diff --git a/Engine/Source/Programs/Horde/Horde.Build/Compute/Impl/RedisTaskScheduler.cs b/Engine/Source/Programs/Horde/Horde.Build/Compute/Impl/RedisTaskScheduler.cs
--- a/Engine/Source/Programs/Horde/Horde.Build/Compute/Impl/RedisTaskScheduler.cs
+++ b/Engine/Source/Programs/Horde/Horde.Build/Compute/Impl/RedisTaskScheduler.cs
@@ -22,7 +22,6 @@ using System.Threading.Tasks;
 
 namespace HordeServer.Compute.Impl
 {
-	using NamespaceId = StringId<INamespace>;
 	using Condition = StackExchange.Redis.Condition;
 
 	/// <summary>
@@ -31,7 +30,7 @@ namespace HordeServer.Compute.Impl
 	/// <typeparam name="TQueueId">Type used to identify a particular queue</typeparam>
 	/// <typeparam name="TTask">Type used to describe a task to be performed</typeparam>
 	interface ITaskScheduler<TQueueId, TTask> 
-		where TQueueId : struct 
+		where TQueueId : notnull 
 		where TTask : class
 	{
 		/// <summary>
@@ -70,7 +69,7 @@ namespace HordeServer.Compute.Impl
 	/// <typeparam name="TQueueId">Type used to identify a particular queue</typeparam>
 	/// <typeparam name="TTask">Type used to describe a task to be performed</typeparam>
 	class RedisTaskScheduler<TQueueId, TTask> : ITaskScheduler<TQueueId, TTask>, IDisposable 
-		where TQueueId : struct 
+		where TQueueId : notnull
 		where TTask : class
 	{
 		class Listener
@@ -387,7 +386,7 @@ namespace HordeServer.Compute.Impl
 			while (await NewQueues.Reader.WaitToReadAsync(CancellationToken))
 			{
 				HashSet<TQueueId> NewQueueIds = new HashSet<TQueueId>();
-				while (NewQueues.Reader.TryRead(out TQueueId QueueId))
+				while (NewQueues.Reader.TryRead(out TQueueId? QueueId))
 				{
 					NewQueueIds.Add(QueueId);
 				}
diff --git a/Engine/Source/Programs/Horde/Horde.Build/Defaults/globals.json b/Engine/Source/Programs/Horde/Horde.Build/Defaults/globals.json
--- a/Engine/Source/Programs/Horde/Horde.Build/Defaults/globals.json
+++ b/Engine/Source/Programs/Horde/Horde.Build/Defaults/globals.json
@@ -5,6 +5,12 @@
 			"path": "ue5.project.json"
 		}
 	],
+	"compute": [
+		{
+			"id": "default",
+			"namespaceid": "horde.compute"
+		}
+	],
 	"storage": {
 		"namespaces": [
 			{
diff --git a/Engine/Source/Programs/Horde/Horde.Build/Horde.Build.csproj b/Engine/Source/Programs/Horde/Horde.Build/Horde.Build.csproj
--- a/Engine/Source/Programs/Horde/Horde.Build/Horde.Build.csproj
+++ b/Engine/Source/Programs/Horde/Horde.Build/Horde.Build.csproj
@@ -39,6 +39,7 @@
     <PackageReference Include="Microsoft.AspNetCore.DataProtection.StackExchangeRedis" Version="3.1.7" />
     <PackageReference Include="Microsoft.AspNetCore.SpaServices.Extensions" Version="3.1.17" />
     <PackageReference Include="Microsoft.Extensions.Caching.Memory" Version="6.0.0" />
+    <PackageReference Include="Microsoft.Extensions.Http" Version="3.1.22" />
     <PackageReference Include="MongoDB.Driver" Version="2.11.6" />
     <PackageReference Include="OpenTracing" Version="0.12.1" />
     <PackageReference Include="OpenTracing.Contrib.Grpc" Version="0.4.0" />
@@ -68,24 +69,23 @@
   </ItemGroup>
 
   <ItemGroup>
-    <Protobuf Include="..\..\Shared\EpicGames.Horde\Protos\horde\storage\ref_table.proto" ProtoRoot="..\..\Shared\EpicGames.Horde\Protos" GrpcServices="Server" />
-    <Protobuf Include="..\..\Shared\EpicGames.Horde\Protos\horde\storage\blob_store.proto" ProtoRoot="..\..\Shared\EpicGames.Horde\Protos" GrpcServices="Server" />
-    <Protobuf Include="..\..\Shared\EpicGames.Horde\Protos\horde\compute\compute_rpc.proto" ProtoRoot="..\..\Shared\EpicGames.Horde\Protos" GrpcServices="Server" />
     <Protobuf Include="Protos\horde\horde_rpc.proto" ProtoRoot="Protos" GrpcServices="Server" />
     <Protobuf Include="Protos\horde\horde_rpc_messages.proto" ProtoRoot="Protos" />
     <Protobuf Include="Protos\horde\common\agent_status.proto" ProtoRoot="Protos" />
     <Protobuf Include="Protos\horde\common\event_severity.proto" ProtoRoot="Protos" />
+    <Protobuf Include="Protos\horde\common\iohash_wrapper.proto" ProtoRoot="Protos" />
     <Protobuf Include="Protos\horde\common\jobstep_outcome.proto" ProtoRoot="Protos" />
     <Protobuf Include="Protos\horde\common\jobstep_state.proto" ProtoRoot="Protos" />
     <Protobuf Include="Protos\horde\common\label_change.proto" ProtoRoot="Protos" />
     <Protobuf Include="Protos\horde\common\lease_outcome.proto" ProtoRoot="Protos" />
     <Protobuf Include="Protos\horde\common\lease_state.proto" ProtoRoot="Protos" />
     <Protobuf Include="Protos\horde\common\priority.proto" ProtoRoot="Protos" />
+    <Protobuf Include="Protos\horde\common\refid_wrapper.proto" ProtoRoot="Protos" />
     <Protobuf Include="Protos\horde\messages\agent_capabilities.proto" ProtoRoot="Protos" />
     <Protobuf Include="Protos\horde\messages\agent_requirements.proto" ProtoRoot="Protos" />
     <Protobuf Include="Protos\horde\messages\agent_workspace.proto" ProtoRoot="Protos" />
     <Protobuf Include="Protos\horde\messages\lease.proto" ProtoRoot="Protos" />
-    <Protobuf Include="Protos\horde\tasks\compute_task.proto" ProtoRoot="Protos" AdditionalImportDirs="..\..\Shared\EpicGames.Horde\Protos" />
+    <Protobuf Include="Protos\horde\tasks\compute_task.proto" ProtoRoot="Protos" />
     <Protobuf Include="Protos\horde\tasks\conform_task.proto" ProtoRoot="Protos" />
     <Protobuf Include="Protos\horde\tasks\job_task.proto" ProtoRoot="Protos" />
     <Protobuf Include="Protos\horde\tasks\restart_task.proto" ProtoRoot="Protos" />
diff --git a/Engine/Source/Programs/Horde/Horde.Build/Models/Globals.cs b/Engine/Source/Programs/Horde/Horde.Build/Models/Globals.cs
--- a/Engine/Source/Programs/Horde/Horde.Build/Models/Globals.cs
+++ b/Engine/Source/Programs/Horde/Horde.Build/Models/Globals.cs
@@ -321,6 +321,11 @@ namespace HordeServer.Models
 		/// </summary>
 		public List<ScheduledDowntime> ScheduledDowntime { get; set; } = new List<ScheduledDowntime>();
 
+		/// <summary>
+		/// List of compute profiles
+		/// </summary>
+		public List<ComputeClusterConfig> ComputeClusters { get; set; } = new List<ComputeClusterConfig>();
+
 		/// <summary>
 		/// Private constructor for serialization
 		/// </summary>
diff --git a/Engine/Source/Programs/Horde/Horde.Build/Perforce/ICommit.cs b/Engine/Source/Programs/Horde/Horde.Build/Perforce/ICommit.cs
--- a/Engine/Source/Programs/Horde/Horde.Build/Perforce/ICommit.cs
+++ b/Engine/Source/Programs/Horde/Horde.Build/Perforce/ICommit.cs
@@ -1,6 +1,7 @@
 ﻿// Copyright Epic Games, Inc. All Rights Reserved.
 
 using EpicGames.Core;
+using EpicGames.Horde.Storage;
 using EpicGames.Perforce;
 using HordeServer.Models;
 using HordeServer.Utilities;
@@ -8,6 +9,7 @@ using MongoDB.Bson;
 using System;
 using System.Collections.Generic;
 using System.Linq;
+using System.Text;
 using System.Threading.Tasks;
 
 namespace HordeServer.Commits
@@ -127,4 +129,20 @@ namespace HordeServer.Commits
 			this.DateUtc = DateUtc;
 		}
 	}
+
+	/// <summary>
+	/// Extension methods for commits
+	/// </summary>
+	static class CommitExtensions
+	{
+		/// <summary>
+		/// Get the ref id for a particular commit
+		/// </summary>
+		/// <param name="Commit"></param>
+		/// <returns></returns>
+		public static RefId GetRefId(this ICommit Commit)
+		{
+			return new RefId(IoHash.Compute(Encoding.UTF8.GetBytes(Commit.TreeRef!)));
+		}
+	}
 }
diff --git a/Engine/Source/Programs/Horde/Horde.Build/Perforce/Impl/CommitService.cs b/Engine/Source/Programs/Horde/Horde.Build/Perforce/Impl/CommitService.cs
--- a/Engine/Source/Programs/Horde/Horde.Build/Perforce/Impl/CommitService.cs
+++ b/Engine/Source/Programs/Horde/Horde.Build/Perforce/Impl/CommitService.cs
@@ -14,7 +14,6 @@ using System.Text.RegularExpressions;
 using System.Threading;
 using System.Threading.Tasks;
 using HordeServer.Storage;
-using HordeServer.Storage.Primitives;
 using HordeServer.Services;
 using EpicGames.Serialization;
 using MongoDB.Bson.Serialization.Attributes;
@@ -30,13 +29,14 @@ using System.Diagnostics;
 using System.IO.Compression;
 using System.IO;
 using HordeCommon;
+using EpicGames.Horde.Storage;
+using System.Text;
 
 namespace HordeServer.Commits.Impl
 {
 	using P4 = Perforce.P4;
-	using NamespaceId = StringId<INamespace>;
-	using BucketId = StringId<IBucket>;
 	using StreamId = StringId<IStream>;
+	using IRef = EpicGames.Horde.Storage.IRef;
 
 	/// <summary>
 	/// Service which mirrors changes from Perforce
@@ -68,9 +68,7 @@ namespace HordeServer.Commits.Impl
 
 		// Collections
 		ICommitCollection CommitCollection;
-		IBlobCollection BlobCollection;
-		IObjectCollection ObjectCollection;
-		IRefCollection RefCollection;
+		IStorageClient StorageClient;
 		IStreamCollection StreamCollection;
 		IPerforceService PerforceService;
 		IUserCollection UserCollection;
@@ -82,16 +80,14 @@ namespace HordeServer.Commits.Impl
 		/// <summary>
 		/// Constructor
 		/// </summary>
-		public CommitService(IDatabase Redis, ICommitCollection CommitCollection, IBlobCollection BlobCollection, IObjectCollection ObjectCollection, IRefCollection RefCollection, IStreamCollection StreamCollection, IPerforceService PerforceService, IUserCollection UserCollection, IClock Clock, ILogger<CommitService> Logger)
+		public CommitService(IDatabase Redis, ICommitCollection CommitCollection, IStorageClient StorageClient, IStreamCollection StreamCollection, IPerforceService PerforceService, IUserCollection UserCollection, IClock Clock, ILogger<CommitService> Logger)
 		{
 			this.Redis = Redis;
 			this.RedisDirtyStreams = new RedisSet<StreamId>(Redis, RedisBaseKey.Append("streams"));
 			this.RedisReservations = new RedisSortedSet<StreamId>(Redis, RedisBaseKey.Append("reservations"));
 
 			this.CommitCollection = CommitCollection;
-			this.BlobCollection = BlobCollection;
-			this.ObjectCollection = ObjectCollection;
-			this.RefCollection = RefCollection;
+			this.StorageClient = StorageClient;
 			this.StreamCollection = StreamCollection;
 			this.PerforceService = PerforceService;
 			this.UserCollection = UserCollection;
@@ -575,7 +571,7 @@ namespace HordeServer.Commits.Impl
 			ICommit? PrevCommit = null;
 			StreamTreeRef? PrevRoot = null;
 #endif
-			ObjectSet ObjectSet = new ObjectSet(BlobCollection, NamespaceId, 256 * 1024, DateTime.UtcNow);
+			ObjectSet ObjectSet = new ObjectSet(StorageClient, NamespaceId, 256 * 1024, DateTime.UtcNow);
 			for (; ; )
 			{
 				// Get the first two changelists to update
@@ -650,7 +646,7 @@ namespace HordeServer.Commits.Impl
 				return null;
 			}
 
-			IRef? Ref = await RefCollection.GetAsync(NamespaceId, BucketId, Commit.TreeRef);
+			IRef? Ref = await StorageClient.GetRefAsync(NamespaceId, BucketId, Commit.GetRefId());
 			if (Ref == null)
 			{
 				return null;
@@ -675,12 +671,8 @@ namespace HordeServer.Commits.Impl
 
 			CommitTree Tree = await FlushAsync(ObjectSet, Root);
 			string RefName = $"tree_{Commit.StreamId}_{Commit.Change}";
-
-			List<IoHash> MissingHashes = await RefCollection.SetAsync(NamespaceId, BucketId, RefName, Tree.Serialize(Stream.Name));
-			if (MissingHashes.Count > 0)
-			{
-				throw new Exception($"Missing hashes when attempting to add ref: {String.Join(", ", MissingHashes.Select(x => x.ToString()))}");
-			}
+			RefId RefId = new RefId(IoHash.Compute(Encoding.UTF8.GetBytes(RefName)));
+			await StorageClient.SetRefAsync(NamespaceId, BucketId, RefId, Tree.Serialize(Stream.Name));
 
 			NewCommit NewCommit = new NewCommit(Commit);
 			NewCommit.TreeRef = RefName;
@@ -689,9 +681,9 @@ namespace HordeServer.Commits.Impl
 			return Tree.Root;
 		}
 
-#endregion
+		#endregion
 
-#region Tree Snapshots
+		#region Tree Snapshots
 
 		/// <summary>
 		/// Creates a snapshot of a stream at a particular changelist
@@ -1099,7 +1091,7 @@ namespace HordeServer.Commits.Impl
 		/// <param name="Commit"></param>
 		async Task<CommitTree> ReadCommitTreeAsync(ICommit Commit)
 		{
-			IRef? Ref = await RefCollection.GetAsync(NamespaceId, BucketId, Commit.TreeRef!);
+			IRef? Ref = await StorageClient.GetRefAsync(NamespaceId, BucketId, Commit.GetRefId());
 			return CbSerializer.Deserialize<CommitTree>(Ref!.Value.AsField());
 		}
 
diff --git a/Engine/Source/Programs/Horde/Horde.Build/Perforce/Impl/CommitTree.cs b/Engine/Source/Programs/Horde/Horde.Build/Perforce/Impl/CommitTree.cs
--- a/Engine/Source/Programs/Horde/Horde.Build/Perforce/Impl/CommitTree.cs
+++ b/Engine/Source/Programs/Horde/Horde.Build/Perforce/Impl/CommitTree.cs
@@ -15,8 +15,6 @@ using System.Threading.Tasks;
 
 namespace HordeServer.Commits.Impl
 {
-	using NamespaceId = StringId<INamespace>;
-
 	class CommitTree
 	{
 		class ObjectFormat
diff --git a/Engine/Source/Programs/Horde/Horde.Build/Server/Config.cs b/Engine/Source/Programs/Horde/Horde.Build/Server/Config.cs
--- a/Engine/Source/Programs/Horde/Horde.Build/Server/Config.cs
+++ b/Engine/Source/Programs/Horde/Horde.Build/Server/Config.cs
@@ -2,6 +2,8 @@
 
 using EpicGames.Core;
 using EpicGames.Horde.Common;
+using EpicGames.Horde.Compute;
+using EpicGames.Horde.Storage;
 using EpicGames.Serialization;
 using HordeServer.Api;
 using HordeServer.Utilities;
@@ -43,6 +45,11 @@ namespace HordeServer.Models
 		/// </summary>
 		public List<AgentRateConfig> Rates { get; set; } = new List<AgentRateConfig>();
 
+		/// <summary>
+		/// List of compute profiles
+		/// </summary>
+		public List<ComputeClusterConfig> Compute { get; set; } = new List<ComputeClusterConfig>();
+
 		/// <summary>
 		/// Maximum number of conforms to run at once
 		/// </summary>
@@ -59,6 +66,37 @@ namespace HordeServer.Models
 		public UpdateAclRequest? Acl { get; set; }
 	}
 
+	/// <summary>
+	/// Profile for executing compute requests
+	/// </summary>
+	public class ComputeClusterConfig
+	{
+		/// <summary>
+		/// Name of the partition
+		/// </summary>
+		public string Id { get; set; } = "default";
+
+		/// <summary>
+		/// Name of the namespace to use
+		/// </summary>
+		public string NamespaceId { get; set; } = "horde.compute";
+
+		/// <summary>
+		/// Name of the input bucket
+		/// </summary>
+		public string RequestBucketId { get; set; } = "requests";
+
+		/// <summary>
+		/// Name of the output bucket
+		/// </summary>
+		public string ResponseBucketId { get; set; } = "responses";
+
+		/// <summary>
+		/// Access control list
+		/// </summary>
+		public UpdateAclRequest? Acl { get; set; }
+	}
+
 	/// <summary>
 	/// References a project configuration
 	/// </summary>
diff --git a/Engine/Source/Programs/Horde/Horde.Build/Server/ConfigService.cs b/Engine/Source/Programs/Horde/Horde.Build/Server/ConfigService.cs
--- a/Engine/Source/Programs/Horde/Horde.Build/Server/ConfigService.cs
+++ b/Engine/Source/Programs/Horde/Horde.Build/Server/ConfigService.cs
@@ -35,8 +35,6 @@ namespace HordeServer.Services
 	using PoolId = StringId<IPool>;
 	using ProjectId = StringId<IProject>;
 	using StreamId = StringId<IStream>;
-	using NamespaceId = StringId<INamespace>;
-	using BucketId = StringId<IBucket>;
 
 	/// <summary>
 	/// Polls Perforce for stream config changes
@@ -49,7 +47,7 @@ namespace HordeServer.Services
 		/// <summary>
 		/// Config file version number
 		/// </summary>
-		const int Version = 6;
+		const int Version = 9;
 
 		DatabaseService DatabaseService;
 		ProjectService ProjectService;
@@ -57,8 +55,6 @@ namespace HordeServer.Services
 		IPerforceService PerforceService;
 		PerforceLoadBalancer PerforceLoadBalancer;
 		INotificationService NotificationService;
-		INamespaceCollection NamespaceCollection;
-		IBucketCollection BucketCollection;
 		AgentService AgentService;
 		PoolService PoolService;
 		IOptionsMonitor<ServerSettings> Settings;
@@ -68,7 +64,7 @@ namespace HordeServer.Services
 		/// <summary>
 		/// Constructor
 		/// </summary>
-		public ConfigService(DatabaseService DatabaseService, IPerforceService PerforceService, ProjectService ProjectService, StreamService StreamService, INotificationService NotificationService, INamespaceCollection NamespaceCollection, IBucketCollection BucketCollection, PoolService PoolService, AgentService AgentService, PerforceLoadBalancer PerforceLoadBalancer, IClock Clock, IOptionsMonitor<ServerSettings> Settings, ILogger<ConfigService> Logger)
+		public ConfigService(DatabaseService DatabaseService, IPerforceService PerforceService, ProjectService ProjectService, StreamService StreamService, INotificationService NotificationService, PoolService PoolService, AgentService AgentService, PerforceLoadBalancer PerforceLoadBalancer, IClock Clock, IOptionsMonitor<ServerSettings> Settings, ILogger<ConfigService> Logger)
 		{
 			this.DatabaseService = DatabaseService;
 			this.PerforceService = PerforceService;
@@ -76,8 +72,6 @@ namespace HordeServer.Services
 			this.ProjectService = ProjectService;
 			this.StreamService = StreamService;
 			this.NotificationService = NotificationService;
-			this.NamespaceCollection = NamespaceCollection;
-			this.BucketCollection = BucketCollection;
 			this.PoolService = PoolService;
 			this.AgentService = AgentService;
 			this.Settings = Settings;
@@ -150,10 +144,9 @@ namespace HordeServer.Services
 				Globals.PerforceClusters = CachedGlobalConfig.PerforceClusters;
 				Globals.ScheduledDowntime = CachedGlobalConfig.Downtime;
 				Globals.MaxConformCount = CachedGlobalConfig.MaxConformCount;
+				Globals.ComputeClusters = CachedGlobalConfig.Compute;
 				Globals.RootAcl = Acl.Merge(null, CachedGlobalConfig.Acl);
 
-				await UpdateStorageConfigAsync(GlobalConfig.Storage);
-
 				if (await DatabaseService.TryUpdateSingletonAsync(Globals))
 				{
 					break;
@@ -308,44 +301,6 @@ namespace HordeServer.Services
 			}
 		}
 
-		async Task UpdateStorageConfigAsync(StorageConfig? Config)
-		{
-			List<INamespace> Namespaces = await NamespaceCollection.FindAsync();
-			List<NamespaceId> RemoveNamespaceIds = Namespaces.ConvertAll(x => x.Id);
-
-			if (Config != null)
-			{
-				foreach (NamespaceConfig NamespaceConfig in Config.Namespaces)
-				{
-					NamespaceId NamespaceId = new NamespaceId(NamespaceConfig.Id);
-					RemoveNamespaceIds.Remove(NamespaceId);
-					await NamespaceCollection.AddOrUpdateAsync(NamespaceId, NamespaceConfig);
-
-					List<IBucket> Buckets = await BucketCollection.FindAsync(NamespaceId);
-					List<BucketId> RemoveBucketIds = Buckets.ConvertAll(x => x.BucketId);
-
-					foreach(BucketConfig BucketConfig in NamespaceConfig.Buckets)
-					{
-						BucketId BucketId = new BucketId(BucketConfig.Id);
-						await BucketCollection.AddOrUpdateAsync(NamespaceId, BucketId, BucketConfig);
-						RemoveBucketIds.Remove(BucketId);
-					}
-
-					foreach (BucketId RemoveBucketId in RemoveBucketIds)
-					{
-						Logger.LogInformation("Removing bucket {BucketId}", RemoveBucketId);
-						await BucketCollection.RemoveAsync(NamespaceId, RemoveBucketId);
-					}
-				}
-			}
-
-			foreach (NamespaceId RemoveNamespaceId in RemoveNamespaceIds)
-			{
-				Logger.LogInformation("Removing namespace {NamespaceId}", RemoveNamespaceId);
-				await NamespaceCollection.RemoveAsync(RemoveNamespaceId);
-			}
-		}
-
 		static FileExtensionContentTypeProvider ContentTypeProvider = new FileExtensionContentTypeProvider();
 
 		static string GetMimeTypeFromPath(Uri Path)
diff --git a/Engine/Source/Programs/Horde/Horde.Build/ServerSettings.cs b/Engine/Source/Programs/Horde/Horde.Build/ServerSettings.cs
--- a/Engine/Source/Programs/Horde/Horde.Build/ServerSettings.cs
+++ b/Engine/Source/Programs/Horde/Horde.Build/ServerSettings.cs
@@ -6,7 +6,6 @@ using HordeServer.Models;
 using HordeServer.Services;
 using HordeServer.Storage;
 using HordeServer.Storage.Backends;
-using HordeServer.Storage.Collections;
 using HordeServer.Utilities;
 using System;
 using System.Collections.Generic;
@@ -16,6 +15,7 @@ using System.Runtime.InteropServices;
 using System.Threading.Tasks;
 using Horde.Build.Fleet.Autoscale;
 using TimeZoneConverter;
+using EpicGames.Horde.Storage;
 
 namespace HordeServer
 {
@@ -561,7 +561,12 @@ namespace HordeServer
 		/// Settings for remote execution
 		/// </summary>
 		public RemoteExecSettings RemoteExecSettings { get; set; } = new RemoteExecSettings();
-		
+
+		/// <summary>
+		/// Settings for the storage service
+		/// </summary>
+		public StorageSettings? Storage { get; set; }
+
 		/// <summary>
 		/// Lazily computed timezone value
 		/// </summary>
diff --git a/Engine/Source/Programs/Horde/Horde.Build/Startup.cs b/Engine/Source/Programs/Horde/Horde.Build/Startup.cs
--- a/Engine/Source/Programs/Horde/Horde.Build/Startup.cs
+++ b/Engine/Source/Programs/Horde/Horde.Build/Startup.cs
@@ -67,12 +67,10 @@ using HordeServer.Notifications.Impl;
 using HordeServer.Notifications;
 using Microsoft.AspNetCore.Hosting.Server.Features;
 using System.Runtime.InteropServices;
-using HordeServer.Storage.Collections;
 using Amazon.Extensions.NETCore.Setup;
 using Amazon.S3;
 using Amazon.Runtime;
 using HordeServer.Storage.Backends;
-using HordeServer.Storage.Services;
 using HordeServer.Commits.Impl;
 using HordeServer.Commits;
 using OpenTracing.Contrib.Grpc.Interceptors;
@@ -88,6 +86,10 @@ using Serilog.Events;
 using HordeServer.Jobs;
 using Microsoft.AspNetCore.Mvc;
 using Microsoft.AspNetCore.Mvc.ModelBinding;
+using Horde.Build.Storage.Services;
+using EpicGames.Horde.Storage;
+using System.Net.Http;
+using EpicGames.Horde.Storage.Impl;
 
 namespace HordeServer
 {
@@ -251,8 +253,9 @@ namespace HordeServer
 			Services.Configure<ServerSettings>(Configuration.GetSection("Horde"));
 
 			// Settings used for configuring services
+			IConfigurationSection ConfigSection = Configuration.GetSection("Horde");
 			ServerSettings Settings = new ServerSettings();
-			Configuration.GetSection("Horde").Bind(Settings);			
+			ConfigSection.Bind(Settings);			
 
 			if (Settings.GlobalThreadPoolMinSize != null)
 			{
@@ -324,13 +327,6 @@ namespace HordeServer
 			// Auditing
 			Services.AddSingleton<IAuditLog<AgentId>>(SP => SP.GetRequiredService<IAuditLogFactory<AgentId>>().Create("Agents.Log", "AgentId"));
 
-			// Storage
-			Services.AddSingleton<IBlobCollection>(SP => new CachingBlobCollection(new BlobCollection(SP.GetRequiredService<IStorageBackend<BlobCollection>>()), 256 * 1024 * 1024));
-			Services.AddSingleton<IBucketCollection, BucketCollection>();
-			Services.AddSingleton<INamespaceCollection, NamespaceCollection>();
-			Services.AddSingleton<IObjectCollection, ObjectCollection>();
-			Services.AddSingleton<IRefCollection, RefCollection>();
-
 			Services.AddSingleton(typeof(IAuditLogFactory<>), typeof(AuditLogFactory<>));
 			Services.AddSingleton(typeof(ISingletonDocument<>), typeof(SingletonDocument<>));
 
@@ -424,12 +420,11 @@ namespace HordeServer
 
 			Services.AddSingleton(new StorageBackendSettings<PersistentLogStorage> { Type = Settings.ExternalStorageProviderType, BaseDir = Settings.LocalLogsDir, BucketName = Settings.S3LogBucketName });
 			Services.AddSingleton(new StorageBackendSettings<ArtifactCollection> { Type = Settings.ExternalStorageProviderType, BaseDir = Settings.LocalArtifactsDir, BucketName = Settings.S3ArtifactBucketName });
-			Services.AddSingleton(new StorageBackendSettings<BlobCollection> { Type = Settings.ExternalStorageProviderType, BaseDir = Settings.LocalBlobsDir, BucketName = Settings.S3LogBucketName, BucketPath = "blobs/" });
 			Services.AddSingleton(typeof(IStorageBackend<>), typeof(StorageBackendFactory<>));
 
+			Services.AddHordeStorage(Settings => ConfigSection.GetSection("Storage").Bind(Settings));
+
 			ConfigureLogStorage(Services);
-			Services.AddSingleton<IStorageService, SimpleStorageService>();
-			//			ConfigureLogFileWriteCache(Services, Settings);
 
 			AuthenticationBuilder AuthBuilder = Services.AddAuthentication(Options =>
 				{
@@ -894,11 +889,6 @@ namespace HordeServer
 			{
 				Endpoints.MapGrpcService<HealthService>();
 				Endpoints.MapGrpcService<RpcService>();
-				
-				// Google Remote Execution API
-				Endpoints.MapGrpcService<ComputeRpcServer>();
-				Endpoints.MapGrpcService<BlobStoreRpc>();
-				Endpoints.MapGrpcService<RefTableRpc>();
 
 				Endpoints.MapGrpcReflectionService();
 
diff --git a/Engine/Source/Programs/Horde/Horde.Build/Storage/Collections/BlobCollection.cs b/Engine/Source/Programs/Horde/Horde.Build/Storage/Collections/BlobCollection.cs
deleted file mode 100644
diff --git a/Engine/Source/Programs/Horde/Horde.Build/Storage/Collections/BlobStoreRpc.cs b/Engine/Source/Programs/Horde/Horde.Build/Storage/Collections/BlobStoreRpc.cs
deleted file mode 100644
diff --git a/Engine/Source/Programs/Horde/Horde.Build/Storage/Collections/BucketCollection.cs b/Engine/Source/Programs/Horde/Horde.Build/Storage/Collections/BucketCollection.cs
deleted file mode 100644
diff --git a/Engine/Source/Programs/Horde/Horde.Build/Storage/Collections/CachingBlobCollection.cs b/Engine/Source/Programs/Horde/Horde.Build/Storage/Collections/CachingBlobCollection.cs
deleted file mode 100644
diff --git a/Engine/Source/Programs/Horde/Horde.Build/Storage/Collections/NamespaceCollection.cs b/Engine/Source/Programs/Horde/Horde.Build/Storage/Collections/NamespaceCollection.cs
deleted file mode 100644
diff --git a/Engine/Source/Programs/Horde/Horde.Build/Storage/Collections/ObjectCollection.cs b/Engine/Source/Programs/Horde/Horde.Build/Storage/Collections/ObjectCollection.cs
deleted file mode 100644
diff --git a/Engine/Source/Programs/Horde/Horde.Build/Storage/Collections/RefCollection.cs b/Engine/Source/Programs/Horde/Horde.Build/Storage/Collections/RefCollection.cs
deleted file mode 100644
diff --git a/Engine/Source/Programs/Horde/Horde.Build/Storage/Collections/RefTableRpc.cs b/Engine/Source/Programs/Horde/Horde.Build/Storage/Collections/RefTableRpc.cs
deleted file mode 100644
diff --git a/Engine/Source/Programs/Horde/Horde.Build/Storage/Controllers/BlobsController.cs b/Engine/Source/Programs/Horde/Horde.Build/Storage/Controllers/BlobsController.cs
deleted file mode 100644
diff --git a/Engine/Source/Programs/Horde/Horde.Build/Storage/Controllers/ObjectsController.cs b/Engine/Source/Programs/Horde/Horde.Build/Storage/Controllers/ObjectsController.cs
deleted file mode 100644
diff --git a/Engine/Source/Programs/Horde/Horde.Build/Storage/Controllers/RefsController.cs b/Engine/Source/Programs/Horde/Horde.Build/Storage/Controllers/RefsController.cs
deleted file mode 100644
diff --git a/Engine/Source/Programs/Horde/Horde.Build/Storage/IBlobCollection.cs b/Engine/Source/Programs/Horde/Horde.Build/Storage/IBlobCollection.cs
deleted file mode 100644
diff --git a/Engine/Source/Programs/Horde/Horde.Build/Storage/IBucket.cs b/Engine/Source/Programs/Horde/Horde.Build/Storage/IBucket.cs
deleted file mode 100644
diff --git a/Engine/Source/Programs/Horde/Horde.Build/Storage/IBucketCollection.cs b/Engine/Source/Programs/Horde/Horde.Build/Storage/IBucketCollection.cs
deleted file mode 100644
diff --git a/Engine/Source/Programs/Horde/Horde.Build/Storage/INamespace.cs b/Engine/Source/Programs/Horde/Horde.Build/Storage/INamespace.cs
deleted file mode 100644
diff --git a/Engine/Source/Programs/Horde/Horde.Build/Storage/INamespaceCollection.cs b/Engine/Source/Programs/Horde/Horde.Build/Storage/INamespaceCollection.cs
deleted file mode 100644
diff --git a/Engine/Source/Programs/Horde/Horde.Build/Storage/IObjectCollection.cs b/Engine/Source/Programs/Horde/Horde.Build/Storage/IObjectCollection.cs
deleted file mode 100644
diff --git a/Engine/Source/Programs/Horde/Horde.Build/Storage/IRef.cs b/Engine/Source/Programs/Horde/Horde.Build/Storage/IRef.cs
deleted file mode 100644
diff --git a/Engine/Source/Programs/Horde/Horde.Build/Storage/IRefCollection.cs b/Engine/Source/Programs/Horde/Horde.Build/Storage/IRefCollection.cs
deleted file mode 100644
diff --git a/Engine/Source/Programs/Horde/Horde.Build/Storage/IStorageService.cs b/Engine/Source/Programs/Horde/Horde.Build/Storage/IStorageService.cs
deleted file mode 100644
diff --git a/Engine/Source/Programs/Horde/Horde.Build/Storage/Primitives/BlobRef.cs b/Engine/Source/Programs/Horde/Horde.Build/Storage/Primitives/BlobRef.cs
deleted file mode 100644
diff --git a/Engine/Source/Programs/Horde/Horde.Build/Storage/Primitives/BlobType.cs b/Engine/Source/Programs/Horde/Horde.Build/Storage/Primitives/BlobType.cs
deleted file mode 100644
diff --git a/Engine/Source/Programs/Horde/Horde.Build/Storage/Primitives/KeyValueItem.cs b/Engine/Source/Programs/Horde/Horde.Build/Storage/Primitives/KeyValueItem.cs
deleted file mode 100644
diff --git a/Engine/Source/Programs/Horde/Horde.Build/Storage/Primitives/KeyValueNode.cs b/Engine/Source/Programs/Horde/Horde.Build/Storage/Primitives/KeyValueNode.cs
deleted file mode 100644
diff --git a/Engine/Source/Programs/Horde/Horde.Build/Storage/Primitives/KeyValueTree.cs b/Engine/Source/Programs/Horde/Horde.Build/Storage/Primitives/KeyValueTree.cs
deleted file mode 100644
diff --git a/Engine/Source/Programs/Horde/Horde.Build/Storage/Primitives/ReadOnlyHashArray.cs b/Engine/Source/Programs/Horde/Horde.Build/Storage/Primitives/ReadOnlyHashArray.cs
deleted file mode 100644
diff --git a/Engine/Source/Programs/Horde/Horde.Build/Storage/Services/BasicStorageClient.cs b/Engine/Source/Programs/Horde/Horde.Build/Storage/Services/BasicStorageClient.cs
new file mode 100644
--- /dev/null
+++ b/Engine/Source/Programs/Horde/Horde.Build/Storage/Services/BasicStorageClient.cs
@@ -0,0 +1,239 @@
+// Copyright Epic Games, Inc. All Rights Reserved.
+
+using EpicGames.Core;
+using EpicGames.Horde.Storage;
+using EpicGames.Serialization;
+using HordeServer.Services;
+using HordeServer.Storage;
+using HordeServer.Utilities;
+using MongoDB.Bson.Serialization.Attributes;
+using MongoDB.Driver;
+using System;
+using System.Collections.Generic;
+using System.IO;
+using System.Linq;
+using System.Threading;
+using System.Threading.Tasks;
+
+namespace Horde.Build.Storage.Services
+{
+	using IRef = EpicGames.Horde.Storage.IRef;
+	using BlobNotFoundException = EpicGames.Horde.Storage.BlobNotFoundException;
+
+	/// <summary>
+	/// Simple implementation of <see cref="IStorageClient"/> which uses the current <see cref="IStorageBackend"/> implementation without garbage collection.
+	/// </summary>
+	public class BasicStorageClient : IStorageClient
+	{
+		class RefImpl : IRef
+		{
+			public string Id { get; set; } = String.Empty;
+
+			[BsonElement("d")]
+			public byte[] Data { get; set; } = null!;
+
+			[BsonElement("f")]
+			public bool Finalized { get; set; }
+
+			[BsonElement("t")]
+			public DateTime LastAccessTime { get; set; } = DateTime.UtcNow;
+
+			public NamespaceId NamespaceId => new NamespaceId(Id.Split('/')[0]);
+			public BucketId BucketId => new BucketId(Id.Split('/')[1]);
+			public RefId RefId => new RefId(IoHash.Parse(Id.Split('/')[2]));
+
+			public CbObject Value => new CbObject(Data);
+		}
+
+		IMongoCollection<RefImpl> Refs;
+		IStorageBackend StorageBackend;
+
+		/// <summary>
+		/// 
+		/// </summary>
+		/// <param name="DatabaseService"></param>
+		/// <param name="StorageBackend"></param>
+		public BasicStorageClient(DatabaseService DatabaseService, IStorageBackend<BasicStorageClient> StorageBackend)
+		{
+			Refs = DatabaseService.GetCollection<RefImpl>("Storage.RefsV2");
+			this.StorageBackend = StorageBackend;
+		}
+
+		/// <inheritdoc/>
+		public async Task<Stream> ReadBlobAsync(NamespaceId NamespaceId, IoHash Hash, CancellationToken CancellationToken)
+		{
+			string Path = GetFullBlobPath(NamespaceId, Hash);
+
+			Stream? Stream = await StorageBackend.ReadAsync(Path);
+			if(Stream == null)
+			{
+				throw new BlobNotFoundException(NamespaceId, Hash);
+			}
+
+			return Stream;
+		}
+
+		/// <inheritdoc/>
+		public async Task WriteBlobAsync(NamespaceId NamespaceId, IoHash Hash, Stream Stream, CancellationToken CancellationToken)
+		{
+			byte[] Data;
+			using (MemoryStream Buffer = new MemoryStream())
+			{
+				await Stream.CopyToAsync(Buffer, CancellationToken);
+				Data = Buffer.ToArray();
+			}
+
+			IoHash DataHash = IoHash.Compute(Data);
+			if (Hash != DataHash)
+			{
+				throw new InvalidDataException($"Incorrect hash for data (was {DataHash}, expected {Hash})");
+			}
+
+			string Path = GetFullBlobPath(NamespaceId, Hash);
+			await StorageBackend.WriteBytesAsync(Path, Data);
+		}
+
+		/// <inheritdoc/>
+		public async Task<IRef> GetRefAsync(NamespaceId NamespaceId, BucketId BucketId, RefId RefId, CancellationToken CancellationToken)
+		{
+			string Id = GetFullRefId(NamespaceId, BucketId, RefId);
+
+			RefImpl? RefImpl = await Refs.FindOneAndUpdateAsync<RefImpl>(x => x.Id == Id, Builders<RefImpl>.Update.Set(x => x.LastAccessTime, DateTime.UtcNow), new FindOneAndUpdateOptions<RefImpl> { ReturnDocument = ReturnDocument.After }, CancellationToken);
+			if (RefImpl == null || !RefImpl.Finalized)
+			{
+				throw new RefNotFoundException(NamespaceId, BucketId, RefId);
+			}
+
+			return RefImpl;
+		}
+
+		/// <inheritdoc/>
+		public async Task<bool> HasRefAsync(NamespaceId NamespaceId, BucketId BucketId, RefId RefId, CancellationToken CancellationToken)
+		{
+			string Id = GetFullRefId(NamespaceId, BucketId, RefId);
+			RefImpl? RefImpl = await Refs.Find<RefImpl>(x => x.Id == Id).FirstOrDefaultAsync(CancellationToken);
+			return RefImpl != null;
+		}
+
+		/// <inheritdoc/>
+		public async Task<List<IoHash>> TrySetRefAsync(NamespaceId NamespaceId, BucketId BucketId, RefId RefId, CbObject Value, CancellationToken CancellationToken)
+		{
+			RefImpl NewRef = new RefImpl();
+			NewRef.Id = GetFullRefId(NamespaceId, BucketId, RefId);
+			NewRef.Data = Value.GetView().ToArray();
+			await Refs.ReplaceOneAsync(x => x.Id == NewRef.Id, NewRef, new ReplaceOptions { IsUpsert = true });
+
+			IoHash Hash = IoHash.Compute(Value.GetView().Span);
+			return await TryFinalizeRefAsync(NamespaceId, BucketId, RefId, Hash, CancellationToken);
+		}
+
+		/// <inheritdoc/>
+		public async Task<List<IoHash>> TryFinalizeRefAsync(NamespaceId NamespaceId, BucketId BucketId, RefId RefId, IoHash Hash, CancellationToken CancellationToken)
+		{
+			string Id = GetFullRefId(NamespaceId, BucketId, RefId);
+
+			RefImpl Ref = await Refs.Find(x => x.Id == Id).FirstAsync(CancellationToken);
+
+			HashSet<IoHash> MissingHashes = new HashSet<IoHash>();
+			if (!Ref.Finalized)
+			{
+				HashSet<IoHash> CheckedObjectHashes = new HashSet<IoHash>();
+				HashSet<IoHash> CheckedBinaryHashes = new HashSet<IoHash>();
+				await FinalizeInternalAsync(Ref.NamespaceId, Ref.Value, CheckedObjectHashes, CheckedBinaryHashes, MissingHashes);
+
+				if (MissingHashes.Count == 0)
+				{
+					await Refs.UpdateOneAsync(x => x.Id == Id, Builders<RefImpl>.Update.Set(x => x.Finalized, true));
+				}
+			}
+			return MissingHashes.ToList();
+		}
+
+		async Task FinalizeInternalAsync(NamespaceId NamespaceId, CbObject Object, HashSet<IoHash> CheckedObjectHashes, HashSet<IoHash> CheckedBinaryHashes, HashSet<IoHash> MissingHashes)
+		{
+			List<IoHash> NewObjectHashes = new List<IoHash>();
+			List<IoHash> NewBinaryHashes = new List<IoHash>();
+			Object.IterateAttachments(Field =>
+			{
+				IoHash AttachmentHash = Field.AsAttachment();
+				if (Field.IsObjectAttachment())
+				{
+					if (CheckedObjectHashes.Add(AttachmentHash))
+					{
+						NewObjectHashes.Add(AttachmentHash);
+					}
+				}
+				else
+				{
+					if (CheckedBinaryHashes.Add(AttachmentHash))
+					{
+						NewBinaryHashes.Add(AttachmentHash);
+					}
+				}
+			});
+
+			foreach (IoHash NewHash in NewObjectHashes)
+			{
+				string ObjPath = GetFullBlobPath(NamespaceId, NewHash);
+
+				ReadOnlyMemory<byte>? Data = await StorageBackend.ReadBytesAsync(ObjPath);
+				if (Data == null)
+				{
+					MissingHashes.Add(NewHash);
+				}
+				else
+				{
+					await FinalizeInternalAsync(NamespaceId, new CbObject(Data.Value), CheckedObjectHashes, CheckedBinaryHashes, MissingHashes);
+				}
+			}
+
+			foreach (IoHash NewBinaryHash in NewBinaryHashes)
+			{
+				string Path = GetFullBlobPath(NamespaceId, NewBinaryHash);
+				if (!await StorageBackend.ExistsAsync(Path))
+				{
+					MissingHashes.Add(NewBinaryHash);
+				}
+			}
+		}
+
+		/// <inheritdoc/>
+		public async Task<bool> DeleteRefAsync(NamespaceId NamespaceId, BucketId BucketId, RefId RefId, CancellationToken CancellationToken)
+		{
+			string Id = GetFullRefId(NamespaceId, BucketId, RefId);
+			DeleteResult Result = await Refs.DeleteOneAsync(x => x.Id == Id, CancellationToken);
+			return Result.DeletedCount > 0;
+		}
+
+		/// <inheritdoc/>
+		public async Task<List<RefId>> FindMissingRefsAsync(NamespaceId NamespaceId, BucketId BucketId, List<RefId> RefIds, CancellationToken CancellationToken)
+		{
+			List<RefId> MissingRefIds = new List<RefId>();
+			foreach (RefId RefId in RefIds)
+			{
+				if(!await HasRefAsync(NamespaceId, BucketId, RefId, CancellationToken))
+				{
+					MissingRefIds.Add(RefId);
+				}
+			}
+			return MissingRefIds;
+		}
+
+		/// <summary>
+		/// Gets the formatted id for a particular ref
+		/// </summary>
+		static string GetFullRefId(NamespaceId NamespaceId, BucketId BucketId, RefId RefId)
+		{
+			return $"{NamespaceId}/{BucketId}/{RefId}";
+		}
+
+		/// <summary>
+		/// Gets the path for a blob
+		/// </summary>
+		static string GetFullBlobPath(NamespaceId NamespaceId, IoHash Hash)
+		{
+			string HashText = Hash.ToString();
+			return $"blobs2/{NamespaceId}/{HashText[0..2]}/{HashText[2..4]}/{HashText}.blob";
+		}
+	}
+}
diff --git a/Engine/Source/Programs/Horde/Horde.Build/Storage/Services/SimpleStorageService.cs b/Engine/Source/Programs/Horde/Horde.Build/Storage/Services/SimpleStorageService.cs
deleted file mode 100644
diff --git a/Engine/Source/Programs/Horde/Horde.Build/Utilities/ObjectSet.cs b/Engine/Source/Programs/Horde/Horde.Build/Utilities/ObjectSet.cs
--- a/Engine/Source/Programs/Horde/Horde.Build/Utilities/ObjectSet.cs
+++ b/Engine/Source/Programs/Horde/Horde.Build/Utilities/ObjectSet.cs
@@ -1,6 +1,7 @@
 // Copyright Epic Games, Inc. All Rights Reserved.
 
 using EpicGames.Core;
+using EpicGames.Horde.Storage;
 using EpicGames.Serialization;
 using EpicGames.Serialization.Converters;
 using HordeServer.Storage;
@@ -12,8 +13,6 @@ using System.Threading.Tasks;
 
 namespace HordeServer.Utilities
 {
-	using NamespaceId = StringId<INamespace>;
-
 	/// <summary>
 	/// Information about a blob stored in a blob pack file
 	/// </summary>
@@ -138,7 +137,7 @@ namespace HordeServer.Utilities
 	/// </summary>
 	class ObjectSet
 	{
-		readonly IBlobCollection BlobCollection;
+		readonly IStorageClient StorageClient;
 		readonly NamespaceId NamespaceId;
 		public int MaxPackSize { get; }
 
@@ -157,13 +156,13 @@ namespace HordeServer.Utilities
 		/// <summary>
 		/// Constructor
 		/// </summary>
-		/// <param name="BlobCollection"></param>
+		/// <param name="StorageClient"></param>
 		/// <param name="NamespaceId"></param>
 		/// <param name="MaxPackSize"></param>
 		/// <param name="Time">The initial update time; used to determine the age of blobs</param>
-		public ObjectSet(IBlobCollection BlobCollection, NamespaceId NamespaceId, int MaxPackSize, DateTime Time)
+		public ObjectSet(IStorageClient StorageClient, NamespaceId NamespaceId, int MaxPackSize, DateTime Time)
 		{
-			this.BlobCollection = BlobCollection;
+			this.StorageClient = StorageClient;
 			this.NamespaceId = NamespaceId;
 			this.MaxPackSize = MaxPackSize;
 
@@ -263,7 +262,7 @@ namespace HordeServer.Utilities
 			{
 				if (Pack.TryGetEntry(Hash, out Entry))
 				{
-					ReadOnlyMemory<byte> PackData = await BlobCollection.ReadBytesAsync(NamespaceId, Pack.DataHash);
+					ReadOnlyMemory<byte> PackData = await StorageClient.ReadBlobToMemoryAsync(NamespaceId, Pack.DataHash);
 					return PackData.Slice(Entry.Offset, Entry.Length);
 				}
 			}
@@ -371,7 +370,7 @@ namespace HordeServer.Utilities
 				}
 
 				// Get the data for this blob
-				ReadOnlyMemory<byte> MergeData = await BlobCollection.ReadBytesAsync(NamespaceId, MergePack.DataHash);
+				ReadOnlyMemory<byte> MergeData = await StorageClient.ReadBlobToMemoryAsync(NamespaceId, MergePack.DataHash);
 
 				// Add anything that's still part of the live set into the new blobs
 				int Offset = 0;
@@ -449,7 +448,7 @@ namespace HordeServer.Utilities
 				Array.Resize(ref NextPackData, NextPackSize);
 				ReadOnlyMemory<byte> Data = NextPackData;
 				IoHash DataHash = IoHash.Compute(Data.Span);
-				WriteTasks.Add(Task.Run(() => BlobCollection.WriteBytesAsync(NamespaceId, DataHash, Data)));
+				WriteTasks.Add(Task.Run(() => StorageClient.WriteBlobFromMemoryAsync(NamespaceId, DataHash, Data)));
 
 				// Create the new index
 				ObjectPackIndex Index = new ObjectPackIndex(Time, NextPackEntries.ToArray(), DataHash, NextPackSize);
diff --git a/Engine/Source/Programs/Horde/Horde.Build/Utilities/ProtoExtensions.cs b/Engine/Source/Programs/Horde/Horde.Build/Utilities/ProtoExtensions.cs
new file mode 100644
--- /dev/null
+++ b/Engine/Source/Programs/Horde/Horde.Build/Utilities/ProtoExtensions.cs
@@ -0,0 +1,77 @@
+// Copyright Epic Games, Inc. All Rights Reserved.
+
+using EpicGames.Core;
+using EpicGames.Horde.Storage;
+using Google.Protobuf;
+
+namespace HordeCommon
+{
+	/// <summary>
+	/// Methods for IoHashWrapper
+	/// </summary>
+	public partial class IoHashWrapper
+	{
+		/// <summary>
+		/// Constructor
+		/// </summary>
+		/// <param name="Hash"></param>
+		public IoHashWrapper(IoHash Hash)
+		{
+			this.Hash = ByteString.CopyFrom(Hash.Memory.ToArray());
+		}
+
+		/// <summary>
+		/// Convert from an IoHashWrapper to an IoHash
+		/// </summary>
+		/// <param name="Hash"></param>
+		public static implicit operator IoHash(IoHashWrapper Hash)
+		{
+			return new IoHash(Hash.Hash.ToByteArray());
+		}
+
+		/// <summary>
+		/// Convert from an IoHash to an IoHashWrapper
+		/// </summary>
+		/// <param name="Hash"></param>
+		public static implicit operator IoHashWrapper(IoHash Hash)
+		{
+			return new IoHashWrapper(Hash);
+		}
+	}
+
+	/// <summary>
+	/// Methods for IoHashWrapper
+	/// </summary>
+	public partial class RefIdWrapper
+	{
+		/// <summary>
+		/// Constructor
+		/// </summary>
+		/// <param name="RefId"></param>
+		public RefIdWrapper(RefId RefId)
+		{
+			this.Hash = ByteString.CopyFrom(RefId.Hash.Span);
+		}
+
+		/// <summary>
+		/// Convert from an IoHashWrapper to an IoHash
+		/// </summary>
+		/// <param name="RefId"></param>
+		public static implicit operator RefId(RefIdWrapper RefId) => RefId.AsRefId();
+
+		/// <summary>
+		/// Convert from an IoHash to an IoHashWrapper
+		/// </summary>
+		/// <param name="RefId"></param>
+		public static implicit operator RefIdWrapper(RefId RefId)
+		{
+			return new RefIdWrapper(RefId);
+		}
+
+		/// <summary>
+		/// Convert from an IoHashWrapper to an IoHash
+		/// </summary>
+		/// <returns></returns>
+		public RefId AsRefId() => new RefId(new IoHash(Hash.ToByteArray()));
+	}
+}
diff --git a/Engine/Source/Programs/Horde/Horde.Build/appsettings.json b/Engine/Source/Programs/Horde/Horde.Build/appsettings.json
--- a/Engine/Source/Programs/Horde/Horde.Build/appsettings.json
+++ b/Engine/Source/Programs/Horde/Horde.Build/appsettings.json
@@ -4,7 +4,10 @@
 		"DisableAuth": true,
 		"CorsEnabled": true,
 		"CorsOrigin": "https://localhost:5003;http://localhost:3000",
-		"LogFileProviderType": "Local"
+		"LogFileProviderType": "Local",
+		"Storage": {
+			"Url": "https://localhost:57001"
+		}
 	},
 	"Serilog": {
 		"MinimumLevel": {
diff --git a/Engine/Source/Programs/Shared/EpicGames.Core/HttpClientExtensions.cs b/Engine/Source/Programs/Shared/EpicGames.Core/HttpClientExtensions.cs
--- a/Engine/Source/Programs/Shared/EpicGames.Core/HttpClientExtensions.cs
+++ b/Engine/Source/Programs/Shared/EpicGames.Core/HttpClientExtensions.cs
@@ -122,7 +122,7 @@ namespace EpicGames.Core
 		private static async Task<T> ParseJsonContent<T>(HttpResponseMessage Message)
 		{
 			byte[] Bytes = await Message.Content.ReadAsByteArrayAsync();
-			return JsonSerializer.Deserialize<T>(Bytes);
+			return JsonSerializer.Deserialize<T>(Bytes, new JsonSerializerOptions { PropertyNameCaseInsensitive = true });
 		}
 	}
 }
diff --git a/Engine/Source/Programs/Shared/EpicGames.Horde.Client/EpicGames.Horde.Client.csproj b/Engine/Source/Programs/Shared/EpicGames.Horde.Client/EpicGames.Horde.Client.csproj
deleted file mode 100644
diff --git a/Engine/Source/Programs/Shared/EpicGames.Horde.Client/Storage/BlobStoreExtensions.cs b/Engine/Source/Programs/Shared/EpicGames.Horde.Client/Storage/BlobStoreExtensions.cs
deleted file mode 100644
diff --git a/Engine/Source/Programs/Shared/EpicGames.Horde/Auth/IAuthProvider.cs b/Engine/Source/Programs/Shared/EpicGames.Horde/Auth/IAuthProvider.cs
new file mode 100644
--- /dev/null
+++ b/Engine/Source/Programs/Shared/EpicGames.Horde/Auth/IAuthProvider.cs
@@ -0,0 +1,56 @@
+// Copyright Epic Games, Inc. All Rights Reserved.
+
+using System;
+using System.Collections.Generic;
+using System.Net.Http;
+using System.Text;
+using System.Threading;
+using System.Threading.Tasks;
+
+namespace EpicGames.Horde.Auth
+{
+	/// <summary>
+	/// Interface for an authorization provider.
+	/// </summary>
+	public interface IAuthProvider
+	{
+		/// <summary>
+		/// Add authorization headers to the given request
+		/// </summary>
+		/// <param name="Request">Request to modify</param>
+		/// <param name="CancellationToken">Cancellation token for the request</param>
+		public Task AddAuthorizationAsync(HttpRequestMessage Request, CancellationToken CancellationToken);
+	}
+
+	/// <summary>
+	/// Typed authorization provider. For use with dependency injection.
+	/// </summary>
+	/// <typeparam name="T"></typeparam>
+	public interface IAuthProvider<T> : IAuthProvider
+	{
+	}
+
+	/// <summary>
+	/// Implementation of <see cref="IAuthProvider{T}"/> which wraps an underlying auth provider implementation
+	/// </summary>
+	/// <typeparam name="T">Nominal type for the implementation</typeparam>
+	public class TypedAuthProvider<T> : IAuthProvider<T>
+	{
+		public IAuthProvider AuthProvider;
+
+		/// <summary>
+		/// Constructor
+		/// </summary>
+		/// <param name="AuthProvider"></param>
+		public TypedAuthProvider(IAuthProvider AuthProvider)
+		{
+			this.AuthProvider = AuthProvider;
+		}
+
+		/// <inheritdoc/>
+		public Task AddAuthorizationAsync(HttpRequestMessage Request, CancellationToken CancellationToken)
+		{
+			return AuthProvider.AddAuthorizationAsync(Request, CancellationToken);
+		}
+	}
+}
diff --git a/Engine/Source/Programs/Shared/EpicGames.Horde/Auth/OAuth2AuthProvider.cs b/Engine/Source/Programs/Shared/EpicGames.Horde/Auth/OAuth2AuthProvider.cs
new file mode 100644
--- /dev/null
+++ b/Engine/Source/Programs/Shared/EpicGames.Horde/Auth/OAuth2AuthProvider.cs
@@ -0,0 +1,162 @@
+// Copyright Epic Games, Inc. All Rights Reserved.
+
+using EpicGames.Core;
+using EpicGames.Serialization;
+using Microsoft.Extensions.Options;
+using System;
+using System.Collections.Generic;
+using System.ComponentModel.DataAnnotations;
+using System.Net;
+using System.Net.Http;
+using System.Text;
+using System.Text.Json;
+using System.Threading;
+using System.Threading.Tasks;
+
+namespace EpicGames.Horde.Auth
+{
+	/// <summary>
+	/// Exception thrown due to failed authorization
+	/// </summary>
+	public class AuthenticationException : Exception
+	{
+		public AuthenticationException(string Message, Exception? InnerException)
+			: base(Message, InnerException)
+		{
+		}
+	}
+
+	/// <summary>
+	/// Options for authenticating particular requests
+	/// </summary>
+	public class OAuth2Options
+	{
+		/// <summary>
+		/// Url of the auth server
+		/// </summary>
+		public string AuthUrl { get; set; } = String.Empty;
+
+		/// <summary>
+		/// Type of grant
+		/// </summary>
+		public string GrantType { get; set; } = String.Empty;
+
+		/// <summary>
+		/// Client id
+		/// </summary>
+		public string ClientId { get; set; } = String.Empty;
+
+		/// <summary>
+		/// Client secret
+		/// </summary>
+		public string ClientSecret { get; set; } = String.Empty;
+
+		/// <summary>
+		/// Scope of the token
+		/// </summary>
+		public string Scope { get; set; } = String.Empty;
+	}
+
+	/// <summary>
+	/// Http message handler which adds an OAuth authorization header using a cached/periodically refreshed bearer token
+	/// </summary>
+	public class OAuth2AuthProvider : IAuthProvider
+	{
+		class ClientCredentialsResponse
+		{
+			public string? access_token { get; set; }
+			public string? token_type { get; set; }
+			public int? expires_in { get; set; }
+			public string? scope { get; set; }
+		}
+
+		HttpClient Client;
+		OAuth2Options Options;
+		string CachedAccessToken = String.Empty;
+		DateTime ExpiresAt = DateTime.MinValue;
+
+		public OAuth2AuthProvider(HttpClient Client, OAuth2Options Options)
+		{
+			this.Client = Client;
+			this.Options = Options;
+		}
+
+		/// <inheritdoc/>
+		public async Task AddAuthorizationAsync(HttpRequestMessage Request, CancellationToken CancellationToken)
+		{
+			if (DateTime.UtcNow > ExpiresAt)
+			{
+				await UpdateAccessTokenAsync(CancellationToken);
+			}
+
+			Request.Headers.Add("Authorization", $"Bearer {CachedAccessToken}");
+		}
+
+		/// <summary>
+		/// Updates the current access token
+		/// </summary>
+		/// <returns></returns>
+		async Task UpdateAccessTokenAsync(CancellationToken CancellationToken)
+		{
+			KeyValuePair<string, string>[] Content = new KeyValuePair<string, string>[]
+			{
+				new KeyValuePair<string, string>("grant_type", Options.GrantType),
+				new KeyValuePair<string, string>("client_id", Options.ClientId),
+				new KeyValuePair<string, string>("client_secret", Options.ClientSecret),
+				new KeyValuePair<string, string>("scope", Options.Scope)
+			};
+
+			try
+			{
+				using HttpRequestMessage Message = new HttpRequestMessage(HttpMethod.Post, Options.AuthUrl);
+				Message.Content = new FormUrlEncodedContent(Content);
+
+				HttpResponseMessage Response = await Client.SendAsync(Message, CancellationToken);
+				if (!Response.IsSuccessStatusCode)
+				{
+					throw new AuthenticationException($"Authentication failed. Response: {Response.Content}", null);
+				}
+
+				byte[] ResponseData = await Response.Content.ReadAsByteArrayAsync();
+				ClientCredentialsResponse Result = JsonSerializer.Deserialize<ClientCredentialsResponse>(ResponseData)!;
+
+				string? accessToken = Result?.access_token;
+				if (string.IsNullOrEmpty(accessToken))
+				{
+					throw new AuthenticationException("The authentication token received by the server is null or empty. Body received was: " + Encoding.UTF8.GetString(ResponseData), null);
+				}
+				CachedAccessToken = accessToken;
+				// renew after half the renewal time
+				ExpiresAt = DateTime.UtcNow + TimeSpan.FromSeconds((Result?.expires_in ?? 3200) / 2.0);
+			}
+			catch (WebException Ex)
+			{
+				throw new AuthenticationException("Unable to authenticate.", Ex);
+			}
+		}
+	}
+
+	/// <summary>
+	/// Factory for creating OAuth2AuthProvider instances from a set of options
+	/// </summary>
+	public class OAuth2AuthProviderFactory<T>
+	{
+		HttpClient HttpClient;
+
+		/// <summary>
+		/// Constructor
+		/// </summary>
+		/// <param name="HttpClient"></param>
+		public OAuth2AuthProviderFactory(HttpClient HttpClient)
+		{
+			this.HttpClient = HttpClient;
+		}
+
+		/// <summary>
+		/// Create an instance of the auth provider
+		/// </summary>
+		/// <param name="Options"></param>
+		/// <returns></returns>
+		public IAuthProvider<T> Create(OAuth2Options Options) => new TypedAuthProvider<T>(new OAuth2AuthProvider(HttpClient, Options));
+	}
+}
diff --git a/Engine/Source/Programs/Shared/EpicGames.Horde/Auth/OAuthMessageHandler.cs b/Engine/Source/Programs/Shared/EpicGames.Horde/Auth/OAuthMessageHandler.cs
deleted file mode 100644
diff --git a/Engine/Source/Programs/Shared/EpicGames.Horde/Common/CbBinaryAttachmentWrapper.cs b/Engine/Source/Programs/Shared/EpicGames.Horde/Common/CbBinaryAttachmentWrapper.cs
deleted file mode 100644
diff --git a/Engine/Source/Programs/Shared/EpicGames.Horde/Common/CbObjectAttachmentWrapper.cs b/Engine/Source/Programs/Shared/EpicGames.Horde/Common/CbObjectAttachmentWrapper.cs
deleted file mode 100644
diff --git a/Engine/Source/Programs/Shared/EpicGames.Horde/Common/CbObjectWrapper.cs b/Engine/Source/Programs/Shared/EpicGames.Horde/Common/CbObjectWrapper.cs
deleted file mode 100644
diff --git a/Engine/Source/Programs/Shared/EpicGames.Horde/Common/IoHashWrapper.cs b/Engine/Source/Programs/Shared/EpicGames.Horde/Common/IoHashWrapper.cs
deleted file mode 100644
diff --git a/Engine/Source/Programs/Shared/EpicGames.Horde/Compute/Api.cs b/Engine/Source/Programs/Shared/EpicGames.Horde/Compute/Api.cs
deleted file mode 100644
diff --git a/Engine/Source/Programs/Shared/EpicGames.Horde/Compute/ChannelId.cs b/Engine/Source/Programs/Shared/EpicGames.Horde/Compute/ChannelId.cs
new file mode 100644
--- /dev/null
+++ b/Engine/Source/Programs/Shared/EpicGames.Horde/Compute/ChannelId.cs
@@ -0,0 +1,95 @@
+// Copyright Epic Games, Inc. All Rights Reserved.
+
+using EpicGames.Core;
+using EpicGames.Serialization;
+using System;
+using System.Collections.Generic;
+using System.ComponentModel;
+using System.Diagnostics.CodeAnalysis;
+using System.Globalization;
+using System.Text;
+using System.Text.Json;
+using System.Text.Json.Serialization;
+
+namespace EpicGames.Horde.Compute
+{
+	/// <summary>
+	/// Identifier for a channel for receiving compute responses
+	/// </summary>
+	[CbConverter(typeof(ChannelIdCbConverter))]
+	[JsonConverter(typeof(ChannelIdJsonConverter))]
+	[TypeConverter(typeof(ChannelIdTypeConverter))]
+	public struct ChannelId : IEquatable<ChannelId>
+	{
+		/// <summary>
+		/// The text representing this id
+		/// </summary>
+		readonly StringId Inner;
+
+		/// <summary>
+		/// Constructor
+		/// </summary>
+		/// <param name="Input">Unique id for the string</param>
+		public ChannelId(string Input)
+		{
+			Inner = new StringId(Input);
+		}
+
+		/// <inheritdoc/>
+		public override bool Equals(object? Obj) => Obj is ChannelId Id && Inner.Equals(Id.Inner);
+
+		/// <inheritdoc/>
+		public override int GetHashCode() => Inner.GetHashCode();
+
+		/// <inheritdoc/>
+		public bool Equals(ChannelId Other) => Inner.Equals(Other.Inner);
+
+		/// <inheritdoc/>
+		public override string ToString() => Inner.ToString();
+
+		/// <inheritdoc cref="StringId.op_Equality"/>
+		public static bool operator ==(ChannelId Left, ChannelId Right) => Left.Inner == Right.Inner;
+
+		/// <inheritdoc cref="StringId.op_Inequality"/>
+		public static bool operator !=(ChannelId Left, ChannelId Right) => Left.Inner != Right.Inner;
+	}
+
+	/// <summary>
+	/// Compact binary converter for ChannelId
+	/// </summary>
+	sealed class ChannelIdCbConverter : CbConverterBase<ChannelId>
+	{
+		/// <inheritdoc/>
+		public override ChannelId Read(CbField Field) => new ChannelId(Field.AsString());
+
+		/// <inheritdoc/>
+		public override void Write(CbWriter Writer, ChannelId Value) => Writer.WriteStringValue(Value.ToString());
+
+		/// <inheritdoc/>
+		public override void WriteNamed(CbWriter Writer, Utf8String Name, ChannelId Value) => Writer.WriteString(Name, Value.ToString());
+	}
+
+	/// <summary>
+	/// Type converter for ChannelId to and from JSON
+	/// </summary>
+	sealed class ChannelIdJsonConverter : JsonConverter<ChannelId>
+	{
+		/// <inheritdoc/>
+		public override ChannelId Read(ref Utf8JsonReader Reader, Type TypeToConvert, JsonSerializerOptions Options) => new ChannelId(Reader.GetString() ?? String.Empty);
+
+		/// <inheritdoc/>
+		public override void Write(Utf8JsonWriter Writer, ChannelId Value, JsonSerializerOptions Options) => Writer.WriteStringValue(Value.ToString());
+	}
+
+	/// <summary>
+	/// Type converter from strings to ChannelId objects
+	/// </summary>
+	sealed class ChannelIdTypeConverter : TypeConverter
+	{
+		/// <inheritdoc/>
+		public override bool CanConvertFrom(ITypeDescriptorContext Context, Type SourceType) => SourceType == typeof(string);
+
+		/// <inheritdoc/>
+		public override object ConvertFrom(ITypeDescriptorContext Context, CultureInfo Culture, object Value) => new ChannelId((string)Value);
+	}
+}
diff --git a/Engine/Source/Programs/Shared/EpicGames.Horde/Compute/ClusterId.cs b/Engine/Source/Programs/Shared/EpicGames.Horde/Compute/ClusterId.cs
new file mode 100644
--- /dev/null
+++ b/Engine/Source/Programs/Shared/EpicGames.Horde/Compute/ClusterId.cs
@@ -0,0 +1,96 @@
+// Copyright Epic Games, Inc. All Rights Reserved.
+
+using EpicGames.Core;
+using EpicGames.Serialization;
+using System;
+using System.Collections.Generic;
+using System.ComponentModel;
+using System.Diagnostics.CodeAnalysis;
+using System.Globalization;
+using System.Text;
+using System.Text.Json;
+using System.Text.Json.Serialization;
+
+namespace EpicGames.Horde.Compute
+{
+	/// <summary>
+	/// Identifier for a compute cluster
+	/// </summary>
+	[CbConverter(typeof(ClusterIdCbConverter))]
+	[JsonConverter(typeof(ClusterIdJsonConverter))]
+	[TypeConverter(typeof(ClusterIdTypeConverter))]
+	public struct ClusterId : IEquatable<ClusterId>
+	{
+		/// <summary>
+		/// The text representing this id
+		/// </summary>
+		readonly StringId Inner;
+
+		/// <summary>
+		/// Constructor
+		/// </summary>
+		/// <param name="Input">Unique id for the string</param>
+		public ClusterId(string Input)
+		{
+			Inner = new StringId(Input);
+		}
+
+		/// <inheritdoc/>
+		public override bool Equals(object? Obj) => Obj is ClusterId Id && Inner.Equals(Id.Inner);
+
+		/// <inheritdoc/>
+		public override int GetHashCode() => Inner.GetHashCode();
+
+		/// <inheritdoc/>
+		public bool Equals(ClusterId Other) => Inner.Equals(Other.Inner);
+
+		/// <inheritdoc/>
+		public override string ToString() => Inner.ToString();
+
+		/// <inheritdoc cref="StringId.op_Equality"/>
+		public static bool operator ==(ClusterId Left, ClusterId Right) => Left.Inner == Right.Inner;
+
+		/// <inheritdoc cref="StringId.op_Inequality"/>
+		public static bool operator !=(ClusterId Left, ClusterId Right) => Left.Inner != Right.Inner;
+	}
+
+	/// <summary>
+	/// Compact binary converter for ClusterId
+	/// </summary>
+	sealed class ClusterIdCbConverter : CbConverterBase<ClusterId>
+	{
+		/// <inheritdoc/>
+		public override ClusterId Read(CbField Field) => new ClusterId(Field.AsString());
+
+		/// <inheritdoc/>
+		public override void Write(CbWriter Writer, ClusterId Value) => Writer.WriteStringValue(Value.ToString());
+
+		/// <inheritdoc/>
+		public override void WriteNamed(CbWriter Writer, Utf8String Name, ClusterId Value) => Writer.WriteString(Name, Value.ToString());
+	}
+
+	/// <summary>
+	/// Type converter for ClusterId to and from JSON
+	/// </summary>
+	sealed class ClusterIdJsonConverter : JsonConverter<ClusterId>
+	{
+		/// <inheritdoc/>
+		public override ClusterId Read(ref Utf8JsonReader Reader, Type TypeToConvert, JsonSerializerOptions Options) => new ClusterId(Reader.GetString() ?? String.Empty);
+
+		/// <inheritdoc/>
+		public override void Write(Utf8JsonWriter Writer, ClusterId Value, JsonSerializerOptions Options) => Writer.WriteStringValue(Value.ToString());
+	}
+
+	/// <summary>
+	/// Type converter from strings to ClusterId objects
+	/// </summary>
+	sealed class ClusterIdTypeConverter : TypeConverter
+	{
+		/// <inheritdoc/>
+		public override bool CanConvertFrom(ITypeDescriptorContext Context, Type SourceType) => SourceType == typeof(string);
+
+		/// <inheritdoc/>
+		public override object ConvertFrom(ITypeDescriptorContext Context, CultureInfo Culture, object Value) => new ClusterId((string)Value);
+	}
+}
+
diff --git a/Engine/Source/Programs/Shared/EpicGames.Horde/Compute/IComputeClient.cs b/Engine/Source/Programs/Shared/EpicGames.Horde/Compute/IComputeClient.cs
new file mode 100644
--- /dev/null
+++ b/Engine/Source/Programs/Shared/EpicGames.Horde/Compute/IComputeClient.cs
@@ -0,0 +1,204 @@
+// Copyright Epic Games, Inc. All Rights Reserved.
+
+using EpicGames.Core;
+using EpicGames.Horde.Compute.Impl;
+using EpicGames.Horde.Storage;
+using EpicGames.Serialization;
+using System;
+using System.Collections.Generic;
+using System.Runtime.CompilerServices;
+using System.Text;
+using System.Threading;
+using System.Threading.Tasks;
+
+namespace EpicGames.Horde.Compute
+{
+	/// <summary>
+	/// Response decribing a cluster
+	/// </summary>
+	public interface IComputeClusterInfo
+	{
+		/// <summary>
+		/// Id of the cluster
+		/// </summary>
+		public ClusterId Id { get; }
+
+		/// <summary>
+		/// Namespace containing the requests and responses
+		/// </summary>
+		public NamespaceId NamespaceId { get; }
+
+		/// <summary>
+		/// The request bucket id
+		/// </summary>
+		public BucketId RequestBucketId { get; }
+
+		/// <summary>
+		/// Bucket to store response refs
+		/// </summary>
+		public BucketId ResponseBucketId { get; }
+	}
+
+	/// <summary>
+	/// State of a compute task
+	/// </summary>
+	public enum ComputeTaskState
+	{
+		/// <summary>
+		/// The task is queued for execution
+		/// </summary>
+		Queued = 0,
+
+		/// <summary>
+		/// Currently being executed by a remote agent
+		/// </summary>
+		Executing = 1,
+
+		/// <summary>
+		/// Completed
+		/// </summary>
+		Complete = 2,
+	}
+
+	/// <summary>
+	/// Outcome of a task, if it's in the complete state
+	/// </summary>
+	public enum ComputeTaskOutcome
+	{
+		/// <summary>
+		/// The task was executed successfully
+		/// </summary>
+		Success = 0,
+
+		/// <summary>
+		/// The lease failed with an error
+		/// </summary>
+		Failed = 1,
+
+		/// <summary>
+		/// The lease was cancelled
+		/// </summary>
+		Cancelled = 2,
+
+		/// <summary>
+		/// The task ran but did not return a result for unknown reasons
+		/// </summary>
+		NoResult = 3,
+
+		/// <summary>
+		/// The item was not scheduled for execution in the default time period, and was expired.
+		/// </summary>
+		Expired = 4,
+
+		/// <summary>
+		/// A blob could not be found. The detail field includes the hash of the blob.
+		/// </summary>
+		BlobNotFound = 5,
+
+		/// <summary>
+		/// An uncaught exception occurred during task execution
+		/// </summary>
+		Exception = 6,
+	}
+
+	/// <summary>
+	/// Supplies information about the current execution state of a task
+	/// </summary>
+	public interface IComputeTaskInfo
+	{
+		/// <summary>
+		/// Reference to the task decriptor that was requested
+		/// </summary>
+		public RefId TaskRefId { get; }
+
+		/// <summary>
+		/// Time that the event happened
+		/// </summary>
+		public DateTime Time { get; }
+
+		/// <summary>
+		/// New state of the task
+		/// </summary>
+		public ComputeTaskState State { get; }
+
+		/// <summary>
+		/// Outcome of the task execution.Note that this reflects the outcome of the execution rather than the outcome of the task.        
+		/// </summary>
+		public ComputeTaskOutcome Outcome { get; }
+
+		/// <summary>
+		/// Additional information about the outcome of the task.Dependent on the value of Outcome.        
+		/// </summary>
+		public string? Detail { get; }
+
+		/// <summary>
+		/// Ref containing the task result
+		/// </summary>
+		public RefId? ResultRefId { get; }
+
+		/// <summary>
+		/// When transitioning to the executing state, includes the name of the agent performing the work
+		/// </summary>
+		public string? AgentId { get; }
+
+		/// <summary>
+		/// When transitioning to the executing state, includes the id of the lease assigned to the agent
+		/// </summary>
+		public string? LeaseId { get; }
+	}
+
+	/// <summary>
+	/// Interface for communicating with a compute server
+	/// </summary>
+	public interface IComputeClient
+	{
+		/// <summary>
+		/// Gets information about a cluster
+		/// </summary>
+		/// <param name="ClusterId">The cluster to retreive information on</param>
+		/// <param name="CancellationToken">Cancellation token for the request</param>
+		/// <returns>Cluster information</returns>
+		Task<IComputeClusterInfo> GetClusterInfoAsync(ClusterId ClusterId, CancellationToken CancellationToken = default);
+
+		/// <summary>
+		/// Queues a set of tasks for execution
+		/// </summary>
+		/// <param name="ClusterId">Cluster for executing the tasks</param>
+		/// <param name="ChannelId">Indicates an identifier generated by the client to query responses on</param>
+		/// <param name="RequirementsHash">Requirements of the task to execute</param>
+		/// <param name="TaskRefIds">Refs describing the request</param>
+		/// <param name="SkipCacheLookup">Whether to skip the lookup of cached output for this task</param>
+		/// <param name="CancellationToken">Cancellation token for the request</param>
+		/// <returns>The new channel id</returns>
+		Task AddTasksAsync(ClusterId ClusterId, ChannelId ChannelId, IEnumerable<RefId> TaskRefIds, IoHash RequirementsHash, bool SkipCacheLookup, CancellationToken CancellationToken = default);
+
+		/// <summary>
+		/// Read updates from the given remote execution channel
+		/// </summary>
+		/// <param name="ClusterId">Cluster for executing the tasks</param>
+		/// <param name="ChannelId">Channel to receive updates on</param>
+		/// <param name="CancellationToken">Cancellation token for the operation</param>
+		/// <returns></returns>
+		IAsyncEnumerable<IComputeTaskInfo> GetTaskUpdatesAsync(ClusterId ClusterId, ChannelId ChannelId, CancellationToken CancellationToken = default);
+	}
+
+	/// <summary>
+	/// Extension methods for <see cref="IComputeClient"/>
+	/// </summary>
+	public static class ComputeClientExtensions
+	{
+		/// <summary>
+		/// Queues a single task for execution
+		/// </summary>
+		/// <param name="ComputeClient">The compute client instance</param>
+		/// <param name="ClusterId">Name of the profile to use for executing the tasks</param>
+		/// <param name="ChannelId">Channel to receive updates on</param>
+		/// <param name="TaskRefId">Ref describing the request</param>
+		/// <param name="RequirementsHash">Requirements of the task to execute</param>
+		/// <param name="SkipCacheLookup">Whether to skip the lookup of cached output for this task</param>
+		public static Task AddTaskAsync(this IComputeClient ComputeClient, ClusterId ClusterId, ChannelId ChannelId, RefId TaskRefId, IoHash RequirementsHash, bool SkipCacheLookup, CancellationToken CancellationToken = default)
+		{
+			return ComputeClient.AddTasksAsync(ClusterId, ChannelId, new[] { TaskRefId }, RequirementsHash, SkipCacheLookup, CancellationToken);
+		}
+	}
+}
diff --git a/Engine/Source/Programs/Shared/EpicGames.Horde/Compute/Impl/HttpComputeClient.cs b/Engine/Source/Programs/Shared/EpicGames.Horde/Compute/Impl/HttpComputeClient.cs
new file mode 100644
--- /dev/null
+++ b/Engine/Source/Programs/Shared/EpicGames.Horde/Compute/Impl/HttpComputeClient.cs
@@ -0,0 +1,166 @@
+// Copyright Epic Games, Inc. All Rights Reserved.
+
+using EpicGames.Core;
+using EpicGames.Horde.Common;
+using EpicGames.Horde.Storage;
+using EpicGames.Serialization;
+using System;
+using System.Collections.Generic;
+using System.ComponentModel.DataAnnotations;
+using System.Linq;
+using System.Net.Http;
+using System.Runtime.CompilerServices;
+using System.Text;
+using System.Threading;
+using System.Threading.Tasks;
+
+namespace EpicGames.Horde.Compute.Impl
+{
+	#region Requests / Responses
+
+	public class GetComputeClusterInfo : IComputeClusterInfo
+	{
+		/// <inheritdoc/>
+		public ClusterId Id { get; set; }
+
+		/// <inheritdoc/>
+		public NamespaceId NamespaceId { get; set; }
+
+		/// <inheritdoc/>
+		public BucketId RequestBucketId { get; set; }
+
+		/// <inheritdoc/>
+		public BucketId ResponseBucketId { get; set; }
+	}
+
+	/// <summary>
+	/// Supplies information about the current execution state of a task
+	/// </summary>
+	public class GetTaskUpdateResponse : IComputeTaskInfo
+	{
+		/// <inheritdoc/>
+		[CbField("h")]
+		public RefId TaskRefId { get; set; }
+
+		/// <inheritdoc/>
+		[CbField("t")]
+		public DateTime Time { get; set; }
+
+		/// <inheritdoc/>
+		[CbField("s")]
+		public ComputeTaskState State { get; set; }
+
+		/// <inheritdoc/>
+		[CbField("o")]
+		public ComputeTaskOutcome Outcome { get; set; }
+
+		/// <inheritdoc/>
+		[CbField("d")]
+		public string? Detail { get; set; }
+
+		/// <inheritdoc/>
+		[CbField("r")]
+		public RefId? ResultRefId { get; set; }
+
+		/// <inheritdoc/>
+		[CbField("a")]
+		public string? AgentId { get; set; }
+
+		/// <inheritdoc/>
+		[CbField("l")]
+		public string? LeaseId { get; set; }
+	}
+
+	/// <summary>
+	/// Request to add tasks to the compute queue
+	/// </summary>
+	public class AddTasksRequest
+	{
+		/// <inheritdoc cref="AddTasksRpcRequest.ChannelId"/>
+		[CbField("c")]
+		public ChannelId ChannelId { get; set; }
+
+		/// <inheritdoc cref="AddTasksRpcRequest.TaskHashes"/>
+		[Required, CbField("t")]
+		public List<RefId> TaskRefIds { get; set; } = new List<RefId>();
+
+		/// <inheritdoc cref="AddTasksRpcRequest.RequirementsHash"/>
+		[CbField("r")]
+		public CbObjectAttachment RequirementsHash { get; set; }
+
+		/// <inheritdoc cref="AddTasksRpcRequest.DoNotCache"/>
+		[CbField("nc")]
+		public bool DoNotCache { get; set; }
+	}
+
+	/// <summary>
+	/// Specifies updates from the given channel
+	/// </summary>
+	public class GetTaskUpdatesResponse
+	{
+		/// <summary>
+		/// Task updates
+		/// </summary>
+		[CbField("u")]
+		public List<GetTaskUpdateResponse> Updates { get; set; } = new List<GetTaskUpdateResponse>();
+	}
+
+	#endregion
+
+	public class HttpComputeClient : IComputeClient
+	{
+		HttpClient HttpClient;
+
+		public HttpComputeClient(HttpClient HttpClient)
+		{
+			this.HttpClient = HttpClient;
+		}
+
+		/// <inheritdoc/>
+		public async Task<IComputeClusterInfo> GetClusterInfoAsync(ClusterId ClusterId, CancellationToken CancellationToken)
+		{
+			return await HttpClient.GetAsync<GetComputeClusterInfo>($"api/v1/compute/{ClusterId}", CancellationToken);
+		}
+
+		/// <inheritdoc/>
+		public async Task AddTasksAsync(ClusterId ClusterId, ChannelId ChannelId, IEnumerable<RefId> TaskRefIds, IoHash RequirementsHash, bool SkipCacheLookup, CancellationToken CancellationToken)
+		{
+			AddTasksRequest AddTasks = new AddTasksRequest();
+			AddTasks.ChannelId = ChannelId;
+			AddTasks.TaskRefIds.AddRange(TaskRefIds);
+			AddTasks.RequirementsHash = RequirementsHash;
+			AddTasks.DoNotCache = SkipCacheLookup;
+
+			ReadOnlyMemoryContent Content = new ReadOnlyMemoryContent(CbSerializer.Serialize(AddTasks).GetView());
+			Content.Headers.Add("Content-Type", "application/x-ue-cb");
+
+			HttpRequestMessage Request = new HttpRequestMessage(HttpMethod.Post, $"api/v1/compute/{ClusterId}");
+			Request.Content = Content;
+
+			HttpResponseMessage Response = await HttpClient.SendAsync(Request, CancellationToken);
+			Response.EnsureSuccessStatusCode();
+		}
+
+		/// <inheritdoc/>
+		public async IAsyncEnumerable<IComputeTaskInfo> GetTaskUpdatesAsync(ClusterId ClusterId, ChannelId ChannelId, [EnumeratorCancellation] CancellationToken CancellationToken)
+		{
+			for (; ; )
+			{
+				HttpRequestMessage Request = new HttpRequestMessage(HttpMethod.Post, $"api/v1/compute/{ClusterId}/updates/{ChannelId}?wait=10");
+				Request.Headers.Add("Accept", "application/x-ue-cb");
+
+				using HttpResponseMessage Response = await HttpClient.SendAsync(Request, HttpCompletionOption.ResponseContentRead, CancellationToken);
+				Response.EnsureSuccessStatusCode();
+
+				byte[] Data = await Response.Content.ReadAsByteArrayAsync();
+				GetTaskUpdatesResponse ParsedResponse = CbSerializer.Deserialize<GetTaskUpdatesResponse>(new CbField(Data));
+
+				foreach (GetTaskUpdateResponse Update in ParsedResponse.Updates)
+				{
+					CancellationToken.ThrowIfCancellationRequested();
+					yield return Update;
+				}
+			}
+		}
+	}
+}
diff --git a/Engine/Source/Programs/Shared/EpicGames.Horde/EpicGames.Horde.csproj b/Engine/Source/Programs/Shared/EpicGames.Horde/EpicGames.Horde.csproj
--- a/Engine/Source/Programs/Shared/EpicGames.Horde/EpicGames.Horde.csproj
+++ b/Engine/Source/Programs/Shared/EpicGames.Horde/EpicGames.Horde.csproj
@@ -28,13 +28,9 @@
     </Compile>
   </ItemGroup>
   <ItemGroup>
-    <PackageReference Include="Google.Protobuf" Version="3.12.3" />
-    <PackageReference Include="Grpc.Core.Api" Version="2.35.0" />
-    <PackageReference Include="Grpc.Tools" Version="2.30.0">
-      <PrivateAssets>all</PrivateAssets>
-      <IncludeAssets>runtime; build; native; contentfiles; analyzers; buildtransitive</IncludeAssets>
-    </PackageReference>
     <PackageReference Include="Microsoft.CodeAnalysis.CSharp.Workspaces" Version="3.10.0" />
+    <PackageReference Include="Microsoft.Extensions.DependencyInjection" Version="6.0.0" />
+    <PackageReference Include="Microsoft.Extensions.Http" Version="3.1.22" />
     <PackageReference Include="Microsoft.Extensions.Logging" Version="5.0.0" />
   </ItemGroup>
   <ItemGroup>
@@ -42,13 +38,4 @@
     <ProjectReference Include="..\EpicGames.IoHash\EpicGames.IoHash.csproj" />
     <ProjectReference Include="..\EpicGames.Serialization\EpicGames.Serialization.csproj" />
   </ItemGroup>
-  <ItemGroup>
-    <Protobuf Include="Protos\horde\common\cbbinaryattachment_wrapper.proto" ProtoRoot="Protos" />
-    <Protobuf Include="Protos\horde\common\cbobject_wrapper.proto" ProtoRoot="Protos" />
-    <Protobuf Include="Protos\horde\common\cbobjectattachment_wrapper.proto" ProtoRoot="Protos" />
-    <Protobuf Include="Protos\horde\common\iohash_wrapper.proto" ProtoRoot="Protos" />
-    <Protobuf Include="Protos\horde\compute\compute_rpc_messages.proto" ProtoRoot="Protos" />
-    <Protobuf Include="Protos\horde\storage\blob_store_messages.proto" ProtoRoot="Protos" />
-    <Protobuf Include="Protos\horde\storage\ref_table_messages.proto" ProtoRoot="Protos" />
-  </ItemGroup>
 </Project>
\ No newline at end of file
diff --git a/Engine/Source/Programs/Shared/EpicGames.Horde/Storage/BlobStoreMessages.cs b/Engine/Source/Programs/Shared/EpicGames.Horde/Storage/BlobStoreMessages.cs
deleted file mode 100644
diff --git a/Engine/Source/Programs/Shared/EpicGames.Horde/Storage/IStorageClient.cs b/Engine/Source/Programs/Shared/EpicGames.Horde/Storage/IStorageClient.cs
--- a/Engine/Source/Programs/Shared/EpicGames.Horde/Storage/IStorageClient.cs
+++ b/Engine/Source/Programs/Shared/EpicGames.Horde/Storage/IStorageClient.cs
@@ -6,6 +6,7 @@ using System;
 using System.Collections.Generic;
 using System.IO;
 using System.Text;
+using System.Threading;
 using System.Threading.Tasks;
 
 namespace EpicGames.Horde.Storage
@@ -82,12 +83,12 @@ namespace EpicGames.Horde.Storage
 		/// <summary>
 		/// Identifier for the ref
 		/// </summary>
-		public IoHash RefId { get; }
+		public RefId RefId { get; }
 
 		/// <summary>
 		/// Constructor
 		/// </summary>
-		public RefException(NamespaceId NamespaceId, BucketId BucketId, IoHash RefId, string Message, Exception? InnerException = null)
+		public RefException(NamespaceId NamespaceId, BucketId BucketId, RefId RefId, string Message, Exception? InnerException = null)
 			: base(Message, InnerException)
 		{
 			this.NamespaceId = NamespaceId;
@@ -104,12 +105,26 @@ namespace EpicGames.Horde.Storage
 		/// <summary>
 		/// Constructor
 		/// </summary>
-		public RefNotFoundException(NamespaceId NamespaceId, BucketId BucketId, IoHash RefId, Exception? InnerException = null)
+		public RefNotFoundException(NamespaceId NamespaceId, BucketId BucketId, RefId RefId, Exception? InnerException = null)
 			: base(NamespaceId, BucketId, RefId, $"Ref {NamespaceId}/{BucketId}/{RefId} not found", InnerException)
 		{
 		}
 	}
 
+	/// <summary>
+	/// Indicates that a ref cannot be finalized due to a missing blob
+	/// </summary>
+	public sealed class RefMissingBlobException : RefException
+	{
+		/// <summary>
+		/// Constructor
+		/// </summary>
+		public RefMissingBlobException(NamespaceId NamespaceId, BucketId BucketId, RefId RefId, List<IoHash> MissingBlobs, Exception? InnerException = null)
+			: base(NamespaceId, BucketId, RefId, $"Ref {NamespaceId}/{BucketId}/{RefId} cannot be finalized; missing {MissingBlobs.Count} blobs ({MissingBlobs[0]}...)", InnerException)
+		{
+		}
+	}
+
 	/// <summary>
 	/// Interface for an object reference
 	/// </summary>
@@ -128,7 +143,7 @@ namespace EpicGames.Horde.Storage
 		/// <summary>
 		/// Ref identifier
 		/// </summary>
-		IoHash RefId { get; }
+		RefId RefId { get; }
 
 		/// <summary>
 		/// The value stored for this ref
@@ -148,8 +163,9 @@ namespace EpicGames.Horde.Storage
 		/// </summary>
 		/// <param name="NamespaceId">Namespace to operate on</param>
 		/// <param name="Hash">Hash of the blob</param>
+		/// <param name="CancellationToken">Cancellation token for the operation</param>
 		/// <returns>Stream for the blob, or null if it does not exist</returns>
-		Task<Stream> ReadBlobAsync(NamespaceId NamespaceId, IoHash Hash);
+		Task<Stream> ReadBlobAsync(NamespaceId NamespaceId, IoHash Hash, CancellationToken CancellationToken = default);
 
 		/// <summary>
 		/// Writes a blob to storage
@@ -157,7 +173,8 @@ namespace EpicGames.Horde.Storage
 		/// <param name="NamespaceId">Namespace to operate on</param>
 		/// <param name="Hash">Hash of the blob</param>
 		/// <param name="Stream">The stream to write</param>
-		Task WriteBlobAsync(NamespaceId NamespaceId, IoHash Hash, Stream Stream);
+		/// <param name="CancellationToken">Cancellation token for the operation</param>
+		Task WriteBlobAsync(NamespaceId NamespaceId, IoHash Hash, Stream Stream, CancellationToken CancellationToken = default);
 
 		#endregion
 
@@ -169,8 +186,9 @@ namespace EpicGames.Horde.Storage
 		/// <param name="NamespaceId">Namespace identifier</param>
 		/// <param name="BucketId">Bucket identifier</param>
 		/// <param name="RefId">Name of the reference</param>
+		/// <param name="CancellationToken">Cancellation token for the operation</param>
 		/// <returns>The reference data if the ref exists</returns>
-		Task<IRef> GetRefAsync(NamespaceId NamespaceId, BucketId BucketId, IoHash RefId);
+		Task<IRef> GetRefAsync(NamespaceId NamespaceId, BucketId BucketId, RefId RefId, CancellationToken CancellationToken = default);
 
 		/// <summary>
 		/// Gets the given reference
@@ -178,8 +196,9 @@ namespace EpicGames.Horde.Storage
 		/// <param name="NamespaceId">Namespace identifier</param>
 		/// <param name="BucketId">Bucket identifier</param>
 		/// <param name="RefId">Ref identifier</param>
+		/// <param name="CancellationToken">Cancellation token for the operation</param>
 		/// <returns>The reference data if the ref exists</returns>
-		Task<bool> HasRefAsync(NamespaceId NamespaceId, BucketId BucketId, IoHash RefId);
+		Task<bool> HasRefAsync(NamespaceId NamespaceId, BucketId BucketId, RefId RefId, CancellationToken CancellationToken = default);
 
 		/// <summary>
 		/// Determines which refs are missing
@@ -187,18 +206,20 @@ namespace EpicGames.Horde.Storage
 		/// <param name="NamespaceId">Namespace identifier</param>
 		/// <param name="BucketId">Bucket identifier</param>
 		/// <param name="RefIds">Names of the references</param>
+		/// <param name="CancellationToken">Cancellation token for the operation</param>
 		/// <returns>List of missing references</returns>
-		Task<List<IoHash>> FindMissingRefsAsync(NamespaceId NamespaceId, BucketId BucketId, List<IoHash> RefIds);
+		Task<List<RefId>> FindMissingRefsAsync(NamespaceId NamespaceId, BucketId BucketId, List<RefId> RefIds, CancellationToken CancellationToken = default);
 
 		/// <summary>
-		/// Sets the given reference
+		/// Attempts to sets the given reference, returning a list of missing objects on failure.
 		/// </summary>
 		/// <param name="NamespaceId">Namespace identifier</param>
 		/// <param name="BucketId">Bucket identifier</param>
 		/// <param name="RefId">Ref identifier</param>
 		/// <param name="Value">New value for the reference</param>
+		/// <param name="CancellationToken">Cancellation token for the operation</param>
 		/// <returns>List of missing references</returns>
-		Task<List<IoHash>> SetRefAsync(NamespaceId NamespaceId, BucketId BucketId, IoHash RefId, CbObject Value);
+		Task<List<IoHash>> TrySetRefAsync(NamespaceId NamespaceId, BucketId BucketId, RefId RefId, CbObject Value, CancellationToken CancellationToken = default);
 
 		/// <summary>
 		/// Attempts to finalize a reference, turning its references into hard references
@@ -207,8 +228,9 @@ namespace EpicGames.Horde.Storage
 		/// <param name="BucketId">Bucket identifier</param>
 		/// <param name="RefId">Ref identifier</param>
 		/// <param name="Hash">Hash of the referenced object</param>
+		/// <param name="CancellationToken">Cancellation token for the operation</param>
 		/// <returns></returns>
-		Task<List<IoHash>> FinalizeRefAsync(NamespaceId NamespaceId, BucketId BucketId, IoHash RefId, IoHash Hash);
+		Task<List<IoHash>> TryFinalizeRefAsync(NamespaceId NamespaceId, BucketId BucketId, RefId RefId, IoHash Hash, CancellationToken CancellationToken = default);
 
 		/// <summary>
 		/// Removes the given reference
@@ -216,8 +238,9 @@ namespace EpicGames.Horde.Storage
 		/// <param name="NamespaceId">Namespace identifier</param>
 		/// <param name="BucketId">Bucket identifier</param>
 		/// <param name="RefId">Ref identifier</param>
+		/// <param name="CancellationToken">Cancellation token for the operation</param>
 		/// <returns>True if the ref was deleted, false if it did not exist</returns>
-		Task<bool> DeleteRefAsync(NamespaceId NamespaceId, BucketId BucketId, IoHash RefId);
+		Task<bool> DeleteRefAsync(NamespaceId NamespaceId, BucketId BucketId, RefId RefId, CancellationToken CancellationToken = default);
 
 		#endregion
 	}
@@ -232,10 +255,11 @@ namespace EpicGames.Horde.Storage
 		/// <summary>
 		/// Gets a blob as a byte array
 		/// </summary>
-		/// <param name="StorageClient"></param>
-		/// <param name="NamespaceId"></param>
-		/// <param name="Hash"></param>
-		/// <returns></returns>
+		/// <param name="StorageClient">The storage interface</param>
+		/// <param name="NamespaceId">Namespace containing the blob</param>
+		/// <param name="Hash">Hash of the blob to read</param>
+		/// <param name="MaxInMemoryBlobLength">Maximum allowed memory allocation to store the blob</param>
+		/// <returns>Data for the blob that was read. Throws an exception if the blob was not found.</returns>
 		public static async Task<byte[]> ReadBlobToMemoryAsync(this IStorageClient StorageClient, NamespaceId NamespaceId, IoHash Hash, int MaxInMemoryBlobLength = DefaultMaxInMemoryBlobLength)
 		{
 			using Stream Stream = await StorageClient.ReadBlobAsync(NamespaceId, Hash);
@@ -260,12 +284,39 @@ namespace EpicGames.Horde.Storage
 			return Buffer;
 		}
 
+		/// <summary>
+		/// Reads a blob and decodes it as a compact binary object
+		/// </summary>
+		/// <param name="StorageClient">The storage interface</param>
+		/// <param name="NamespaceId">Namespace containing the blob</param>
+		/// <param name="Hash">Hash of the blob to read</param>
+		/// <param name="MaxInMemoryBlobLength">Maximum allowed memory allocation to store the blob</param>
+		/// <returns>The decoded object</returns>
+		public static async Task<CbObject> ReadObjectAsync(this IStorageClient StorageClient, NamespaceId NamespaceId, IoHash Hash, int MaxInMemoryBlobLength = DefaultMaxInMemoryBlobLength)
+		{
+			return new CbObject(await ReadBlobToMemoryAsync(StorageClient, NamespaceId, Hash, MaxInMemoryBlobLength));
+		}
+
+		/// <summary>
+		/// Reads a blob and deserializes it as the given compact-binary encoded type
+		/// </summary>
+		/// <typeparam name="T">Type of object to deserialize</typeparam>
+		/// <param name="StorageClient">The storage interface</param>
+		/// <param name="NamespaceId">Namespace containing the blob</param>
+		/// <param name="Hash">Hash of the blob to read</param>
+		/// <param name="MaxInMemoryBlobLength">Maximum allowed memory allocation to store the blob</param>
+		/// <returns>The decoded object</returns>
+		public static async Task<T> ReadObjectAsync<T>(this IStorageClient StorageClient, NamespaceId NamespaceId, IoHash Hash, int MaxInMemoryBlobLength = DefaultMaxInMemoryBlobLength)
+		{
+			return CbSerializer.Deserialize<T>(await ReadBlobToMemoryAsync(StorageClient, NamespaceId, Hash, MaxInMemoryBlobLength));
+		}
+
 		/// <summary>
 		/// Gets a blob as a byte array
 		/// </summary>
-		/// <param name="StorageClient"></param>
-		/// <param name="NamespaceId"></param>
-		/// <param name="Data">The data to be written</param>
+		/// <param name="StorageClient">The storage interface</param>
+		/// <param name="NamespaceId">Namespace containing the blob</param>
+		/// <param name="Data">Data to write</param>
 		/// <returns></returns>
 		public static async Task<IoHash> WriteBlobFromMemoryAsync(this IStorageClient StorageClient, NamespaceId NamespaceId, ReadOnlyMemory<byte> Data)
 		{
@@ -277,8 +328,32 @@ namespace EpicGames.Horde.Storage
 		/// <summary>
 		/// Gets a blob as a byte array
 		/// </summary>
-		/// <param name="StorageClient"></param>
-		/// <param name="NamespaceId"></param>
+		/// <param name="StorageClient">The storage interface</param>
+		/// <param name="NamespaceId">Namespace containing the blob</param>
+		/// <param name="Object">The object to be written</param>
+		/// <returns></returns>
+		public static Task<IoHash> WriteObjectAsync(this IStorageClient StorageClient, NamespaceId NamespaceId, CbObject Object)
+		{
+			return WriteBlobFromMemoryAsync(StorageClient, NamespaceId, Object.GetView());
+		}
+
+		/// <summary>
+		/// Gets a blob as a byte array
+		/// </summary>
+		/// <param name="StorageClient">The storage interface</param>
+		/// <param name="NamespaceId">Namespace containing the blob</param>
+		/// <param name="Object">The object to be written</param>
+		/// <returns></returns>
+		public static Task<IoHash> WriteObjectAsync<T>(this IStorageClient StorageClient, NamespaceId NamespaceId, T Object)
+		{
+			return WriteObjectAsync(StorageClient, NamespaceId, CbSerializer.Serialize<T>(Object));
+		}
+
+		/// <summary>
+		/// Gets a blob as a byte array
+		/// </summary>
+		/// <param name="StorageClient">The storage interface</param>
+		/// <param name="NamespaceId">Namespace containing the blob</param>
 		/// <param name="Hash">Hash of the data</param>
 		/// <param name="Data">The data to be written</param>
 		/// <returns></returns>
@@ -287,5 +362,84 @@ namespace EpicGames.Horde.Storage
 			using ReadOnlyMemoryStream Stream = new ReadOnlyMemoryStream(Data);
 			await StorageClient.WriteBlobAsync(NamespaceId, Hash, Stream);
 		}
+
+		/// <summary>
+		/// Reads a reference as a specific type
+		/// </summary>
+		/// <typeparam name="T"></typeparam>
+		/// <param name="StorageClient">The storage interface</param>
+		/// <param name="NamespaceId">Namespace containing the ref</param>
+		/// <param name="BucketId">Bucket containing the ref</param>
+		/// <param name="RefId">The ref id</param>
+		/// <returns>Deserialized object for the given ref</returns>
+		public static async Task<T> GetRefAsync<T>(this IStorageClient StorageClient, NamespaceId NamespaceId, BucketId BucketId, RefId RefId)
+		{
+			IRef Ref = await StorageClient.GetRefAsync(NamespaceId, BucketId, RefId);
+			return CbSerializer.Deserialize<T>(Ref.Value);
+		}
+
+		/// <summary>
+		/// Attempts to set a ref to a particular value
+		/// </summary>
+		/// <typeparam name="T"></typeparam>
+		/// <param name="StorageClient">The storage interface</param>
+		/// <param name="NamespaceId">Namespace containing the ref</param>
+		/// <param name="BucketId">Bucket containing the ref</param>
+		/// <param name="RefId">The ref id</param>
+		/// <param name="Value">The new object for the ref</param>
+		/// <returns>List of missing blob hashes</returns>
+		public static Task<List<IoHash>> TrySetRefAsync<T>(this IStorageClient StorageClient, NamespaceId NamespaceId, BucketId BucketId, RefId RefId, T Value)
+		{
+			CbObject Object = CbSerializer.Serialize<T>(Value);
+			return StorageClient.TrySetRefAsync(NamespaceId, BucketId, RefId, Object);
+		}
+
+		/// <summary>
+		/// Sets a ref to a particular value
+		/// </summary>
+		/// <param name="StorageClient">The storage interface</param>
+		/// <param name="NamespaceId">Namespace containing the ref</param>
+		/// <param name="BucketId">Bucket containing the ref</param>
+		/// <param name="RefId">The ref id</param>
+		/// <param name="Value">The new object for the ref</param>
+		public static async Task SetRefAsync(this IStorageClient StorageClient, NamespaceId NamespaceId, BucketId BucketId, RefId RefId, CbObject Value)
+		{
+			List<IoHash> MissingHashes = await StorageClient.TrySetRefAsync(NamespaceId, BucketId, RefId, Value);
+			if (MissingHashes.Count > 0)
+			{
+				throw new RefMissingBlobException(NamespaceId, BucketId, RefId, MissingHashes);
+			}
+		}
+
+		/// <summary>
+		/// Sets a ref to a particular value
+		/// </summary>
+		/// <param name="StorageClient">The storage interface</param>
+		/// <param name="NamespaceId">Namespace containing the ref</param>
+		/// <param name="BucketId">Bucket containing the ref</param>
+		/// <param name="RefId">The ref id</param>
+		/// <param name="Value">The new object for the ref</param>
+		public static Task SetRefAsync<T>(this IStorageClient StorageClient, NamespaceId NamespaceId, BucketId BucketId, RefId RefId, T Value)
+		{
+			CbObject Object = CbSerializer.Serialize<T>(Value);
+			return SetRefAsync(StorageClient, NamespaceId, BucketId, RefId, Object);
+		}
+
+		/// <summary>
+		/// Finalize a ref, throwing an exception if finalization fails
+		/// </summary>
+		/// <param name="StorageClient">The storage interface</param>
+		/// <param name="NamespaceId">Namespace containing the ref</param>
+		/// <param name="BucketId">Bucket containing the ref</param>
+		/// <param name="RefId">The ref id</param>
+		/// <param name="ValueHash">Hash of the ref value</param>
+		public static async Task FinalizeRefAsync(this IStorageClient StorageClient, NamespaceId NamespaceId, BucketId BucketId, RefId RefId, IoHash ValueHash)
+		{
+			List<IoHash> MissingHashes = await StorageClient.TryFinalizeRefAsync(NamespaceId, BucketId, RefId, ValueHash);
+			if (MissingHashes.Count > 0)
+			{
+				throw new RefMissingBlobException(NamespaceId, BucketId, RefId, MissingHashes);
+			}
+		}
 	}
 }
diff --git a/Engine/Source/Programs/Shared/EpicGames.Horde/Storage/Impl/HttpStorageClient.cs b/Engine/Source/Programs/Shared/EpicGames.Horde/Storage/Impl/HttpStorageClient.cs
--- a/Engine/Source/Programs/Shared/EpicGames.Horde/Storage/Impl/HttpStorageClient.cs
+++ b/Engine/Source/Programs/Shared/EpicGames.Horde/Storage/Impl/HttpStorageClient.cs
@@ -1,8 +1,9 @@
 // Copyright Epic Games, Inc. All Rights Reserved.
 
 using EpicGames.Core;
+using EpicGames.Horde.Auth;
 using EpicGames.Serialization;
-using Google.Protobuf;
+using Microsoft.Extensions.Options;
 using System;
 using System.Collections.Generic;
 using System.IO;
@@ -11,6 +12,7 @@ using System.Net.Http;
 using System.Net.Http.Headers;
 using System.Text;
 using System.Text.Json;
+using System.Threading;
 using System.Threading.Tasks;
 
 namespace EpicGames.Horde.Storage.Impl
@@ -34,7 +36,7 @@ namespace EpicGames.Horde.Storage.Impl
 		/// <summary>
 		/// List of hashes that the blob store needs.
 		/// </summary>
-		public List<IoHash> Needs { get; set; } = new List<IoHash>();
+		public List<RefId> Needs { get; set; } = new List<RefId>();
 	}
 
 	/// <summary>
@@ -48,11 +50,11 @@ namespace EpicGames.Horde.Storage.Impl
 
 			public BucketId BucketId { get; set; }
 
-			public IoHash RefId { get; set; }
+			public RefId RefId { get; set; }
 
 			public CbObject Value { get; }
 
-			public RefImpl(NamespaceId NamespaceId, BucketId BucketId, IoHash RefId, CbObject Value)
+			public RefImpl(NamespaceId NamespaceId, BucketId BucketId, RefId RefId, CbObject Value)
 			{
 				this.NamespaceId = NamespaceId;
 				this.BucketId = BucketId;
@@ -67,32 +69,59 @@ namespace EpicGames.Horde.Storage.Impl
 		const string CompactBinaryMimeType = "application/x-ue-cb";
 
 		HttpClient HttpClient;
+		StorageSettings Settings;
+		Uri BaseUrl;
+		IAuthProvider AuthProvider;
 
 		/// <summary>
 		/// Constructor
 		/// </summary>
 		/// <param name="HttpClient">Http client for the blob store. Should be pre-configured with a base address and appropriate authentication headers.</param>
-		public HttpStorageClient(HttpClient HttpClient)
+		public HttpStorageClient(HttpClient HttpClient, IOptions<StorageSettings> Settings, IAuthProvider<HttpStorageClient> AuthProvider)
 		{
 			this.HttpClient = HttpClient;
+			this.Settings = Settings.Value;
+			this.BaseUrl = new Uri(Settings.Value.Url);
+			this.AuthProvider = AuthProvider;
 		}
 
 		#region Blobs
 
 		/// <inheritdoc/>
-		public async Task<Stream> ReadBlobAsync(NamespaceId NamespaceId, IoHash Hash)
+		public async Task<Stream> ReadBlobAsync(NamespaceId NamespaceId, IoHash Hash, CancellationToken CancellationToken)
 		{
-			using HttpResponseMessage Response = await HttpClient.GetAsync($"/api/v1/blobs/{NamespaceId}/{Hash}");
-			return await Response.Content.ReadAsStreamAsync();
+			using HttpRequestMessage Request = new HttpRequestMessage(HttpMethod.Get, new Uri(BaseUrl, $"api/v1/blobs/{NamespaceId}/{Hash}"));
+			await AuthProvider.AddAuthorizationAsync(Request, CancellationToken);
+
+			HttpResponseMessage Response = await HttpClient.SendAsync(Request, CancellationToken);
+			try
+			{
+				if (Response.StatusCode == System.Net.HttpStatusCode.NotFound)
+				{
+					throw new BlobNotFoundException(NamespaceId, Hash);
+				}
+
+				Response.EnsureSuccessStatusCode();
+				return await Response.Content.ReadAsStreamAsync();
+			}
+			catch
+			{
+				Response.Dispose();
+				throw;
+			}
 		}
 
 		/// <inheritdoc/>
-		public async Task WriteBlobAsync(NamespaceId NamespaceId, IoHash Hash, Stream Stream)
+		public async Task WriteBlobAsync(NamespaceId NamespaceId, IoHash Hash, Stream Stream, CancellationToken CancellationToken)
 		{
+			using HttpRequestMessage Request = new HttpRequestMessage(HttpMethod.Put, new Uri(BaseUrl, $"api/v1/blobs/{NamespaceId}/{Hash}"));
+			await AuthProvider.AddAuthorizationAsync(Request, CancellationToken);
+
 			StreamContent Content = new StreamContent(Stream);
 			Content.Headers.ContentType = new MediaTypeHeaderValue("application/octet-stream");
+			Request.Content = Content;
 
-			HttpResponseMessage Response = await HttpClient.PutAsync($"api/v1/blobs/{NamespaceId}/{Hash}", Content);
+			HttpResponseMessage Response = await HttpClient.SendAsync(Request, CancellationToken);
 			Response.EnsureSuccessStatusCode();
 		}
 
@@ -101,9 +130,11 @@ namespace EpicGames.Horde.Storage.Impl
 		#region Refs
 
 		/// <inheritdoc/>
-		public async Task<IRef> GetRefAsync(NamespaceId NamespaceId, BucketId BucketId, IoHash RefId)
+		public async Task<IRef> GetRefAsync(NamespaceId NamespaceId, BucketId BucketId, RefId RefId, CancellationToken CancellationToken)
 		{
-			using HttpRequestMessage Request = new HttpRequestMessage(HttpMethod.Get, $"api/v1/refs/{NamespaceId}/{BucketId}/{RefId}");
+			using HttpRequestMessage Request = new HttpRequestMessage(HttpMethod.Get, new Uri(BaseUrl, $"api/v1/refs/{NamespaceId}/{BucketId}/{RefId}"));
+			await AuthProvider.AddAuthorizationAsync(Request, CancellationToken);
+
 			Request.Headers.Accept.Clear();
 			Request.Headers.Accept.Add(new MediaTypeWithQualityHeaderValue(CompactBinaryMimeType));
 
@@ -125,7 +156,7 @@ namespace EpicGames.Horde.Storage.Impl
 		}
 
 		/// <inheritdoc/>
-		public async Task<List<IoHash>> SetRefAsync(NamespaceId NamespaceId, BucketId BucketId, IoHash RefId, CbObject Value)
+		public async Task<List<IoHash>> TrySetRefAsync(NamespaceId NamespaceId, BucketId BucketId, RefId RefId, CbObject Value, CancellationToken CancellationToken)
 		{
 			using ReadOnlyMemoryStream Stream = new ReadOnlyMemoryStream(Value.GetView());
 
@@ -133,7 +164,9 @@ namespace EpicGames.Horde.Storage.Impl
 			Content.Headers.ContentType = new MediaTypeHeaderValue(CompactBinaryMimeType);
 			Content.Headers.Add(HashHeaderName, IoHash.Compute(Value.GetView().Span).ToString());
 
-			using HttpRequestMessage Request = new HttpRequestMessage(HttpMethod.Put, $"api/v1/refs/{NamespaceId}/{BucketId}/{RefId}");
+			using HttpRequestMessage Request = new HttpRequestMessage(HttpMethod.Put, new Uri(BaseUrl, $"api/v1/refs/{NamespaceId}/{BucketId}/{RefId}"));
+			await AuthProvider.AddAuthorizationAsync(Request, CancellationToken);
+
 			Request.Headers.Accept.Clear();
 			Request.Headers.Accept.Add(new MediaTypeWithQualityHeaderValue("application/json"));
 			Request.Content = Content;
@@ -154,9 +187,11 @@ namespace EpicGames.Horde.Storage.Impl
 		}
 
 		/// <inheritdoc/>
-		public async Task<List<IoHash>> FinalizeRefAsync(NamespaceId NamespaceId, BucketId BucketId, IoHash RefId, IoHash Hash)
+		public async Task<List<IoHash>> TryFinalizeRefAsync(NamespaceId NamespaceId, BucketId BucketId, RefId RefId, IoHash Hash, CancellationToken CancellationToken)
 		{
-			using HttpRequestMessage Request = new HttpRequestMessage(HttpMethod.Post, $"api/v1/refs/{NamespaceId}/{BucketId}/{RefId}/finalize/{Hash}");
+			using HttpRequestMessage Request = new HttpRequestMessage(HttpMethod.Post, new Uri(BaseUrl, $"api/v1/refs/{NamespaceId}/{BucketId}/{RefId}/finalize/{Hash}"));
+			await AuthProvider.AddAuthorizationAsync(Request, CancellationToken);
+
 			Request.Headers.Accept.Clear();
 			Request.Headers.Accept.Add(new MediaTypeWithQualityHeaderValue("application/json"));
 
@@ -176,48 +211,54 @@ namespace EpicGames.Horde.Storage.Impl
 		}
 
 		/// <inheritdoc/>
-		public async Task<bool> DeleteRefAsync(NamespaceId NamespaceId, BucketId BucketId, IoHash RefId)
+		public async Task<bool> DeleteRefAsync(NamespaceId NamespaceId, BucketId BucketId, RefId RefId, CancellationToken CancellationToken)
 		{
-			using HttpResponseMessage Response = await HttpClient.DeleteAsync($"api/v1/refs/{NamespaceId}/{BucketId}/{RefId}");
+			using HttpRequestMessage Request = new HttpRequestMessage(HttpMethod.Delete, new Uri(BaseUrl, $"api/v1/refs/{NamespaceId}/{BucketId}/{RefId}"));
+			await AuthProvider.AddAuthorizationAsync(Request, CancellationToken);
+
+			using HttpResponseMessage Response = await HttpClient.SendAsync(Request, CancellationToken);
 			return Response.IsSuccessStatusCode;
 		}
 
 		/// <inheritdoc/>
-		public async Task<bool> HasRefAsync(NamespaceId NamespaceId, BucketId BucketId, IoHash RefId)
+		public async Task<bool> HasRefAsync(NamespaceId NamespaceId, BucketId BucketId, RefId RefId, CancellationToken CancellationToken)
 		{
-			using (HttpRequestMessage Request = new HttpRequestMessage(HttpMethod.Head, $"api/v1/refs/{NamespaceId}/{BucketId}/{RefId}"))
+			using HttpRequestMessage Request = new HttpRequestMessage(HttpMethod.Head, new Uri(BaseUrl, $"api/v1/refs/{NamespaceId}/{BucketId}/{RefId}"));
+			await AuthProvider.AddAuthorizationAsync(Request, CancellationToken);
+
+			using HttpResponseMessage Response = await HttpClient.SendAsync(Request, CancellationToken);
+			if (Response.IsSuccessStatusCode)
 			{
-				using HttpResponseMessage Response = await HttpClient.SendAsync(Request);
-				if (Response.IsSuccessStatusCode)
-				{
-					return true;
-				}
-				else if (Response.StatusCode == System.Net.HttpStatusCode.NotFound)
-				{
-					return false;
-				}
-				else
-				{
-					throw new RefException(NamespaceId, BucketId, RefId, $"Unexpected response for {nameof(HasRefAsync)} call: {Response.StatusCode}");
-				}
+				return true;
+			}
+			else if (Response.StatusCode == System.Net.HttpStatusCode.NotFound)
+			{
+				return false;
+			}
+			else
+			{
+				throw new RefException(NamespaceId, BucketId, RefId, $"Unexpected response for {nameof(HasRefAsync)} call: {Response.StatusCode}");
 			}
 		}
 
 		/// <inheritdoc/>
-		public async Task<List<IoHash>> FindMissingRefsAsync(NamespaceId NamespaceId, BucketId BucketId, List<IoHash> RefIds)
+		public async Task<List<RefId>> FindMissingRefsAsync(NamespaceId NamespaceId, BucketId BucketId, List<RefId> RefIds, CancellationToken CancellationToken)
 		{
 			string RefList = String.Join("&", RefIds.Select(x => $"id={x}"));
 
-			using HttpResponseMessage Response = await HttpClient.PostAsync($"/api/v1/refs/{NamespaceId}/{BucketId}/exists?{RefList}", null);
+			using HttpRequestMessage Request = new HttpRequestMessage(HttpMethod.Post, new Uri(BaseUrl, $"/api/v1/refs/{NamespaceId}/{BucketId}/exists?{RefList}"));
+			await AuthProvider.AddAuthorizationAsync(Request, CancellationToken);
+
+			using HttpResponseMessage Response = await HttpClient.PostAsync(new Uri(BaseUrl, $"/api/v1/refs/{NamespaceId}/{BucketId}/exists?{RefList}"), null);
 			if (!Response.IsSuccessStatusCode)
 			{
-				throw new RefException(NamespaceId, BucketId, IoHash.Zero, await GetMessageFromResponse(Response));
+				throw new RefException(NamespaceId, BucketId, new RefId(IoHash.Zero), await GetMessageFromResponse(Response));
 			}
 
 			RefExistsResponse? ResponseContent = await ReadJsonResponse<RefExistsResponse>(Response.Content);
 			if (ResponseContent == null)
 			{
-				throw new RefException(NamespaceId, BucketId, IoHash.Zero, "Unable to parse response body");
+				throw new RefException(NamespaceId, BucketId, new RefId(IoHash.Zero), "Unable to parse response body");
 			}
 
 			return ResponseContent.Needs;
diff --git a/Engine/Source/Programs/Shared/EpicGames.Horde/Storage/NamespaceId.cs b/Engine/Source/Programs/Shared/EpicGames.Horde/Storage/NamespaceId.cs
--- a/Engine/Source/Programs/Shared/EpicGames.Horde/Storage/NamespaceId.cs
+++ b/Engine/Source/Programs/Shared/EpicGames.Horde/Storage/NamespaceId.cs
@@ -2,14 +2,20 @@
 
 using System;
 using System.Collections.Generic;
+using System.ComponentModel;
 using System.Diagnostics.CodeAnalysis;
+using System.Globalization;
 using System.Text;
+using System.Text.Json;
+using System.Text.Json.Serialization;
 
 namespace EpicGames.Horde.Storage
 {
 	/// <summary>
 	/// Identifier for a storage namespace
 	/// </summary>
+	[JsonConverter(typeof(NamespaceIdJsonConverter))]
+	[TypeConverter(typeof(NamespaceIdTypeConverter))]
 	public struct NamespaceId : IEquatable<NamespaceId>
 	{
 		/// <summary>
@@ -44,4 +50,28 @@ namespace EpicGames.Horde.Storage
 		/// <inheritdoc cref="StringId.op_Inequality"/>
 		public static bool operator !=(NamespaceId Left, NamespaceId Right) => Left.Inner != Right.Inner;
 	}
+
+	/// <summary>
+	/// Type converter for NamespaceId to and from JSON
+	/// </summary>
+	sealed class NamespaceIdJsonConverter : JsonConverter<NamespaceId>
+	{
+		/// <inheritdoc/>
+		public override NamespaceId Read(ref Utf8JsonReader Reader, Type TypeToConvert, JsonSerializerOptions Options) => new NamespaceId(Reader.GetString() ?? String.Empty);
+
+		/// <inheritdoc/>
+		public override void Write(Utf8JsonWriter Writer, NamespaceId Value, JsonSerializerOptions Options) => Writer.WriteStringValue(Value.ToString());
+	}
+
+	/// <summary>
+	/// Type converter from strings to NamespaceId objects
+	/// </summary>
+	sealed class NamespaceIdTypeConverter : TypeConverter
+	{
+		/// <inheritdoc/>
+		public override bool CanConvertFrom(ITypeDescriptorContext Context, Type SourceType) => SourceType == typeof(string);
+
+		/// <inheritdoc/>
+		public override object ConvertFrom(ITypeDescriptorContext Context, CultureInfo Culture, object Value) => new NamespaceId((string)Value);
+	}
 }
diff --git a/Engine/Source/Programs/Shared/EpicGames.Horde/Storage/RefId.cs b/Engine/Source/Programs/Shared/EpicGames.Horde/Storage/RefId.cs
new file mode 100644
--- /dev/null
+++ b/Engine/Source/Programs/Shared/EpicGames.Horde/Storage/RefId.cs
@@ -0,0 +1,109 @@
+// Copyright Epic Games, Inc. All Rights Reserved.
+
+using EpicGames.Core;
+using EpicGames.Serialization;
+using System;
+using System.Collections.Generic;
+using System.ComponentModel;
+using System.Globalization;
+using System.Text;
+using System.Text.Json;
+using System.Text.Json.Serialization;
+
+namespace EpicGames.Horde.Storage
+{
+	/// <summary>
+	/// Identifier for a storage namespace
+	/// </summary>
+	[JsonConverter(typeof(RefIdJsonConverter))]
+	[TypeConverter(typeof(RefIdTypeConverter))]
+	[CbConverter(typeof(RefIdCbConverter))]
+	public struct RefId : IEquatable<RefId>
+	{
+		/// <summary>
+		/// Hash identifier for the ref
+		/// </summary>
+		public IoHash Hash { get; }
+
+		/// <summary>
+		/// Constructor
+		/// </summary>
+		/// <param name="Hash"></param>
+		public RefId(IoHash Hash)
+		{
+			this.Hash = Hash;
+		}
+
+		/// <summary>
+		/// Constructor
+		/// </summary>
+		/// <param name="Name">Name to identify this ref</param>
+		public RefId(string Name)
+		{
+			this.Hash = IoHash.Compute(Encoding.UTF8.GetBytes(Name));
+		}
+
+		/// <inheritdoc/>
+		public override bool Equals(object? Obj) => Obj is RefId RefId && Equals(RefId);
+
+		/// <inheritdoc/>
+		public override int GetHashCode() => Hash.GetHashCode();
+
+		/// <inheritdoc/>
+		public bool Equals(RefId RefId) => Hash == RefId.Hash;
+
+		/// <inheritdoc/>
+		public override string ToString() => Hash.ToString();
+
+		/// <inheritdoc/>
+		public static bool operator ==(RefId Lhs, RefId Rhs) => Lhs.Equals(Rhs);
+
+		/// <inheritdoc/>
+		public static bool operator !=(RefId Lhs, RefId Rhs) => !Lhs.Equals(Rhs);
+	}
+
+	/// <summary>
+	/// Type converter for IoHash to and from JSON
+	/// </summary>
+	sealed class RefIdJsonConverter : JsonConverter<RefId>
+	{
+		/// <inheritdoc/>
+		public override RefId Read(ref Utf8JsonReader Reader, Type TypeToConvert, JsonSerializerOptions Options) => new RefId(IoHash.Parse(Reader.ValueSpan));
+
+		/// <inheritdoc/>
+		public override void Write(Utf8JsonWriter Writer, RefId Value, JsonSerializerOptions Options) => Writer.WriteStringValue(Value.Hash.ToUtf8String().Span);
+	}
+
+	/// <summary>
+	/// Type converter from strings to IoHash objects
+	/// </summary>
+	sealed class RefIdTypeConverter : TypeConverter
+	{
+		/// <inheritdoc/>
+		public override bool CanConvertFrom(ITypeDescriptorContext Context, Type SourceType)
+		{
+			return SourceType == typeof(string);
+		}
+
+		/// <inheritdoc/>
+		public override object ConvertFrom(ITypeDescriptorContext Context, CultureInfo Culture, object Value)
+		{
+			return new RefId(IoHash.Parse((string)Value));
+		}
+	}
+
+	/// <summary>
+	/// Type converter to compact binary
+	/// </summary>
+	sealed class RefIdCbConverter : CbConverterBase<RefId>
+	{
+		/// <inheritdoc/>
+		public override RefId Read(CbField Field) => new RefId(Field.AsHash());
+
+		/// <inheritdoc/>
+		public override void Write(CbWriter Writer, RefId Value) => Writer.WriteHashValue(Value.Hash);
+
+		/// <inheritdoc/>
+		public override void WriteNamed(CbWriter Writer, Utf8String Name, RefId Value) => Writer.WriteHash(Name, Value.Hash);
+	}
+}
diff --git a/Engine/Source/Programs/Shared/EpicGames.Horde/Storage/RefTableMessages.cs b/Engine/Source/Programs/Shared/EpicGames.Horde/Storage/RefTableMessages.cs
deleted file mode 100644
diff --git a/Engine/Source/Programs/Shared/EpicGames.Horde/Storage/StorageExtensions.cs b/Engine/Source/Programs/Shared/EpicGames.Horde/Storage/StorageExtensions.cs
new file mode 100644
--- /dev/null
+++ b/Engine/Source/Programs/Shared/EpicGames.Horde/Storage/StorageExtensions.cs
@@ -0,0 +1,100 @@
+// Copyright Epic Games, Inc. All Rights Reserved.
+
+using EpicGames.Horde.Auth;
+using EpicGames.Horde.Storage.Impl;
+using Microsoft.Extensions.DependencyInjection;
+using Microsoft.Extensions.Options;
+using System;
+using System.Collections.Generic;
+using System.Net.Http;
+using System.Text;
+
+namespace EpicGames.Horde.Storage
+{
+	/// <summary>
+	/// Auth type for the storage system
+	/// </summary>
+	public enum StorageAuthType
+	{
+		/// <summary>
+		/// No auth necessary
+		/// </summary>
+		None,
+
+		/// <summary>
+		/// Use OAuth2 for auth
+		/// </summary>
+		OAuth2,
+	}
+
+	/// <summary>
+	/// Settings to configure a connected Horde.Storage instance
+	/// </summary>
+	[System.Diagnostics.CodeAnalysis.SuppressMessage("Design", "CA1056:URI-like properties should not be strings")]
+	public class StorageSettings
+	{
+		/// <summary>
+		/// URL for the storage service
+		/// </summary>
+		public string Url { get; set; } = String.Empty;
+
+		/// <summary>
+		/// Options for OAuth2
+		/// </summary>
+		public OAuth2Options? OAuth2 { get; set; }
+	}
+
+	/// <summary>
+	/// Extension methods for configuring Horde.Storage
+	/// </summary>
+	public static class StorageExtensions
+	{
+		/// <summary>
+		/// Registers services for Horde Storage
+		/// </summary>
+		/// <param name="Services">The current service collection</param>
+		public static void AddHordeStorage(this IServiceCollection Services)
+		{
+			Services.AddOptions<StorageSettings>();
+
+			Services.AddTransient<NullAuthProvider<HttpStorageClient>>();
+
+			Services.AddScoped<OAuth2AuthProviderFactory<HttpStorageClient>>();
+			Services.AddHttpClient<OAuth2AuthProviderFactory<HttpStorageClient>>();
+
+			Services.AddScoped<IAuthProvider<HttpStorageClient>>(ServiceProvider => CreateAuthProvider(ServiceProvider));
+
+			Services.AddTransient<IStorageClient, HttpStorageClient>();
+			Services.AddHttpClient<IStorageClient, HttpStorageClient>();
+		}
+
+		/// <summary>
+		/// Registers services for Horde Storage
+		/// </summary>
+		/// <param name="Services">The current service collection</param>
+		/// <param name="Configure">Callback for configuring the storage service</param>
+		public static void AddHordeStorage(this IServiceCollection Services, Action<StorageSettings> Configure)
+		{
+			AddHordeStorage(Services);
+			Services.Configure<StorageSettings>(Configure);
+		}
+
+		/// <summary>
+		/// Creates an auth provider based on the configured settings
+		/// </summary>
+		/// <param name="Services"></param>
+		/// <returns></returns>
+		static IAuthProvider<HttpStorageClient> CreateAuthProvider(IServiceProvider Services)
+		{
+			StorageSettings Settings = Services.GetRequiredService<IOptions<StorageSettings>>().Value;
+			if (Settings.OAuth2 != null)
+			{
+				return Services.GetRequiredService<OAuth2AuthProviderFactory<HttpStorageClient>>().Create(Settings.OAuth2);
+			}
+			else
+			{
+				return Services.GetRequiredService<NullAuthProvider<HttpStorageClient>>();
+			}
+		}
+	}
+}
diff --git a/Engine/Source/Programs/Shared/EpicGames.Serialization/CbWriter.cs b/Engine/Source/Programs/Shared/EpicGames.Serialization/CbWriter.cs
--- a/Engine/Source/Programs/Shared/EpicGames.Serialization/CbWriter.cs
+++ b/Engine/Source/Programs/Shared/EpicGames.Serialization/CbWriter.cs
@@ -733,7 +733,13 @@ namespace EpicGames.Serialization
 		/// </summary>
 		/// <param name="Name">Name of the field</param>
 		/// <param name="Value">Value to be written</param>
-		public void WriteString(Utf8String Name, string Value) => WriteUtf8String(Name, Value);
+		public void WriteString(Utf8String Name, string? Value)
+		{
+			if(Value != null)
+			{
+				WriteUtf8String(Name, Value);
+			}
+		}
 
 		/// <summary>
 		/// Writes an unnamed string value
@@ -752,8 +758,11 @@ namespace EpicGames.Serialization
 		/// <param name="Value">Value to be written</param>
 		public void WriteUtf8String(Utf8String Name, Utf8String Value)
 		{
-			WriteFieldHeader(CbFieldType.String, Name);
-			WriteBinaryPayload(Value.Span);
+			if (Value.Length > 0)
+			{
+				WriteFieldHeader(CbFieldType.String, Name);
+				WriteBinaryPayload(Value.Span);
+			}
 		}
 
 		/// <summary>
diff --git a/Engine/Source/Programs/Shared/EpicGames.Serialization/Converters/CbClassConverter.cs b/Engine/Source/Programs/Shared/EpicGames.Serialization/Converters/CbClassConverter.cs
--- a/Engine/Source/Programs/Shared/EpicGames.Serialization/Converters/CbClassConverter.cs
+++ b/Engine/Source/Programs/Shared/EpicGames.Serialization/Converters/CbClassConverter.cs
@@ -436,10 +436,6 @@ namespace EpicGames.Serialization.Converters
 					PropertyList.Add((Name, Property));
 				}
 			}
-			if (PropertyList.Count == 0)
-			{
-				throw new CbException($"{Type.Name} does not have any properties marked with [CbField]");
-			}
 			return PropertyList.ToArray();
 		}
 
diff --git a/Engine/Source/Programs/Shared/EpicGames.Serialization/Converters/CbPrimitiveConverter.cs b/Engine/Source/Programs/Shared/EpicGames.Serialization/Converters/CbPrimitiveConverter.cs
--- a/Engine/Source/Programs/Shared/EpicGames.Serialization/Converters/CbPrimitiveConverter.cs
+++ b/Engine/Source/Programs/Shared/EpicGames.Serialization/Converters/CbPrimitiveConverter.cs
@@ -36,6 +36,6 @@ namespace EpicGames.Serialization.Converters
 
 		public override void Write(CbWriter Writer, T Value) => WriteFunc(Writer, Value);
 
-		public override void WriteNamed(CbWriter Writer, Utf8String Name, T Value) => WriteNamed(Writer, Name, Value);
+		public override void WriteNamed(CbWriter Writer, Utf8String Name, T Value) => WriteNamedFunc(Writer, Name, Value);
 	}
 }
