Horde: Add descriptive errors for incorrect permissions in compute service.

diff --git a/Engine/Source/Programs/Horde/Horde.Build/Compute/Impl/ComputeController.cs b/Engine/Source/Programs/Horde/Horde.Build/Compute/Impl/ComputeController.cs
--- a/Engine/Source/Programs/Horde/Horde.Build/Compute/Impl/ComputeController.cs
+++ b/Engine/Source/Programs/Horde/Horde.Build/Compute/Impl/ComputeController.cs
@@ -23,7 +23,7 @@ namespace HordeServer.Compute.Impl
 	[ApiController]
 	[Authorize]
 	[Route("[controller]")]
-	public class ComputeController : ControllerBase
+	public class ComputeController : HordeControllerBase
 	{
 		AclService AclService;
 		IComputeService ComputeService;
@@ -51,10 +51,18 @@ namespace HordeServer.Compute.Impl
 		{
 			if (!await AclService.AuthorizeAsync(AclAction.ViewComputeTasks, User))
 			{
-				return Forbid();
+				return Forbid(AclAction.ViewComputeTasks);
 			}
 
-			IComputeClusterInfo ClusterInfo = await ComputeService.GetClusterInfoAsync(ClusterId);
+			IComputeClusterInfo ClusterInfo;
+			try
+			{
+				ClusterInfo = await ComputeService.GetClusterInfoAsync(ClusterId);
+			}
+			catch (KeyNotFoundException)
+			{
+				return NotFound("Cluster '{ClusterId}' was found", ClusterId);
+			}
 
 			GetComputeClusterInfo Response = new GetComputeClusterInfo();
 			Response.Id = ClusterInfo.Id;
@@ -78,7 +86,7 @@ namespace HordeServer.Compute.Impl
 		{
 			if(!await AclService.AuthorizeAsync(AclAction.AddComputeTasks, User))
 			{
-				return Forbid();
+				return Forbid(AclAction.AddComputeTasks);
 			}
 			if (Request.TaskRefIds.Count == 0)
 			{
@@ -103,7 +111,7 @@ namespace HordeServer.Compute.Impl
 		{
 			if (!await AclService.AuthorizeAsync(AclAction.ViewComputeTasks, User))
 			{
-				return Forbid();
+				return Forbid(AclAction.ViewComputeTasks);
 			}
 
 			List<IComputeTaskStatus> Results;
