Horde: Extra logging for compute cluster responses.

diff --git a/Engine/Source/Programs/Horde/Horde.Build/Compute/Impl/ComputeService.cs b/Engine/Source/Programs/Horde/Horde.Build/Compute/Impl/ComputeService.cs
--- a/Engine/Source/Programs/Horde/Horde.Build/Compute/Impl/ComputeService.cs
+++ b/Engine/Source/Programs/Horde/Horde.Build/Compute/Impl/ComputeService.cs
@@ -36,6 +36,7 @@ using EpicGames.Redis;
 using HordeServer.Collections;
 using EpicGames.Horde.Storage;
 using EpicGames.Horde.Common;
+using System.Text.Json;
 
 namespace HordeServer.Compute.Impl
 {
@@ -264,6 +265,8 @@ namespace HordeServer.Compute.Impl
 			{
 				throw new KeyNotFoundException();
 			}
+
+			Logger.LogInformation("Found compute cluster: {ClusterId}, Reqs {ReqId} {ReqId2} {Json}", ClusterId, Config.RequestBucketId, new ClusterInfo(Config).RequestBucketId, JsonSerializer.Serialize(new ClusterInfo(Config)));
 			return new ClusterInfo(Config);
 		}
 
