Horde Storage - Transaction log when reading will now only skip a certain number of events (set to 1000 right now) before returning the current offset into the transacation log. This lets us move the pointer to where in the transaction log we were if there are large blocks of events to skip

diff --git a/Engine/Source/Programs/Horde/HordeStorage/Callisto/Implementation/FileTransactionLog.cs b/Engine/Source/Programs/Horde/HordeStorage/Callisto/Implementation/FileTransactionLog.cs
--- a/Engine/Source/Programs/Horde/HordeStorage/Callisto/Implementation/FileTransactionLog.cs
+++ b/Engine/Source/Programs/Horde/HordeStorage/Callisto/Implementation/FileTransactionLog.cs
@@ -196,6 +196,10 @@ namespace Callisto.Implementation
             if (!_transactionLogHandle.Exists)
                 throw new FileNotFoundException("Transaction log is empty");
 
+            // we set a max amount of events to skip to make sure we actually move the current event counter forward for the replicator and do not have to start from the beginning again
+            const int maxEventsSkippedDueToLocation = 1000;
+            int countOfEventsSkipped = 0;
+
             // we allow parallel reads from the file, we only need to lock on writes.
             // ReSharper disable once InconsistentlySynchronizedField
             await using FileStream fs = _transactionLogHandle.Open(FileMode.Open, FileAccess.Read, FileShare.Read);
@@ -249,8 +253,14 @@ namespace Callisto.Implementation
                 {
                     // we have a site filter, lets check that the event has not been seen at this site (flag is set)
                     bool hasBeenSeen = (op.Op!.Locations & siteFilter) != 0;
+                    if (hasBeenSeen && countOfEventsSkipped > maxEventsSkippedDueToLocation)
+                    {
+                        // we have reached the maximum number of events we skip, lets stop searching for more events for now.
+                        break;
+                    }
                     if (hasBeenSeen)
                     {
+                        ++countOfEventsSkipped;
                         continue;
                     }
                 }
