Horde Storage - Filesystem store will avoid attempting to update last write time for files while we have them open.
Now we just run the UpdateLastWrite time together with the Get/Put request (not in a seperate task), as this is generally a quick operation its not worth spawning a task for it and this allows us more control over when that runs.

diff --git a/Engine/Source/Programs/Horde/HordeStorage/Horde.Storage/Implementation/Blob/FileSystemStore.cs b/Engine/Source/Programs/Horde/HordeStorage/Horde.Storage/Implementation/Blob/FileSystemStore.cs
--- a/Engine/Source/Programs/Horde/HordeStorage/Horde.Storage/Implementation/Blob/FileSystemStore.cs
+++ b/Engine/Source/Programs/Horde/HordeStorage/Horde.Storage/Implementation/Blob/FileSystemStore.cs
@@ -56,10 +56,12 @@ namespace Horde.Storage.Implementation
             FileInfo filePath = GetFilesystemPath(ns, blobIdentifier);
             filePath.Directory?.Create();
 
-            await using FileStream fs = new FileStream(filePath.FullName, FileMode.Create, FileAccess.Write, FileShare.Read, DefaultBufferSize, FileOptions.Asynchronous | FileOptions.SequentialScan);
-            CancellationToken cancellationToken = default(CancellationToken);
-            await fs.WriteAsync(content, cancellationToken);
-            await fs.FlushAsync(cancellationToken);
+            {
+                await using FileStream fs = new FileStream(filePath.FullName, FileMode.Create, FileAccess.Write, FileShare.Read, DefaultBufferSize, FileOptions.Asynchronous | FileOptions.SequentialScan);
+                CancellationToken cancellationToken = default(CancellationToken);
+                await fs.WriteAsync(content, cancellationToken);
+                await fs.FlushAsync(cancellationToken);
+            }
 
             UpdateLastWriteTime(filePath.FullName, DateTime.UnixEpoch);
             return blobIdentifier;
@@ -71,10 +73,12 @@ namespace Horde.Storage.Implementation
             FileInfo filePath = GetFilesystemPath(ns, blobIdentifier);
             filePath.Directory?.Create();
 
-            await using FileStream fs = new FileStream(filePath.FullName, FileMode.Create, FileAccess.Write, FileShare.Read, DefaultBufferSize, FileOptions.Asynchronous | FileOptions.SequentialScan);
-            CancellationToken cancellationToken = default(CancellationToken);
-            await content.CopyToAsync(fs, cancellationToken);
-            await fs.FlushAsync(cancellationToken);
+            {
+                await using FileStream fs = new FileStream(filePath.FullName, FileMode.Create, FileAccess.Write, FileShare.Read, DefaultBufferSize, FileOptions.Asynchronous | FileOptions.SequentialScan);
+                CancellationToken cancellationToken = default(CancellationToken);
+                await content.CopyToAsync(fs, cancellationToken);
+                await fs.FlushAsync(cancellationToken);
+            }
 
             UpdateLastWriteTime(filePath.FullName, DateTime.UnixEpoch);
             return blobIdentifier;
@@ -85,10 +89,12 @@ namespace Horde.Storage.Implementation
             FileInfo filePath = GetFilesystemPath(ns, blobIdentifier);
             filePath.Directory?.Create();
 
-            await using FileStream fs = new FileStream(filePath.FullName, FileMode.Create, FileAccess.Write, FileShare.Read, DefaultBufferSize, FileOptions.Asynchronous | FileOptions.SequentialScan);
-            CancellationToken cancellationToken = default(CancellationToken);
-            await fs.WriteAsync(content, cancellationToken);
-            await fs.FlushAsync(cancellationToken);
+            {
+                await using FileStream fs = new FileStream(filePath.FullName, FileMode.Create, FileAccess.Write, FileShare.Read, DefaultBufferSize, FileOptions.Asynchronous | FileOptions.SequentialScan);
+                CancellationToken cancellationToken = default(CancellationToken);
+                await fs.WriteAsync(content, cancellationToken);
+                await fs.FlushAsync(cancellationToken);
+            }
 
             UpdateLastWriteTime(filePath.FullName, DateTime.UnixEpoch);
             return blobIdentifier;
@@ -101,9 +107,9 @@ namespace Horde.Storage.Implementation
             if (!filePath.Exists)
                 throw new BlobNotFoundException(ns, blob);
 
-            FileStream fs = new FileStream(filePath.FullName, FileMode.Open, FileAccess.Read, FileShare.Read, DefaultBufferSize, FileOptions.Asynchronous | FileOptions.SequentialScan);
-            
             UpdateLastWriteTime(filePath.FullName, DateTime.UtcNow);
+            FileStream fs = new FileStream(filePath.FullName, FileMode.Open, FileAccess.Read, FileShare.Read, DefaultBufferSize, FileOptions.Asynchronous | FileOptions.SequentialScan);
+
             return Task.FromResult(new BlobContents(fs, fs.Length));
         }
 
@@ -122,23 +128,14 @@ namespace Horde.Storage.Implementation
         /// <param name="lastAccessed">Time the file was last accessed</param>
         private void UpdateLastWriteTime(string filePath, DateTime lastAccessed)
         {
-            Task.Run(() =>
+            try
             {
-                try
-                {
-                    File.SetLastWriteTimeUtc(filePath, lastAccessed);
-                }
-                catch (FileNotFoundException)
-                {
-                    // it is okay if the file does not exist anymore, that just means it got gced
-                }
-            }).ContinueWith(task =>
+                File.SetLastWriteTimeUtc(filePath, lastAccessed);
+            }
+            catch (FileNotFoundException)
             {
-                if (task.Exception != null)
-                {
-                    _logger.Error(task.Exception, "Exception when updating write time for {FilePath}", filePath);
-                }
-            });
+                // it is okay if the file does not exist anymore, that just means it got gced
+            }
         }
 
         public Task<bool> Exists(NamespaceId ns, BlobIdentifier blob)
