Horde Storage - Fixed failing unit tests due to out of date mock

diff --git a/Engine/Source/Programs/Horde/HordeStorage/Horde.Storage/Tests/UnitTests/BlobCleanupServiceTest.cs b/Engine/Source/Programs/Horde/HordeStorage/Horde.Storage/Tests/UnitTests/BlobCleanupServiceTest.cs
--- a/Engine/Source/Programs/Horde/HordeStorage/Horde.Storage/Tests/UnitTests/BlobCleanupServiceTest.cs
+++ b/Engine/Source/Programs/Horde/HordeStorage/Horde.Storage/Tests/UnitTests/BlobCleanupServiceTest.cs
@@ -25,8 +25,9 @@ namespace Horde.Storage.UnitTests
             BlobCleanupService blobCleanupService = new BlobCleanupService(Mock.Of<IServiceProvider>(), gcSettingsMon);
 
             Mock<IBlobCleanup> store1 = new Mock<IBlobCleanup>();
+            store1.Setup(cleanup => cleanup.ShouldRun()).Returns(true);
             Mock<IBlobCleanup> store2 = new Mock<IBlobCleanup>();
-            
+            store2.Setup(cleanup => cleanup.ShouldRun()).Returns(true);
             blobCleanupService.RegisterCleanup(store1.Object);
             blobCleanupService.RegisterCleanup(store2.Object);
 
