Horde Storage - Fixes for inconsistent handling of replicator root, not always replacing variables in it.

diff --git a/Engine/Source/Programs/Horde/HordeStorage/Horde.Storage/Implementation/Replication/RefsReplicator.cs b/Engine/Source/Programs/Horde/HordeStorage/Horde.Storage/Implementation/Replication/RefsReplicator.cs
--- a/Engine/Source/Programs/Horde/HordeStorage/Horde.Storage/Implementation/Replication/RefsReplicator.cs
+++ b/Engine/Source/Programs/Horde/HordeStorage/Horde.Storage/Implementation/Replication/RefsReplicator.cs
@@ -53,7 +53,7 @@ namespace Horde.Storage.Implementation
             _httpClient.BaseAddress = new Uri(replicatorSettings.ConnectionString);
 
             string stateFileName = $"{_name}.json";
-            DirectoryInfo stateRoot = new DirectoryInfo(PathUtil.ResolvePath(replicationSettings.CurrentValue.StateRoot));
+            DirectoryInfo stateRoot = new DirectoryInfo(replicationSettings.CurrentValue.GetStateRoot());
             _stateFile = new FileInfo(Path.Combine(stateRoot.FullName, stateFileName));
 
             if (_stateFile.Exists)
diff --git a/Engine/Source/Programs/Horde/HordeStorage/Horde.Storage/Implementation/Replication/ReplicationService.cs b/Engine/Source/Programs/Horde/HordeStorage/Horde.Storage/Implementation/Replication/ReplicationService.cs
--- a/Engine/Source/Programs/Horde/HordeStorage/Horde.Storage/Implementation/Replication/ReplicationService.cs
+++ b/Engine/Source/Programs/Horde/HordeStorage/Horde.Storage/Implementation/Replication/ReplicationService.cs
@@ -40,7 +40,7 @@ namespace Horde.Storage.Implementation
 
             _leaderElection.OnLeaderChanged += OnLeaderChanged;
 
-            DirectoryInfo di = new DirectoryInfo(settings.CurrentValue.StateRoot);
+            DirectoryInfo di = new DirectoryInfo(settings.CurrentValue.GetStateRoot());
             Directory.CreateDirectory(di.FullName);
 
             foreach (var replicator in settings.CurrentValue.Replicators)
diff --git a/Engine/Source/Programs/Horde/HordeStorage/Horde.Storage/Implementation/Replication/ReplicationSettings.cs b/Engine/Source/Programs/Horde/HordeStorage/Horde.Storage/Implementation/Replication/ReplicationSettings.cs
--- a/Engine/Source/Programs/Horde/HordeStorage/Horde.Storage/Implementation/Replication/ReplicationSettings.cs
+++ b/Engine/Source/Programs/Horde/HordeStorage/Horde.Storage/Implementation/Replication/ReplicationSettings.cs
@@ -2,6 +2,7 @@
 
 using System.Collections.Generic;
 using System.ComponentModel.DataAnnotations;
+using Horde.Storage.Implementation;
 using Jupiter.Implementation;
 
 namespace Horde.Storage
@@ -35,6 +36,11 @@ namespace Horde.Storage
         [Required]
         // ReSharper disable once CollectionNeverUpdated.Global
         public List<ReplicatorSettings> Replicators { get; set; } = new List<ReplicatorSettings>();
+
+        public string GetStateRoot()
+        {
+            return PathUtil.ResolvePath(StateRoot);
+        }
     }
 
     public class ReplicatorSettings
diff --git a/Engine/Source/Programs/Horde/HordeStorage/Horde.Storage/Implementation/Replication/Replicator.cs b/Engine/Source/Programs/Horde/HordeStorage/Horde.Storage/Implementation/Replication/Replicator.cs
--- a/Engine/Source/Programs/Horde/HordeStorage/Horde.Storage/Implementation/Replication/Replicator.cs
+++ b/Engine/Source/Programs/Horde/HordeStorage/Horde.Storage/Implementation/Replication/Replicator.cs
@@ -123,7 +123,7 @@ namespace Horde.Storage.Implementation
             _replicatorSettings = replicatorSettings;
 
             string stateFileName = $"{_name}.json";
-            DirectoryInfo stateRoot = new DirectoryInfo(PathUtil.ResolvePath(replicationSettings.CurrentValue.StateRoot));
+            DirectoryInfo stateRoot = new DirectoryInfo(replicationSettings.CurrentValue.GetStateRoot());
             _stateFile = new FileInfo(Path.Combine(stateRoot.FullName, stateFileName));
 
             if (_stateFile.Exists)
