Horde: Terminate tasks in queues with missing/invalid requirements objects.

#skipci

diff --git a/Engine/Source/Programs/Horde/Horde.Build/Compute/Impl/ComputeService.cs b/Engine/Source/Programs/Horde/Horde.Build/Compute/Impl/ComputeService.cs
--- a/Engine/Source/Programs/Horde/Horde.Build/Compute/Impl/ComputeService.cs
+++ b/Engine/Source/Programs/Horde/Horde.Build/Compute/Impl/ComputeService.cs
@@ -113,7 +113,7 @@ namespace HordeServer.Compute.Impl
 		/// <inheritdoc/>
 		public LeaseId? LeaseId
 		{
-			get => (LeaseIdBytes == null || LeaseIdBytes.Length == 0)? (LeaseId?)null : new LeaseId(LeaseIdBytes);
+			get => (LeaseIdBytes == null || LeaseIdBytes.Length == 0) ? (LeaseId?)null : new LeaseId(LeaseIdBytes);
 			set => LeaseIdBytes = (value != null) ? value.Value.Value.ToByteArray() : null;
 		}
 
@@ -126,8 +126,8 @@ namespace HordeServer.Compute.Impl
 			this.TaskRefId = TaskRefId;
 			this.Time = DateTime.UtcNow;
 			this.State = State;
-			this.AgentId = (AgentId == null)? Utf8String.Empty : AgentId.Value.ToString();
-			this.LeaseIdBytes = (LeaseId != null)? LeaseId.Value.Value.ToByteArray() : null;
+			this.AgentId = (AgentId == null) ? Utf8String.Empty : AgentId.Value.ToString();
+			this.LeaseIdBytes = (LeaseId != null) ? LeaseId.Value.Value.ToByteArray() : null;
 		}
 	}
 
@@ -419,39 +419,70 @@ namespace HordeServer.Compute.Impl
 		/// <summary>
 		/// Gets the requirements object from the CAS
 		/// </summary>
-		/// <param name="QueueKey"></param>
-		/// <returns></returns>
+		/// <param name="QueueKey">Queue identifier</param>
+		/// <returns>Requirements object for the queue</returns>
 		async ValueTask<Requirements?> GetCachedRequirementsAsync(QueueKey QueueKey)
 		{
 			Requirements? Requirements;
 			if (!RequirementsCache.TryGetValue(QueueKey.RequirementsHash, out Requirements))
 			{
-				ComputeClusterConfig? ClusterConfig = await GetClusterAsync(QueueKey.ClusterId);
-				if (ClusterConfig != null)
+				Requirements = await GetRequirementsAsync(QueueKey);
+				if (Requirements != null)
 				{
-					try
+					using (ICacheEntry Entry = RequirementsCache.CreateEntry(QueueKey.RequirementsHash))
 					{
-						Requirements = await StorageClient.ReadObjectAsync<Requirements>(new NamespaceId(ClusterConfig.NamespaceId), QueueKey.RequirementsHash);
-					}
-					catch (BlobNotFoundException)
-					{
-						Requirements = null;
+						Entry.SetSlidingExpiration(TimeSpan.FromMinutes(10.0));
+						Entry.SetValue(Requirements);
 					}
 				}
+			}
+			return Requirements;
+		}
 
-				using (ICacheEntry Entry = RequirementsCache.CreateEntry(QueueKey.RequirementsHash))
+		/// <summary>
+		/// Gets the requirements object for a given queue. Fails tasks in the queue if the requirements object is missing.
+		/// </summary>
+		/// <param name="QueueKey">Queue identifier</param>
+		/// <returns>Requirements object for the queue</returns>
+		async ValueTask<Requirements?> GetRequirementsAsync(QueueKey QueueKey)
+		{
+			Requirements? Requirements = null;
+
+			ComputeClusterConfig? ClusterConfig = await GetClusterAsync(QueueKey.ClusterId);
+			if (ClusterConfig != null)
+			{
+				NamespaceId NamespaceId = new NamespaceId(ClusterConfig.NamespaceId);
+				try
 				{
-					if (Requirements == null)
-					{
-						Entry.SetAbsoluteExpiration(TimeSpan.FromSeconds(10.0));
-					}
-					else
+					Requirements = await StorageClient.ReadObjectAsync<Requirements>(NamespaceId, QueueKey.RequirementsHash);
+				}
+				catch (BlobNotFoundException)
+				{
+				}
+				catch (Exception Ex)
+				{
+					Logger.LogError(Ex, "Unable to read blob {NamespaceId}/{RequirementsHash} from storage service", ClusterConfig.NamespaceId, QueueKey.RequirementsHash);
+				}
+			}
+
+			if (Requirements == null)
+			{
+				Logger.LogWarning("Unable to fetch requirements object for queue {QueueKey}; failing queued tasks.", QueueKey);
+				for (; ; )
+				{
+					ComputeTaskInfo? ComputeTask = await TaskScheduler.DequeueAsync(QueueKey);
+					if (ComputeTask == null)
 					{
-						Entry.SetSlidingExpiration(TimeSpan.FromMinutes(10.0));
+						break;
 					}
-					Entry.SetValue(Requirements);
+
+					ComputeTaskStatus Status = new ComputeTaskStatus(ComputeTask.TaskRefId, ComputeTaskState.Complete, null, null);
+					Status.Outcome = ComputeTaskOutcome.BlobNotFound;
+					Logger.LogInformation("Compute task failed due to missing requirements (queue: {RequirementsHash}, task: {TaskHash}, channel: {ChannelId})", QueueKey, ComputeTask.TaskRefId, ComputeTask.ChannelId);
+					await MessageQueue.PostAsync(ComputeTask.ChannelId.ToString(), Status);
 				}
 			}
+
 			return Requirements;
 		}
 	}
