Horde: Rename ExecuteCommandV2 to ComputeCommand.

#skipci

diff --git a/Engine/Source/Programs/Horde/Horde.Agent/Commands/Compute/ComputeCommand.cs b/Engine/Source/Programs/Horde/Horde.Agent/Commands/Compute/ComputeCommand.cs
new file mode 100644
--- /dev/null
+++ b/Engine/Source/Programs/Horde/Horde.Agent/Commands/Compute/ComputeCommand.cs
@@ -0,0 +1,312 @@
+// Copyright Epic Games, Inc. All Rights Reserved.
+
+using Google.Protobuf;
+using Microsoft.Extensions.Logging;
+using System;
+using System.Collections.Concurrent;
+using System.Collections.Generic;
+using System.IO;
+using System.Linq;
+using System.Text;
+using System.Threading.Tasks;
+using EpicGames.Core;
+using System.Text.Json;
+using HordeAgent.Services;
+using Grpc.Net.Client;
+using System.Net.Http.Headers;
+using Grpc.Core;
+using Microsoft.Extensions.Hosting;
+using Microsoft.Extensions.DependencyInjection;
+using Microsoft.Extensions.Configuration;
+using Serilog.Events;
+using EpicGames.Serialization;
+using EpicGames.Horde.Compute;
+using EpicGames.Horde.Storage;
+using System.Net.Http;
+using EpicGames.Horde.Compute.Impl;
+using EpicGames.Horde.Storage.Impl;
+
+namespace HordeAgent.Commands
+{
+	/// <summary>
+	/// Installs the agent as a service
+	/// </summary>
+	[Command("Compute", "Executes a command through the Horde Compute API")]
+	class ComputeCommand : Command
+	{
+		class JsonRequirements
+		{
+			public string? Condition { get; set; }
+			public Dictionary<string, int> Resources { get; set; } = new Dictionary<string, int>();
+			public bool Exclusive { get; set; }
+		}
+
+		class JsonComputeTask
+		{
+			public ComputeOptions ComputeServer { get; set; } = new ComputeOptions();
+			public StorageOptions StorageServer { get; set; } = new StorageOptions();
+			public ClusterId ClusterId { get; set; } = new ClusterId("default");
+			public string Executable { get; set; } = String.Empty;
+			public List<string> Arguments { get; set; } = new List<string>();
+			public Dictionary<string, string> EnvVars { get; set; } = new Dictionary<string, string>();
+			public List<string> OutputPaths { get; set; } = new List<string>();
+			public string WorkingDirectory { get; set; } = String.Empty;
+			public JsonRequirements Requirements { get; set; } = new JsonRequirements();
+		}
+
+		/// <summary>
+		/// Input file describing the work to execute
+		/// </summary>
+		[CommandLine(Required = true)]
+		public FileReference Input = null!;
+
+		/// <summary>
+		/// The input directory. By default, the directory containing the input file will be used.
+		/// </summary>
+		[CommandLine]
+		public DirectoryReference? InputDir = null;
+
+		/// <summary>
+		/// Apply a random salt to the cached value
+		/// </summary>
+		[CommandLine("-Salt")]
+		public bool RandomSalt;
+
+		/// <summary>
+		/// Add a known salt value to the cache
+		/// </summary>
+		[CommandLine("-Salt=")]
+		public string? Salt = null;
+						
+		/// <summary>
+		/// Skip checking if a result is already available in the action cache
+		/// </summary>
+		[CommandLine("-SkipCacheLookup")]
+		public bool SkipCacheLookup = false;
+		
+		/// <summary>
+		/// Directory to download the output files to. If not set, no results will be downloaded.
+		/// </summary>
+		[CommandLine("-OutputDir")]
+		public string? OutputDir = null;
+	
+		/// <summary>
+		/// Log verbosity level (use normal Serilog levels such as debug, warning or info)
+		/// </summary>
+		[CommandLine("-LogLevel")]
+		public string LogLevelStr = "debug";
+
+		static (DirectoryTree, IoHash) CreateSandbox(DirectoryInfo BaseDirInfo, Dictionary<IoHash, byte[]> UploadList)
+		{
+			DirectoryTree Tree = new DirectoryTree();
+
+			foreach (DirectoryInfo SubDirInfo in BaseDirInfo.EnumerateDirectories())
+			{
+				(DirectoryTree SubTree, IoHash SubDirHash) = CreateSandbox(SubDirInfo, UploadList);
+				Tree.Directories.Add(new DirectoryNode(SubDirInfo.Name, SubDirHash));
+			}
+			Tree.Directories.SortBy(x => x.Name, Utf8StringComparer.Ordinal);
+
+			foreach (FileInfo FileInfo in BaseDirInfo.EnumerateFiles())
+			{
+				byte[] Data = File.ReadAllBytes(FileInfo.FullName);
+				IoHash Hash = IoHash.Compute(Data);
+				UploadList[Hash] = Data;
+				Tree.Files.Add(new FileNode(FileInfo.Name, Hash, FileInfo.Length, (int)FileInfo.Attributes));
+			}
+			Tree.Files.SortBy(x => x.Name, Utf8StringComparer.Ordinal);
+
+			return (Tree, AddCbObject(UploadList, Tree));
+		}
+
+		/// <inheritdoc/>
+		public override async Task<int> ExecuteAsync(ILogger Logger)
+		{
+			InputDir ??= Input.Directory;
+
+			if (Enum.TryParse(LogLevelStr, true, out LogEventLevel LogEventLevel))
+			{
+				Logging.LogLevelSwitch.MinimumLevel = LogEventLevel;
+			}
+			else
+			{
+				Console.WriteLine($"Unable to parse log level: {this.LogLevelStr}");
+				return 0;
+			}
+
+			IConfiguration Configuration = new ConfigurationBuilder()
+				.AddEnvironmentVariables()
+				.AddJsonFile(Input.FullName)
+				.Build();
+
+			IHostBuilder HostBuilder = Host.CreateDefaultBuilder()
+				.ConfigureLogging(Builder =>
+				{
+					Builder.SetMinimumLevel(LogLevel.Warning);
+				})
+				.ConfigureServices(Services =>
+				{
+					Services.AddLogging();
+
+					IConfigurationSection ComputeSettings = Configuration.GetSection(nameof(JsonComputeTask.ComputeServer));
+					Services.AddHordeCompute(Settings => ComputeSettings.Bind(Settings));
+
+					IConfigurationSection StorageSettings = Configuration.GetSection(nameof(JsonComputeTask.StorageServer));
+					Services.AddHordeStorage(Settings => StorageSettings.Bind(Settings));
+				});
+
+			using (IHost Host = HostBuilder.Build())
+			{
+				IStorageClient StorageClient = Host.Services.GetRequiredService<IStorageClient>();
+				IComputeClient ComputeClient = Host.Services.GetRequiredService<IComputeClient>();
+
+				ByteString SaltBytes = ByteString.Empty;
+				if (Salt != null)
+				{
+					SaltBytes = ByteString.CopyFromUtf8(Salt);
+				}
+				else if(RandomSalt)
+				{
+					SaltBytes = ByteString.CopyFrom(Guid.NewGuid().ToByteArray());
+				}
+
+				if (SaltBytes.Length > 0)
+				{
+					Logger.LogInformation("Using salt: {SaltBytes}", StringUtils.FormatHexString(SaltBytes.ToByteArray()));
+				}
+
+				Dictionary<IoHash, byte[]> Blobs = new Dictionary<IoHash, byte[]>();
+				(_, IoHash SandboxHash) = CreateSandbox(InputDir.ToDirectoryInfo(), Blobs);
+
+				JsonComputeTask JsonComputeTask = new JsonComputeTask();
+				Configuration.Bind(JsonComputeTask);
+
+				Logger.LogInformation("compute server: {ServerUrl}", JsonComputeTask.ComputeServer.Url);
+				Logger.LogInformation("storage server: {ServerUrl}", JsonComputeTask.StorageServer.Url);
+
+				JsonRequirements JsonRequirements = JsonComputeTask.Requirements;
+				Requirements Requirements = new Requirements(JsonRequirements.Condition ?? String.Empty);
+				Requirements.Resources = JsonRequirements.Resources.ToDictionary(x => x.Key, x => x.Value);
+				Requirements.Exclusive = JsonRequirements.Exclusive;
+				IoHash RequirementsHash = AddCbObject(Blobs, Requirements);
+
+				ComputeTask Task = new ComputeTask(JsonComputeTask.Executable, JsonComputeTask.Arguments.ConvertAll<Utf8String>(x => x), JsonComputeTask.WorkingDirectory, SandboxHash);
+				Task.EnvVars = JsonComputeTask.EnvVars.ToDictionary(x => (Utf8String)x.Key, x => (Utf8String)x.Value);
+				Task.OutputPaths.AddRange(JsonComputeTask.OutputPaths.Select(x => (Utf8String)x));
+				Task.RequirementsHash = AddCbObject(Blobs, Requirements);
+
+				await ExecuteAction(Logger, ComputeClient, StorageClient, JsonComputeTask.ClusterId, Task, Blobs);	
+			}
+			return 0;
+		}
+
+		static IoHash AddCbObject<T>(Dictionary<IoHash, byte[]> HashToData, T Source)
+		{
+			CbObject Obj = CbSerializer.Serialize<T>(Source);
+			IoHash Hash = Obj.GetHash();
+			HashToData[Hash] = Obj.GetView().ToArray();
+			return Hash;
+		}
+
+		private async Task ExecuteAction(ILogger Logger, IComputeClient ComputeClient, IStorageClient StorageClient, ClusterId ClusterId, ComputeTask Task, Dictionary<IoHash, byte[]> UploadList)
+		{
+			IComputeClusterInfo Cluster = await ComputeClient.GetClusterInfoAsync(ClusterId);
+
+			List<KeyValuePair<IoHash, byte[]>> UploadBlobs = UploadList.ToList();
+			for (int Idx = 0; Idx < UploadBlobs.Count; Idx++)
+			{
+				KeyValuePair<IoHash, byte[]> Pair = UploadBlobs[Idx];
+				Logger.LogInformation("Uploading blob {Idx}/{Count}: {Hash}", Idx + 1, UploadBlobs.Count, Pair.Key);
+				await StorageClient.WriteBlobFromMemoryAsync(Cluster.NamespaceId, Pair.Key, Pair.Value);
+			}
+
+			CbObject TaskObject = CbSerializer.Serialize(Task);
+			IoHash TaskHash = IoHash.Compute(TaskObject.GetView().Span);
+			RefId TaskRefId = new RefId(TaskHash);
+			await StorageClient.SetRefAsync(Cluster.NamespaceId, Cluster.RequestBucketId, TaskRefId, TaskObject);
+
+			ChannelId ChannelId = new ChannelId(Guid.NewGuid().ToString());
+			Logger.LogInformation("cluster: {ClusterId}", ClusterId);
+			Logger.LogInformation("channel: {ChannelId}", ChannelId);
+			Logger.LogInformation("task: {TaskHash}", TaskHash);
+			Logger.LogInformation("requirements: {RequirementsHash}", Task.RequirementsHash);
+
+
+			// Execute the action
+			await ComputeClient.AddTaskAsync(ClusterId, ChannelId, TaskRefId, Task.RequirementsHash, SkipCacheLookup);
+
+			await foreach(IComputeTaskInfo Response in ComputeClient.GetTaskUpdatesAsync(ClusterId, ChannelId))
+			{
+				Logger.LogInformation("{OperationName}: Execution state: {State}", Response.TaskRefId, Response.State.ToString());
+				if (!String.IsNullOrEmpty(Response.AgentId) || !String.IsNullOrEmpty(Response.LeaseId))
+				{
+					Logger.LogInformation("{OperationName}: Running on agent {AgentId} under lease {LeaseId}", Response.TaskRefId, Response.AgentId, Response.LeaseId);
+				}
+				if (Response.ResultRefId != null)
+				{
+					await HandleCompleteTask(StorageClient, Cluster.NamespaceId, Cluster.ResponseBucketId, Response.ResultRefId.Value, Logger);
+				}
+				if (Response.State == ComputeTaskState.Complete)
+				{
+					break;
+				}
+			}
+		}
+
+		async Task HandleCompleteTask(IStorageClient StorageClient, NamespaceId NamespaceId, BucketId OutputBucketId, RefId OutputRefId, ILogger Logger)
+		{
+			ComputeTaskResult Result = await StorageClient.GetRefAsync<ComputeTaskResult>(NamespaceId, OutputBucketId, OutputRefId);
+			Logger.LogInformation("exit: {ExitCode}", Result.ExitCode);
+
+			await LogTaskOutputAsync(StorageClient, "stdout", NamespaceId, Result.StdOutHash, Logger);
+			await LogTaskOutputAsync(StorageClient, "stderr", NamespaceId, Result.StdErrHash, Logger);
+
+			if (Result.OutputHash != null && OutputDir != null)
+			{
+				await WriteOutputAsync(StorageClient, NamespaceId, Result.OutputHash.Value, new DirectoryReference(OutputDir));
+			}
+		}
+
+		async Task LogTaskOutputAsync(IStorageClient StorageClient, string Channel, NamespaceId NamespaceId, IoHash? LogHash, ILogger Logger)
+		{
+			if (LogHash != null)
+			{
+				byte[] StdOutData = await StorageClient.ReadBlobToMemoryAsync(NamespaceId, LogHash.Value);
+				if (StdOutData.Length > 0)
+				{
+					foreach (string Line in Encoding.UTF8.GetString(StdOutData).Split('\n'))
+					{
+						Logger.LogDebug("{Channel}: {Line}", Channel, Line);
+					}
+				}
+			}
+		}
+
+		async Task WriteOutputAsync(IStorageClient StorageClient, NamespaceId NamespaceId, IoHash TreeHash, DirectoryReference OutputDir)
+		{
+			DirectoryTree Tree = await StorageClient.ReadObjectAsync<DirectoryTree>(NamespaceId, TreeHash);
+
+			List<Task> Tasks = new List<Task>();
+			foreach (FileNode File in Tree.Files)
+			{
+				FileReference OutputFile = FileReference.Combine(OutputDir, File.Name.ToString());
+				Tasks.Add(WriteObjectToFileAsync(StorageClient, NamespaceId, File.Hash, OutputFile));
+			}
+			foreach (DirectoryNode Directory in Tree.Directories)
+			{
+				DirectoryReference NextOutputDir = DirectoryReference.Combine(OutputDir, Directory.Name.ToString());
+				Tasks.Add(WriteOutputAsync(StorageClient, NamespaceId, Directory.Hash, NextOutputDir));
+			}
+
+			await Task.WhenAll(Tasks);
+		}
+
+		static async Task WriteObjectToFileAsync(IStorageClient StorageClient, NamespaceId NamespaceId, IoHash Hash, FileReference OutputFile)
+		{
+			DirectoryReference.CreateDirectory(OutputFile.Directory);
+
+			byte[] Data = await StorageClient.ReadBlobToMemoryAsync(NamespaceId, Hash);
+			await FileReference.WriteAllBytesAsync(OutputFile, Data);
+		}
+	}
+}
diff --git a/Engine/Source/Programs/Horde/Horde.Agent/Commands/Execution/ExecuteCommandV2.cs b/Engine/Source/Programs/Horde/Horde.Agent/Commands/Execution/ExecuteCommandV2.cs
deleted file mode 100644
