Horde: Fix warnings.

#skipci

diff --git a/Engine/Source/Programs/Horde/Horde.Agent/Horde.Agent.csproj b/Engine/Source/Programs/Horde/Horde.Agent/Horde.Agent.csproj
--- a/Engine/Source/Programs/Horde/Horde.Agent/Horde.Agent.csproj
+++ b/Engine/Source/Programs/Horde/Horde.Agent/Horde.Agent.csproj
@@ -69,9 +69,6 @@
   </ItemGroup>
 
   <ItemGroup>
-    <Content Include="..\..\Shared\EpicGames.Horde.Common\Protos\**\*.*">
-      <Link>Protos\%(RecursiveDir)%(FileName)%(Extension)</Link>
-    </Content>
     <Content Include="..\Horde.Build\Protos\horde\horde_rpc_messages.proto">
       <Link>Protos\horde_rpc_messages.proto</Link>
     </Content>
diff --git a/Engine/Source/Programs/Horde/Horde.Build/Storage/Services/BasicStorageClient.cs b/Engine/Source/Programs/Horde/Horde.Build/Storage/Services/BasicStorageClient.cs
--- a/Engine/Source/Programs/Horde/Horde.Build/Storage/Services/BasicStorageClient.cs
+++ b/Engine/Source/Programs/Horde/Horde.Build/Storage/Services/BasicStorageClient.cs
@@ -121,7 +121,7 @@ namespace Horde.Build.Storage.Services
 			RefImpl NewRef = new RefImpl();
 			NewRef.Id = GetFullRefId(NamespaceId, BucketId, RefId);
 			NewRef.Data = Value.GetView().ToArray();
-			await Refs.ReplaceOneAsync(x => x.Id == NewRef.Id, NewRef, new ReplaceOptions { IsUpsert = true });
+			await Refs.ReplaceOneAsync(x => x.Id == NewRef.Id, NewRef, new ReplaceOptions { IsUpsert = true }, CancellationToken);
 
 			IoHash Hash = IoHash.Compute(Value.GetView().Span);
 			return await TryFinalizeRefAsync(NamespaceId, BucketId, RefId, Hash, CancellationToken);
@@ -143,7 +143,7 @@ namespace Horde.Build.Storage.Services
 
 				if (MissingHashes.Count == 0)
 				{
-					await Refs.UpdateOneAsync(x => x.Id == Id, Builders<RefImpl>.Update.Set(x => x.Finalized, true));
+					await Refs.UpdateOneAsync(x => x.Id == Id, Builders<RefImpl>.Update.Set(x => x.Finalized, true), cancellationToken: CancellationToken);
 				}
 			}
 			return MissingHashes.ToList();
