Horde: Fix some messages not being returned for compute queues.

#skipci

diff --git a/Engine/Source/Programs/Horde/Horde.Build/Compute/Impl/ComputeService.cs b/Engine/Source/Programs/Horde/Horde.Build/Compute/Impl/ComputeService.cs
--- a/Engine/Source/Programs/Horde/Horde.Build/Compute/Impl/ComputeService.cs
+++ b/Engine/Source/Programs/Horde/Horde.Build/Compute/Impl/ComputeService.cs
@@ -241,7 +241,7 @@ namespace HordeServer.Compute.Impl
 					ComputeTaskStatus Status = new ComputeTaskStatus(ComputeTask.TaskRefId, ComputeTaskState.Complete, null, null);
 					Status.Outcome = ComputeTaskOutcome.Expired;
 					Logger.LogInformation("Compute task expired (queue: {RequirementsHash}, task: {TaskHash}, channel: {ChannelId})", QueueKey, ComputeTask.TaskRefId, ComputeTask.ChannelId);
-					await MessageQueue.PostAsync(ComputeTask.ChannelId.ToString(), Status);
+					await PostStatusMessageAsync(ComputeTask, Status);
 				}
 			}
 		}
@@ -309,7 +309,7 @@ namespace HordeServer.Compute.Impl
 				{
 					Logger.LogWarning("Invalid cluster '{ClusterId}'; failing task {TaskRefId}", TaskInfo.ClusterId, TaskInfo.TaskRefId);
 					ComputeTaskStatus Status = new ComputeTaskStatus(TaskInfo.TaskRefId, ComputeTaskState.Complete, Agent.Id, null) { Detail = $"Invalid cluster '{TaskInfo.ClusterId}'" };
-					await MessageQueue.PostAsync(TaskInfo.ChannelId.ToString(), Status);
+					await PostStatusMessageAsync(TaskInfo, Status);
 					return null;
 				}
 
@@ -318,7 +318,7 @@ namespace HordeServer.Compute.Impl
 				{
 					Logger.LogWarning("Unable to fetch requirements {RequirementsHash}", QueueKey);
 					ComputeTaskStatus Status = new ComputeTaskStatus(TaskInfo.TaskRefId, ComputeTaskState.Complete, Agent.Id, null) { Detail = $"Unable to retrieve requirements '{QueueKey}'" };
-					await MessageQueue.PostAsync(TaskInfo.ChannelId.ToString(), Status);
+					await PostStatusMessageAsync(TaskInfo, Status);
 					return null;
 				}
 
@@ -352,26 +352,28 @@ namespace HordeServer.Compute.Impl
 		/// <inheritdoc/>
 		public async Task<List<IComputeTaskStatus>> GetTaskUpdatesAsync(ClusterId ClusterId, ChannelId ChannelId)
 		{
-			List<ComputeTaskStatus> Messages = await MessageQueue.ReadMessagesAsync($"{ClusterId}/{ChannelId}");
+			List<ComputeTaskStatus> Messages = await MessageQueue.ReadMessagesAsync(GetMessageQueueId(ClusterId, ChannelId));
 			return Messages.ConvertAll<IComputeTaskStatus>(x => x);
 		}
 
 		public async Task<List<IComputeTaskStatus>> WaitForTaskUpdatesAsync(ClusterId ClusterId, ChannelId ChannelId, CancellationToken CancellationToken)
 		{
-			List<ComputeTaskStatus> Messages = await MessageQueue.WaitForMessagesAsync($"{ClusterId}/{ChannelId}", CancellationToken);
+			List<ComputeTaskStatus> Messages = await MessageQueue.WaitForMessagesAsync(GetMessageQueueId(ClusterId, ChannelId), CancellationToken);
 			return Messages.ConvertAll<IComputeTaskStatus>(x => x);
 		}
 
-		public override Task OnLeaseStartedAsync(IAgent Agent, LeaseId LeaseId, ComputeTaskMessage ComputeTask, ILogger Logger)
+		public override async Task OnLeaseStartedAsync(IAgent Agent, LeaseId LeaseId, ComputeTaskMessage ComputeTask, ILogger Logger)
 		{
-			base.OnLeaseStartedAsync(Agent, LeaseId, ComputeTask, Logger);
+			await base.OnLeaseStartedAsync(Agent, LeaseId, ComputeTask, Logger);
 
 			ComputeTaskStatus Status = new ComputeTaskStatus(ComputeTask.TaskRefId, ComputeTaskState.Executing, Agent.Id, LeaseId);
-			return MessageQueue.PostAsync(ComputeTask.ChannelId, Status);
+			await PostStatusMessageAsync(ComputeTask, Status);
 		}
 
-		public override Task OnLeaseFinishedAsync(IAgent Agent, LeaseId LeaseId, ComputeTaskMessage ComputeTask, LeaseOutcome Outcome, ReadOnlyMemory<byte> Output, ILogger Logger)
+		public override async Task OnLeaseFinishedAsync(IAgent Agent, LeaseId LeaseId, ComputeTaskMessage ComputeTask, LeaseOutcome Outcome, ReadOnlyMemory<byte> Output, ILogger Logger)
 		{
+			await base.OnLeaseFinishedAsync(Agent, LeaseId, ComputeTask, Outcome, Output, Logger);
+
 			ComputeTaskResultMessage Message = ComputeTaskResultMessage.Parser.ParseFrom(Output.ToArray());
 
 			ComputeTaskStatus Status = new ComputeTaskStatus(ComputeTask.TaskRefId, ComputeTaskState.Complete, Agent.Id, LeaseId);
@@ -397,7 +399,7 @@ namespace HordeServer.Compute.Impl
 			}
 
 			Logger.LogInformation("Compute lease finished (lease: {LeaseId}, task: {TaskHash}, agent: {AgentId}, channel: {ChannelId}, result: {ResultHash}, outcome: {Outcome})", LeaseId, ComputeTask.TaskRefId.AsRefId(), Agent.Id, ComputeTask.ChannelId, Status.ResultRefId?.ToString() ?? "(none)", Status.Outcome);
-			return MessageQueue.PostAsync($"{ComputeTask.ClusterId}/{ComputeTask.ChannelId}", Status);
+			await PostStatusMessageAsync(ComputeTask, Status);
 		}
 
 		/// <summary>
@@ -479,11 +481,43 @@ namespace HordeServer.Compute.Impl
 					ComputeTaskStatus Status = new ComputeTaskStatus(ComputeTask.TaskRefId, ComputeTaskState.Complete, null, null);
 					Status.Outcome = ComputeTaskOutcome.BlobNotFound;
 					Logger.LogInformation("Compute task failed due to missing requirements (queue: {RequirementsHash}, task: {TaskHash}, channel: {ChannelId})", QueueKey, ComputeTask.TaskRefId, ComputeTask.ChannelId);
-					await MessageQueue.PostAsync(ComputeTask.ChannelId.ToString(), Status);
+					await PostStatusMessageAsync(ComputeTask, Status);
 				}
 			}
 
 			return Requirements;
 		}
+
+		/// <summary>
+		/// Post a status message for a particular task
+		/// </summary>
+		/// <param name="ComputeTask">The compute task instance</param>
+		/// <param name="Status">New status for the task</param>
+		async Task PostStatusMessageAsync(ComputeTaskInfo ComputeTask, ComputeTaskStatus Status)
+		{
+			await MessageQueue.PostAsync(GetMessageQueueId(ComputeTask.ClusterId, ComputeTask.ChannelId), Status);
+		}
+
+		/// <summary>
+		/// Post a status message for a particular task
+		/// </summary>
+		/// <param name="ComputeTaskMessage">The compute task lease</param>
+		/// <param name="Status">New status for the task</param>
+		/// <returns></returns>
+		async Task PostStatusMessageAsync(ComputeTaskMessage ComputeTaskMessage, ComputeTaskStatus Status)
+		{
+			await MessageQueue.PostAsync(GetMessageQueueId(new ClusterId(ComputeTaskMessage.ClusterId), new ChannelId(ComputeTaskMessage.ChannelId)), Status);
+		}
+
+		/// <summary>
+		/// Gets the name of a particular message queue
+		/// </summary>
+		/// <param name="ClusterId">The compute cluster</param>
+		/// <param name="ChannelId">Identifier for the message channel</param>
+		/// <returns>Name of the message queue</returns>
+		static string GetMessageQueueId(ClusterId ClusterId, ChannelId ChannelId)
+		{
+			return $"{ClusterId}/{ChannelId}";
+		}
 	}
 }
