Horde: Improve diagnostic output from compute command.

diff --git a/Engine/Source/Programs/Horde/Horde.Agent/Commands/Compute/ComputeCommand.cs b/Engine/Source/Programs/Horde/Horde.Agent/Commands/Compute/ComputeCommand.cs
--- a/Engine/Source/Programs/Horde/Horde.Agent/Commands/Compute/ComputeCommand.cs
+++ b/Engine/Source/Programs/Horde/Horde.Agent/Commands/Compute/ComputeCommand.cs
@@ -248,6 +248,10 @@ namespace HordeAgent.Commands
 				}
 				if (Response.State == ComputeTaskState.Complete)
 				{
+					if (Response.Outcome != ComputeTaskOutcome.Success)
+					{
+						Logger.LogError("{OperationName}: Outcome: {Outcome}, Detail: {Detail}", Response.TaskRefId, Response.Outcome.ToString(), Response.Detail ?? "(none)");
+					}
 					break;
 				}
 			}
diff --git a/Engine/Source/Programs/Horde/Horde.Build/Compute/Impl/ComputeService.cs b/Engine/Source/Programs/Horde/Horde.Build/Compute/Impl/ComputeService.cs
--- a/Engine/Source/Programs/Horde/Horde.Build/Compute/Impl/ComputeService.cs
+++ b/Engine/Source/Programs/Horde/Horde.Build/Compute/Impl/ComputeService.cs
@@ -240,6 +240,7 @@ namespace HordeServer.Compute.Impl
 
 					ComputeTaskStatus Status = new ComputeTaskStatus(ComputeTask.TaskRefId, ComputeTaskState.Complete, null, null);
 					Status.Outcome = ComputeTaskOutcome.Expired;
+					Status.Outcome = $"No agents monitoring queue {QueueKey}";
 					Logger.LogInformation("Compute task expired (queue: {RequirementsHash}, task: {TaskHash}, channel: {ChannelId})", QueueKey, ComputeTask.TaskRefId, ComputeTask.ChannelId);
 					await PostStatusMessageAsync(ComputeTask, Status);
 				}
@@ -480,7 +481,8 @@ namespace HordeServer.Compute.Impl
 
 					ComputeTaskStatus Status = new ComputeTaskStatus(ComputeTask.TaskRefId, ComputeTaskState.Complete, null, null);
 					Status.Outcome = ComputeTaskOutcome.BlobNotFound;
-					Logger.LogInformation("Compute task failed due to missing requirements (queue: {RequirementsHash}, task: {TaskHash}, channel: {ChannelId})", QueueKey, ComputeTask.TaskRefId, ComputeTask.ChannelId);
+					Status.Detail = $"Missing requirements object {QueueKey.RequirementsHash}";
+					Logger.LogInformation("Compute task failed due to missing requirements (queue: {QueueKey}, task: {TaskHash}, channel: {ChannelId})", QueueKey, ComputeTask.TaskRefId, ComputeTask.ChannelId);
 					await PostStatusMessageAsync(ComputeTask, Status);
 				}
 			}
