Horde: Missing file.

diff --git a/Engine/Source/Programs/Horde/Horde.Build/Compute/Impl/ComputeService.cs b/Engine/Source/Programs/Horde/Horde.Build/Compute/Impl/ComputeService.cs
--- a/Engine/Source/Programs/Horde/Horde.Build/Compute/Impl/ComputeService.cs
+++ b/Engine/Source/Programs/Horde/Horde.Build/Compute/Impl/ComputeService.cs
@@ -240,7 +240,7 @@ namespace HordeServer.Compute.Impl
 
 					ComputeTaskStatus Status = new ComputeTaskStatus(ComputeTask.TaskRefId, ComputeTaskState.Complete, null, null);
 					Status.Outcome = ComputeTaskOutcome.Expired;
-					Status.Outcome = $"No agents monitoring queue {QueueKey}";
+					Status.Detail = $"No agents monitoring queue {QueueKey}";
 					Logger.LogInformation("Compute task expired (queue: {RequirementsHash}, task: {TaskHash}, channel: {ChannelId})", QueueKey, ComputeTask.TaskRefId, ComputeTask.ChannelId);
 					await PostStatusMessageAsync(ComputeTask, Status);
 				}
