Horde: Fix swagger documentation generation.

* Use existing types in EpicGames.Horde.Storage for NamespaceId and BucketId.
* Rename DDCRefController.BatchCall, DDCRefController.BatchOp, and DDCRefController.BatchOp.Operation to prevent conflicts building docs.

diff --git a/Engine/Source/Programs/Horde/HordeStorage/Horde.Storage/Controllers/AdminController.cs b/Engine/Source/Programs/Horde/HordeStorage/Horde.Storage/Controllers/AdminController.cs
--- a/Engine/Source/Programs/Horde/HordeStorage/Horde.Storage/Controllers/AdminController.cs
+++ b/Engine/Source/Programs/Horde/HordeStorage/Horde.Storage/Controllers/AdminController.cs
@@ -7,6 +7,7 @@ using System.Linq;
 using System.Net;
 using System.Threading;
 using System.Threading.Tasks;
+using EpicGames.Horde.Storage;
 using Horde.Storage.Implementation;
 using Jupiter.Implementation;
 using Microsoft.AspNetCore.Authorization;
diff --git a/Engine/Source/Programs/Horde/HordeStorage/Horde.Storage/Controllers/AuthController.cs b/Engine/Source/Programs/Horde/HordeStorage/Horde.Storage/Controllers/AuthController.cs
--- a/Engine/Source/Programs/Horde/HordeStorage/Horde.Storage/Controllers/AuthController.cs
+++ b/Engine/Source/Programs/Horde/HordeStorage/Horde.Storage/Controllers/AuthController.cs
@@ -11,6 +11,7 @@ using Datadog.Trace;
 using Jupiter;
 using Jupiter.Implementation;
 using Microsoft.AspNetCore.Authorization;
+using EpicGames.Horde.Storage;
 
 namespace Horde.Storage.Controllers
 {
diff --git a/Engine/Source/Programs/Horde/HordeStorage/Horde.Storage/Controllers/CompressedBlobController.cs b/Engine/Source/Programs/Horde/HordeStorage/Horde.Storage/Controllers/CompressedBlobController.cs
--- a/Engine/Source/Programs/Horde/HordeStorage/Horde.Storage/Controllers/CompressedBlobController.cs
+++ b/Engine/Source/Programs/Horde/HordeStorage/Horde.Storage/Controllers/CompressedBlobController.cs
@@ -9,6 +9,7 @@ using System.Linq;
 using System.Net.Mime;
 using System.Threading.Tasks;
 using Datadog.Trace;
+using EpicGames.Horde.Storage;
 using Horde.Storage.Implementation;
 using Jupiter;
 using Jupiter.Common.Implementation;
@@ -20,6 +21,8 @@ using Serilog;
 
 namespace Horde.Storage.Controllers
 {
+    using BlobNotFoundException = Horde.Storage.Implementation.BlobNotFoundException;
+
     [ApiController]
     [Route("api/v1/compressed-blobs")]
     public class CompressedBlobController : ControllerBase
diff --git a/Engine/Source/Programs/Horde/HordeStorage/Horde.Storage/Controllers/DDCRefController.cs b/Engine/Source/Programs/Horde/HordeStorage/Horde.Storage/Controllers/DDCRefController.cs
--- a/Engine/Source/Programs/Horde/HordeStorage/Horde.Storage/Controllers/DDCRefController.cs
+++ b/Engine/Source/Programs/Horde/HordeStorage/Horde.Storage/Controllers/DDCRefController.cs
@@ -9,6 +9,7 @@ using System.Net.Mime;
 using System.Threading.Tasks;
 using Dasync.Collections;
 using Datadog.Trace;
+using EpicGames.Horde.Storage;
 using Horde.Storage.Implementation;
 using Jupiter;
 using Jupiter.Implementation;
@@ -25,6 +26,8 @@ using Serilog;
 
 namespace Horde.Storage.Controllers
 {
+    using BlobNotFoundException = Horde.Storage.Implementation.BlobNotFoundException;
+
     public class RefRequest
     {
         [JsonConstructor]
@@ -656,10 +659,10 @@ namespace Horde.Storage.Controllers
 
         // ReSharper disable UnusedAutoPropertyAccessor.Global
         // ReSharper disable once ClassNeverInstantiated.Global
-        public class BatchOp
+        public class DdcBatchOp
         {
             // ReSharper disable InconsistentNaming
-            public enum Operation
+            public enum DdcOperation
             {
                 INVALID,
                 GET,
@@ -676,31 +679,31 @@ namespace Horde.Storage.Controllers
 
             [Required] public KeyId? Id { get; set; }
 
-            [Required] public Operation Op { get; set; }
+            [Required] public DdcOperation Op { get; set; }
 
             // used for put ops
             public ContentHash? ContentHash { get; set; }
             public byte[]? Content { get; set; }
         }
 
-        public class BatchCall
+        public class DdcBatchCall
         {
-            public BatchOp[]? Operations { get; set; }
+            public DdcBatchOp[]? Operations { get; set; }
         }
 
         [HttpPost("ddc-rpc")]
-        public async Task<IActionResult> Post([FromBody] BatchCall batch)
+        public async Task<IActionResult> Post([FromBody] DdcBatchCall batch)
         {
-            string OpToPolicy(BatchOp.Operation op)
+            string OpToPolicy(DdcBatchOp.DdcOperation op)
             {
                 switch (op)
                 {
-                    case BatchOp.Operation.GET:
-                    case BatchOp.Operation.HEAD:
+                    case DdcBatchOp.DdcOperation.GET:
+                    case DdcBatchOp.DdcOperation.HEAD:
                         return "Cache.read";
-                    case BatchOp.Operation.PUT:
+                    case DdcBatchOp.DdcOperation.PUT:
                         return "Cache.write";
-                    case BatchOp.Operation.DELETE:
+                    case DdcBatchOp.DdcOperation.DELETE:
                         return "Cache.delete";
                     default:
                         throw new ArgumentOutOfRangeException(nameof(op), op, null);
@@ -715,7 +718,7 @@ namespace Horde.Storage.Controllers
             Task<object?>[] tasks = new Task<object?>[batch.Operations.Length];
             for (int index = 0; index < batch.Operations.Length; index++)
             {
-                BatchOp op = batch.Operations[index];
+                DdcBatchOp op = batch.Operations[index];
 
                 AuthorizationResult authorizationResultNamespace = await _authorizationService.AuthorizeAsync(User, op.Namespace, NamespaceAccessRequirement.Name);
                 AuthorizationResult authorizationResultOp = await _authorizationService.AuthorizeAsync(User, op, OpToPolicy(op.Op));
@@ -739,7 +742,7 @@ namespace Horde.Storage.Controllers
             return Ok(results);
         }
 
-        private Task<object?> ProcessOp(BatchOp op)
+        private Task<object?> ProcessOp(DdcBatchOp op)
         {
             if (op.Namespace == null)
             {
@@ -758,9 +761,9 @@ namespace Horde.Storage.Controllers
 
             switch (op.Op)
             {
-                case BatchOp.Operation.INVALID:
+                case DdcBatchOp.DdcOperation.INVALID:
                     throw new ArgumentOutOfRangeException();
-                case BatchOp.Operation.GET:
+                case DdcBatchOp.DdcOperation.GET:
                     // TODO: support field filtering
                     return _ddcRefService.Get(op.Namespace.Value, op.Bucket.Value, op.Id.Value, new string[0]).ContinueWith(t => 
                     {
@@ -773,9 +776,9 @@ namespace Horde.Storage.Controllers
                         }
                         return (object?)response;
                     });
-                case BatchOp.Operation.HEAD:
+                case DdcBatchOp.DdcOperation.HEAD:
                     return _ddcRefService.Exists(op.Namespace.Value, op.Bucket.Value, op.Id.Value).ContinueWith(t => t.Result == null ? (object?)null : op.Id);
-                case BatchOp.Operation.PUT:
+                case DdcBatchOp.DdcOperation.PUT:
                 {
                     if (op.ContentHash == null)
                     {
@@ -794,7 +797,7 @@ namespace Horde.Storage.Controllers
                     }
                     return _ddcRefService.Put(op.Namespace.Value, op.Bucket.Value, op.Id.Value, blobHash, op.Content).ContinueWith(t => (object?)t.Result);
                 }
-                case BatchOp.Operation.DELETE:
+                case DdcBatchOp.DdcOperation.DELETE:
                     return _ddcRefService.Delete(op.Namespace.Value, op.Bucket.Value, op.Id.Value).ContinueWith(t => (object?)null);
                 default:
                     throw new ArgumentOutOfRangeException();
diff --git a/Engine/Source/Programs/Horde/HordeStorage/Horde.Storage/Controllers/ObjectController.cs b/Engine/Source/Programs/Horde/HordeStorage/Horde.Storage/Controllers/ObjectController.cs
--- a/Engine/Source/Programs/Horde/HordeStorage/Horde.Storage/Controllers/ObjectController.cs
+++ b/Engine/Source/Programs/Horde/HordeStorage/Horde.Storage/Controllers/ObjectController.cs
@@ -10,6 +10,7 @@ using System.Net.Mime;
 using System.Threading.Tasks;
 using async_enumerable_dotnet;
 using Datadog.Trace;
+using EpicGames.Horde.Storage;
 using Horde.Storage.Implementation;
 using Jupiter;
 using Jupiter.Common.Implementation;
@@ -22,6 +23,8 @@ using Serilog;
 
 namespace Horde.Storage.Controllers
 {
+    using BlobNotFoundException = Horde.Storage.Implementation.BlobNotFoundException;
+
     [ApiController]
     [Route("api/v1/objects", Order = 0)]
     [Produces(CustomMediaTypeNames.UnrealCompactBinary, MediaTypeNames.Application.Json)]
diff --git a/Engine/Source/Programs/Horde/HordeStorage/Horde.Storage/Controllers/ReferencesController.cs b/Engine/Source/Programs/Horde/HordeStorage/Horde.Storage/Controllers/ReferencesController.cs
--- a/Engine/Source/Programs/Horde/HordeStorage/Horde.Storage/Controllers/ReferencesController.cs
+++ b/Engine/Source/Programs/Horde/HordeStorage/Horde.Storage/Controllers/ReferencesController.cs
@@ -12,6 +12,7 @@ using System.Text;
 using System.Threading.Tasks;
 using Dasync.Collections;
 using Datadog.Trace;
+using EpicGames.Horde.Storage;
 using Horde.Storage.Implementation;
 using Jupiter;
 using Jupiter.Common.Implementation;
@@ -26,6 +27,8 @@ using Serilog;
 
 namespace Horde.Storage.Controllers
 {
+    using BlobNotFoundException = Horde.Storage.Implementation.BlobNotFoundException;
+
     [ApiController]
     [FormatFilter]
     [Produces(MediaTypeNames.Application.Json, MediaTypeNames.Application.Octet, CustomMediaTypeNames.UnrealCompactBinary)]
diff --git a/Engine/Source/Programs/Horde/HordeStorage/Horde.Storage/Controllers/ReplicationController.cs b/Engine/Source/Programs/Horde/HordeStorage/Horde.Storage/Controllers/ReplicationController.cs
--- a/Engine/Source/Programs/Horde/HordeStorage/Horde.Storage/Controllers/ReplicationController.cs
+++ b/Engine/Source/Programs/Horde/HordeStorage/Horde.Storage/Controllers/ReplicationController.cs
@@ -5,6 +5,7 @@ using System.Collections.Generic;
 using System.ComponentModel.DataAnnotations;
 using System.Linq;
 using System.Threading.Tasks;
+using EpicGames.Horde.Storage;
 using Horde.Storage.Implementation;
 using Jupiter;
 using Jupiter.Implementation;
diff --git a/Engine/Source/Programs/Horde/HordeStorage/Horde.Storage/Controllers/ReplicationLogController.cs b/Engine/Source/Programs/Horde/HordeStorage/Horde.Storage/Controllers/ReplicationLogController.cs
--- a/Engine/Source/Programs/Horde/HordeStorage/Horde.Storage/Controllers/ReplicationLogController.cs
+++ b/Engine/Source/Programs/Horde/HordeStorage/Horde.Storage/Controllers/ReplicationLogController.cs
@@ -6,6 +6,7 @@ using System.ComponentModel.DataAnnotations;
 using System.Threading;
 using System.Threading.Tasks;
 using async_enumerable_dotnet;
+using EpicGames.Horde.Storage;
 using Horde.Storage.Implementation;
 using Horde.Storage.Implementation.TransactionLog;
 using Jupiter;
diff --git a/Engine/Source/Programs/Horde/HordeStorage/Horde.Storage/Controllers/StorageController.cs b/Engine/Source/Programs/Horde/HordeStorage/Horde.Storage/Controllers/StorageController.cs
--- a/Engine/Source/Programs/Horde/HordeStorage/Horde.Storage/Controllers/StorageController.cs
+++ b/Engine/Source/Programs/Horde/HordeStorage/Horde.Storage/Controllers/StorageController.cs
@@ -9,6 +9,7 @@ using System.Linq;
 using System.Net.Mime;
 using System.Threading.Tasks;
 using Datadog.Trace;
+using EpicGames.Horde.Storage;
 using Horde.Storage.Implementation;
 using Jupiter;
 using Jupiter.Common.Implementation;
@@ -21,6 +22,8 @@ using Serilog;
 
 namespace Horde.Storage.Controllers
 {
+    using BlobNotFoundException = Horde.Storage.Implementation.BlobNotFoundException;
+
     [ApiController]
     [Route("api/v1/s", Order = 1)]
     [Route("api/v1/blobs", Order = 0)]
diff --git a/Engine/Source/Programs/Horde/HordeStorage/Horde.Storage/Endpoints/DDCEndpoints.cs b/Engine/Source/Programs/Horde/HordeStorage/Horde.Storage/Endpoints/DDCEndpoints.cs
--- a/Engine/Source/Programs/Horde/HordeStorage/Horde.Storage/Endpoints/DDCEndpoints.cs
+++ b/Engine/Source/Programs/Horde/HordeStorage/Horde.Storage/Endpoints/DDCEndpoints.cs
@@ -4,6 +4,7 @@ using System;
 using System.Net.Mime;
 using System.Threading.Tasks;
 using Datadog.Trace;
+using EpicGames.Horde.Storage;
 using Horde.Storage.Controllers;
 using Horde.Storage.Implementation;
 using Jupiter;
@@ -14,6 +15,8 @@ using Microsoft.AspNetCore.Routing;
 
 namespace Horde.Storage
 {
+    using BlobNotFoundException = Horde.Storage.Implementation.BlobNotFoundException;
+
     public class DDCEndpoints
     {
         private readonly IDDCRefService _ddcRefService;
diff --git a/Engine/Source/Programs/Horde/HordeStorage/Horde.Storage/Implementation/Blob/AmazonS3Store.cs b/Engine/Source/Programs/Horde/HordeStorage/Horde.Storage/Implementation/Blob/AmazonS3Store.cs
--- a/Engine/Source/Programs/Horde/HordeStorage/Horde.Storage/Implementation/Blob/AmazonS3Store.cs
+++ b/Engine/Source/Programs/Horde/HordeStorage/Horde.Storage/Implementation/Blob/AmazonS3Store.cs
@@ -7,6 +7,7 @@ using System.Net;
 using System.Threading.Tasks;
 using Amazon.S3;
 using Amazon.S3.Model;
+using EpicGames.Horde.Storage;
 using Horde.Storage.Implementation.Blob;
 using Jupiter;
 using Jupiter.Common;
diff --git a/Engine/Source/Programs/Horde/HordeStorage/Horde.Storage/Implementation/Blob/AzureBlobStore.cs b/Engine/Source/Programs/Horde/HordeStorage/Horde.Storage/Implementation/Blob/AzureBlobStore.cs
--- a/Engine/Source/Programs/Horde/HordeStorage/Horde.Storage/Implementation/Blob/AzureBlobStore.cs
+++ b/Engine/Source/Programs/Horde/HordeStorage/Horde.Storage/Implementation/Blob/AzureBlobStore.cs
@@ -11,6 +11,7 @@ using Azure.Storage.Blobs;
 using Azure.Storage.Blobs.Models;
 using Jupiter.Implementation;
 using Serilog;
+using EpicGames.Horde.Storage;
 
 namespace Horde.Storage.Implementation
 {
diff --git a/Engine/Source/Programs/Horde/HordeStorage/Horde.Storage/Implementation/Blob/ConsistencyCheckService.cs b/Engine/Source/Programs/Horde/HordeStorage/Horde.Storage/Implementation/Blob/ConsistencyCheckService.cs
--- a/Engine/Source/Programs/Horde/HordeStorage/Horde.Storage/Implementation/Blob/ConsistencyCheckService.cs
+++ b/Engine/Source/Programs/Horde/HordeStorage/Horde.Storage/Implementation/Blob/ConsistencyCheckService.cs
@@ -8,6 +8,7 @@ using System.Threading.Tasks;
 using Cassandra;
 using Dasync.Collections;
 using Datadog.Trace;
+using EpicGames.Horde.Storage;
 using Jupiter;
 using Jupiter.Implementation;
 using Microsoft.Extensions.DependencyInjection;
diff --git a/Engine/Source/Programs/Horde/HordeStorage/Horde.Storage/Implementation/Blob/FileSystemStore.cs b/Engine/Source/Programs/Horde/HordeStorage/Horde.Storage/Implementation/Blob/FileSystemStore.cs
--- a/Engine/Source/Programs/Horde/HordeStorage/Horde.Storage/Implementation/Blob/FileSystemStore.cs
+++ b/Engine/Source/Programs/Horde/HordeStorage/Horde.Storage/Implementation/Blob/FileSystemStore.cs
@@ -7,6 +7,7 @@ using System.Linq;
 using System.Threading;
 using System.Threading.Tasks;
 using Dasync.Collections;
+using EpicGames.Horde.Storage;
 using Jupiter.Implementation;
 using Microsoft.Extensions.Options;
 using Serilog;
diff --git a/Engine/Source/Programs/Horde/HordeStorage/Horde.Storage/Implementation/Blob/IBlobIndex.cs b/Engine/Source/Programs/Horde/HordeStorage/Horde.Storage/Implementation/Blob/IBlobIndex.cs
--- a/Engine/Source/Programs/Horde/HordeStorage/Horde.Storage/Implementation/Blob/IBlobIndex.cs
+++ b/Engine/Source/Programs/Horde/HordeStorage/Horde.Storage/Implementation/Blob/IBlobIndex.cs
@@ -2,6 +2,7 @@
 
 using System.Collections.Generic;
 using System.Threading.Tasks;
+using EpicGames.Horde.Storage;
 using Jupiter.Implementation;
 
 namespace Horde.Storage.Implementation.Blob
diff --git a/Engine/Source/Programs/Horde/HordeStorage/Horde.Storage/Implementation/Blob/IBlobService.cs b/Engine/Source/Programs/Horde/HordeStorage/Horde.Storage/Implementation/Blob/IBlobService.cs
--- a/Engine/Source/Programs/Horde/HordeStorage/Horde.Storage/Implementation/Blob/IBlobService.cs
+++ b/Engine/Source/Programs/Horde/HordeStorage/Horde.Storage/Implementation/Blob/IBlobService.cs
@@ -8,6 +8,7 @@ using System.Linq;
 using System.Threading.Tasks;
 using Dasync.Collections;
 using Datadog.Trace;
+using EpicGames.Horde.Storage;
 using Horde.Storage.Controllers;
 using Horde.Storage.Implementation.Blob;
 using Jupiter.Common.Implementation;
diff --git a/Engine/Source/Programs/Horde/HordeStorage/Horde.Storage/Implementation/Blob/IBlobStore.cs b/Engine/Source/Programs/Horde/HordeStorage/Horde.Storage/Implementation/Blob/IBlobStore.cs
--- a/Engine/Source/Programs/Horde/HordeStorage/Horde.Storage/Implementation/Blob/IBlobStore.cs
+++ b/Engine/Source/Programs/Horde/HordeStorage/Horde.Storage/Implementation/Blob/IBlobStore.cs
@@ -9,6 +9,7 @@ using System.Threading.Tasks;
 using async_enumerable_dotnet;
 using Dasync.Collections;
 using Datadog.Trace;
+using EpicGames.Horde.Storage;
 using Jupiter.Implementation;
 
 namespace Horde.Storage.Implementation
diff --git a/Engine/Source/Programs/Horde/HordeStorage/Horde.Storage/Implementation/Blob/MemoryBlobIndex.cs b/Engine/Source/Programs/Horde/HordeStorage/Horde.Storage/Implementation/Blob/MemoryBlobIndex.cs
--- a/Engine/Source/Programs/Horde/HordeStorage/Horde.Storage/Implementation/Blob/MemoryBlobIndex.cs
+++ b/Engine/Source/Programs/Horde/HordeStorage/Horde.Storage/Implementation/Blob/MemoryBlobIndex.cs
@@ -4,6 +4,7 @@ using System;
 using System.Collections.Concurrent;
 using System.Collections.Generic;
 using System.Threading.Tasks;
+using EpicGames.Horde.Storage;
 using Jupiter;
 using Jupiter.Implementation;
 using Microsoft.Extensions.Options;
diff --git a/Engine/Source/Programs/Horde/HordeStorage/Horde.Storage/Implementation/Blob/MemoryBlobStore.cs b/Engine/Source/Programs/Horde/HordeStorage/Horde.Storage/Implementation/Blob/MemoryBlobStore.cs
--- a/Engine/Source/Programs/Horde/HordeStorage/Horde.Storage/Implementation/Blob/MemoryBlobStore.cs
+++ b/Engine/Source/Programs/Horde/HordeStorage/Horde.Storage/Implementation/Blob/MemoryBlobStore.cs
@@ -6,6 +6,7 @@ using System.Collections.Generic;
 using System.IO;
 using System.Linq;
 using System.Threading.Tasks;
+using EpicGames.Horde.Storage;
 using Jupiter;
 using Jupiter.Implementation;
 using Jupiter.Utils;
diff --git a/Engine/Source/Programs/Horde/HordeStorage/Horde.Storage/Implementation/Blob/MemoryCacheBlobStore.cs b/Engine/Source/Programs/Horde/HordeStorage/Horde.Storage/Implementation/Blob/MemoryCacheBlobStore.cs
--- a/Engine/Source/Programs/Horde/HordeStorage/Horde.Storage/Implementation/Blob/MemoryCacheBlobStore.cs
+++ b/Engine/Source/Programs/Horde/HordeStorage/Horde.Storage/Implementation/Blob/MemoryCacheBlobStore.cs
@@ -7,6 +7,7 @@ using System.Reflection;
 using System.Threading;
 using System.Threading.Tasks;
 using Dasync.Collections;
+using EpicGames.Horde.Storage;
 using Jupiter.Implementation;
 using Jupiter.Utils;
 using Microsoft.Extensions.Caching.Memory;
diff --git a/Engine/Source/Programs/Horde/HordeStorage/Horde.Storage/Implementation/Blob/MongoBlobIndex.cs b/Engine/Source/Programs/Horde/HordeStorage/Horde.Storage/Implementation/Blob/MongoBlobIndex.cs
--- a/Engine/Source/Programs/Horde/HordeStorage/Horde.Storage/Implementation/Blob/MongoBlobIndex.cs
+++ b/Engine/Source/Programs/Horde/HordeStorage/Horde.Storage/Implementation/Blob/MongoBlobIndex.cs
@@ -3,6 +3,7 @@
 using System.Collections.Generic;
 using System.Linq;
 using System.Threading.Tasks;
+using EpicGames.Horde.Storage;
 using Jupiter;
 using Jupiter.Implementation;
 using Microsoft.Extensions.Options;
diff --git a/Engine/Source/Programs/Horde/HordeStorage/Horde.Storage/Implementation/Blob/ScyllaBlobIndex.cs b/Engine/Source/Programs/Horde/HordeStorage/Horde.Storage/Implementation/Blob/ScyllaBlobIndex.cs
--- a/Engine/Source/Programs/Horde/HordeStorage/Horde.Storage/Implementation/Blob/ScyllaBlobIndex.cs
+++ b/Engine/Source/Programs/Horde/HordeStorage/Horde.Storage/Implementation/Blob/ScyllaBlobIndex.cs
@@ -6,6 +6,7 @@ using System.Threading.Tasks;
 using Cassandra;
 using Cassandra.Mapping;
 using Datadog.Trace;
+using EpicGames.Horde.Storage;
 using Jupiter;
 using Jupiter.Implementation;
 using Microsoft.Extensions.Options;
diff --git a/Engine/Source/Programs/Horde/HordeStorage/Horde.Storage/Implementation/ContentId/IContentIdStore.cs b/Engine/Source/Programs/Horde/HordeStorage/Horde.Storage/Implementation/ContentId/IContentIdStore.cs
--- a/Engine/Source/Programs/Horde/HordeStorage/Horde.Storage/Implementation/ContentId/IContentIdStore.cs
+++ b/Engine/Source/Programs/Horde/HordeStorage/Horde.Storage/Implementation/ContentId/IContentIdStore.cs
@@ -2,6 +2,7 @@
 
 using System;
 using System.Threading.Tasks;
+using EpicGames.Horde.Storage;
 using Jupiter.Implementation;
 
 namespace Horde.Storage.Implementation
diff --git a/Engine/Source/Programs/Horde/HordeStorage/Horde.Storage/Implementation/ContentId/MemoryContentIdStore.cs b/Engine/Source/Programs/Horde/HordeStorage/Horde.Storage/Implementation/ContentId/MemoryContentIdStore.cs
--- a/Engine/Source/Programs/Horde/HordeStorage/Horde.Storage/Implementation/ContentId/MemoryContentIdStore.cs
+++ b/Engine/Source/Programs/Horde/HordeStorage/Horde.Storage/Implementation/ContentId/MemoryContentIdStore.cs
@@ -7,6 +7,7 @@ using System.Linq;
 using System.Threading.Tasks;
 using Cassandra;
 using Cassandra.Mapping;
+using EpicGames.Horde.Storage;
 using Jupiter.Implementation;
 
 namespace Horde.Storage.Implementation
diff --git a/Engine/Source/Programs/Horde/HordeStorage/Horde.Storage/Implementation/ContentId/MongoContentIdStore.cs b/Engine/Source/Programs/Horde/HordeStorage/Horde.Storage/Implementation/ContentId/MongoContentIdStore.cs
--- a/Engine/Source/Programs/Horde/HordeStorage/Horde.Storage/Implementation/ContentId/MongoContentIdStore.cs
+++ b/Engine/Source/Programs/Horde/HordeStorage/Horde.Storage/Implementation/ContentId/MongoContentIdStore.cs
@@ -3,6 +3,7 @@
 using System.Collections.Generic;
 using System.Linq;
 using System.Threading.Tasks;
+using EpicGames.Horde.Storage;
 using Jupiter.Implementation;
 using Microsoft.Extensions.Options;
 using MongoDB.Bson.Serialization.Attributes;
diff --git a/Engine/Source/Programs/Horde/HordeStorage/Horde.Storage/Implementation/ContentId/ScyllaContentIdStore.cs b/Engine/Source/Programs/Horde/HordeStorage/Horde.Storage/Implementation/ContentId/ScyllaContentIdStore.cs
--- a/Engine/Source/Programs/Horde/HordeStorage/Horde.Storage/Implementation/ContentId/ScyllaContentIdStore.cs
+++ b/Engine/Source/Programs/Horde/HordeStorage/Horde.Storage/Implementation/ContentId/ScyllaContentIdStore.cs
@@ -7,6 +7,7 @@ using System.Threading.Tasks;
 using Cassandra;
 using Cassandra.Mapping;
 using Datadog.Trace;
+using EpicGames.Horde.Storage;
 using Jupiter.Implementation;
 
 namespace Horde.Storage.Implementation
diff --git a/Engine/Source/Programs/Horde/HordeStorage/Horde.Storage/Implementation/DDCRefService.cs b/Engine/Source/Programs/Horde/HordeStorage/Horde.Storage/Implementation/DDCRefService.cs
--- a/Engine/Source/Programs/Horde/HordeStorage/Horde.Storage/Implementation/DDCRefService.cs
+++ b/Engine/Source/Programs/Horde/HordeStorage/Horde.Storage/Implementation/DDCRefService.cs
@@ -6,6 +6,7 @@ using System.IO;
 using System.Linq;
 using System.Threading.Tasks;
 using Datadog.Trace;
+using EpicGames.Horde.Storage;
 using Horde.Storage.Controllers;
 using Jupiter.Implementation;
 using Microsoft.Extensions.Options;
diff --git a/Engine/Source/Programs/Horde/HordeStorage/Horde.Storage/Implementation/GC/CacheCleanup.cs b/Engine/Source/Programs/Horde/HordeStorage/Horde.Storage/Implementation/GC/CacheCleanup.cs
--- a/Engine/Source/Programs/Horde/HordeStorage/Horde.Storage/Implementation/GC/CacheCleanup.cs
+++ b/Engine/Source/Programs/Horde/HordeStorage/Horde.Storage/Implementation/GC/CacheCleanup.cs
@@ -5,6 +5,7 @@ using System.Collections.Generic;
 using System.Threading;
 using System.Threading.Tasks;
 using Datadog.Trace;
+using EpicGames.Horde.Storage;
 using Jupiter.Common;
 using Jupiter.Implementation;
 using Microsoft.Extensions.Options;
diff --git a/Engine/Source/Programs/Horde/HordeStorage/Horde.Storage/Implementation/GC/OrphanBlobCleanup.cs b/Engine/Source/Programs/Horde/HordeStorage/Horde.Storage/Implementation/GC/OrphanBlobCleanup.cs
--- a/Engine/Source/Programs/Horde/HordeStorage/Horde.Storage/Implementation/GC/OrphanBlobCleanup.cs
+++ b/Engine/Source/Programs/Horde/HordeStorage/Horde.Storage/Implementation/GC/OrphanBlobCleanup.cs
@@ -8,6 +8,7 @@ using System.Threading;
 using System.Threading.Tasks;
 using async_enumerable_dotnet;
 using Dasync.Collections;
+using EpicGames.Horde.Storage;
 using Jupiter;
 using Jupiter.Implementation;
 using Microsoft.Extensions.Options;
diff --git a/Engine/Source/Programs/Horde/HordeStorage/Horde.Storage/Implementation/GC/OrphanBlobCleanupRefs.cs b/Engine/Source/Programs/Horde/HordeStorage/Horde.Storage/Implementation/GC/OrphanBlobCleanupRefs.cs
--- a/Engine/Source/Programs/Horde/HordeStorage/Horde.Storage/Implementation/GC/OrphanBlobCleanupRefs.cs
+++ b/Engine/Source/Programs/Horde/HordeStorage/Horde.Storage/Implementation/GC/OrphanBlobCleanupRefs.cs
@@ -8,6 +8,7 @@ using System.Threading.Tasks;
 using async_enumerable_dotnet;
 using Dasync.Collections;
 using Datadog.Trace;
+using EpicGames.Horde.Storage;
 using Horde.Storage.Controllers;
 using Horde.Storage.Implementation.Blob;
 using Jupiter;
diff --git a/Engine/Source/Programs/Horde/HordeStorage/Horde.Storage/Implementation/GC/RefCleanupService.cs b/Engine/Source/Programs/Horde/HordeStorage/Horde.Storage/Implementation/GC/RefCleanupService.cs
--- a/Engine/Source/Programs/Horde/HordeStorage/Horde.Storage/Implementation/GC/RefCleanupService.cs
+++ b/Engine/Source/Programs/Horde/HordeStorage/Horde.Storage/Implementation/GC/RefCleanupService.cs
@@ -6,6 +6,7 @@ using System.Threading;
 using System.Threading.Tasks;
 using Dasync.Collections;
 using Datadog.Trace;
+using EpicGames.Horde.Storage;
 using Jupiter;
 using Jupiter.Implementation;
 using Microsoft.Extensions.Options;
diff --git a/Engine/Source/Programs/Horde/HordeStorage/Horde.Storage/Implementation/IReferenceResolver.cs b/Engine/Source/Programs/Horde/HordeStorage/Horde.Storage/Implementation/IReferenceResolver.cs
--- a/Engine/Source/Programs/Horde/HordeStorage/Horde.Storage/Implementation/IReferenceResolver.cs
+++ b/Engine/Source/Programs/Horde/HordeStorage/Horde.Storage/Implementation/IReferenceResolver.cs
@@ -6,6 +6,7 @@ using System.Collections.Generic;
 using System.Reflection.Metadata;
 using System.Threading.Tasks;
 using Datadog.Trace;
+using EpicGames.Horde.Storage;
 using Jupiter.Implementation;
 using Jupiter.Utils;
 
diff --git a/Engine/Source/Programs/Horde/HordeStorage/Horde.Storage/Implementation/Kestrel/KestrelDispatcher.cs b/Engine/Source/Programs/Horde/HordeStorage/Horde.Storage/Implementation/Kestrel/KestrelDispatcher.cs
--- a/Engine/Source/Programs/Horde/HordeStorage/Horde.Storage/Implementation/Kestrel/KestrelDispatcher.cs
+++ b/Engine/Source/Programs/Horde/HordeStorage/Horde.Storage/Implementation/Kestrel/KestrelDispatcher.cs
@@ -7,6 +7,7 @@ using System.IO.Pipelines;
 using System.Net.Sockets;
 using System.Threading;
 using System.Threading.Tasks;
+using EpicGames.Horde.Storage;
 using Horde.Storage.Controllers;
 using Jupiter.Implementation;
 using Microsoft.AspNetCore.Server.Kestrel.Core.Internal.Http;
diff --git a/Engine/Source/Programs/Horde/HordeStorage/Horde.Storage/Implementation/LastAccess/LastAccessTracker.cs b/Engine/Source/Programs/Horde/HordeStorage/Horde.Storage/Implementation/LastAccess/LastAccessTracker.cs
--- a/Engine/Source/Programs/Horde/HordeStorage/Horde.Storage/Implementation/LastAccess/LastAccessTracker.cs
+++ b/Engine/Source/Programs/Horde/HordeStorage/Horde.Storage/Implementation/LastAccess/LastAccessTracker.cs
@@ -7,6 +7,7 @@ using System.Linq;
 using System.Threading;
 using System.Threading.Tasks;
 using Datadog.Trace;
+using EpicGames.Horde.Storage;
 using Jupiter.Implementation;
 using Serilog;
 
diff --git a/Engine/Source/Programs/Horde/HordeStorage/Horde.Storage/Implementation/ObjectService.cs b/Engine/Source/Programs/Horde/HordeStorage/Horde.Storage/Implementation/ObjectService.cs
--- a/Engine/Source/Programs/Horde/HordeStorage/Horde.Storage/Implementation/ObjectService.cs
+++ b/Engine/Source/Programs/Horde/HordeStorage/Horde.Storage/Implementation/ObjectService.cs
@@ -7,6 +7,7 @@ using System.Linq;
 using System.Threading.Tasks;
 using async_enumerable_dotnet;
 using Datadog.Trace;
+using EpicGames.Horde.Storage;
 using Horde.Storage.Implementation.Blob;
 using Jupiter.Implementation;
 using Jupiter.Utils;
diff --git a/Engine/Source/Programs/Horde/HordeStorage/Horde.Storage/Implementation/Objects/IReferencesStore.cs b/Engine/Source/Programs/Horde/HordeStorage/Horde.Storage/Implementation/Objects/IReferencesStore.cs
--- a/Engine/Source/Programs/Horde/HordeStorage/Horde.Storage/Implementation/Objects/IReferencesStore.cs
+++ b/Engine/Source/Programs/Horde/HordeStorage/Horde.Storage/Implementation/Objects/IReferencesStore.cs
@@ -3,6 +3,7 @@
 using System;
 using System.Collections.Generic;
 using System.Threading.Tasks;
+using EpicGames.Horde.Storage;
 using Jupiter.Implementation;
 
 namespace Horde.Storage.Implementation
diff --git a/Engine/Source/Programs/Horde/HordeStorage/Horde.Storage/Implementation/Objects/MemoryReferencesStore.cs b/Engine/Source/Programs/Horde/HordeStorage/Horde.Storage/Implementation/Objects/MemoryReferencesStore.cs
--- a/Engine/Source/Programs/Horde/HordeStorage/Horde.Storage/Implementation/Objects/MemoryReferencesStore.cs
+++ b/Engine/Source/Programs/Horde/HordeStorage/Horde.Storage/Implementation/Objects/MemoryReferencesStore.cs
@@ -6,6 +6,7 @@ using System.Collections.Generic;
 using System.Linq;
 using System.Threading.Tasks;
 using Dasync.Collections;
+using EpicGames.Horde.Storage;
 using Jupiter.Implementation;
 
 namespace Horde.Storage.Implementation
diff --git a/Engine/Source/Programs/Horde/HordeStorage/Horde.Storage/Implementation/Objects/MongoReferencesStore.cs b/Engine/Source/Programs/Horde/HordeStorage/Horde.Storage/Implementation/Objects/MongoReferencesStore.cs
--- a/Engine/Source/Programs/Horde/HordeStorage/Horde.Storage/Implementation/Objects/MongoReferencesStore.cs
+++ b/Engine/Source/Programs/Horde/HordeStorage/Horde.Storage/Implementation/Objects/MongoReferencesStore.cs
@@ -4,6 +4,7 @@ using System;
 using System.Collections.Generic;
 using System.Threading.Tasks;
 using Dasync.Collections;
+using EpicGames.Horde.Storage;
 using Jupiter.Implementation;
 using Microsoft.Extensions.Options;
 using MongoDB.Bson.Serialization.Attributes;
diff --git a/Engine/Source/Programs/Horde/HordeStorage/Horde.Storage/Implementation/Objects/ScyllaReferencesStore.cs b/Engine/Source/Programs/Horde/HordeStorage/Horde.Storage/Implementation/Objects/ScyllaReferencesStore.cs
--- a/Engine/Source/Programs/Horde/HordeStorage/Horde.Storage/Implementation/Objects/ScyllaReferencesStore.cs
+++ b/Engine/Source/Programs/Horde/HordeStorage/Horde.Storage/Implementation/Objects/ScyllaReferencesStore.cs
@@ -8,6 +8,7 @@ using System.Threading.Tasks;
 using Cassandra;
 using Cassandra.Mapping;
 using Datadog.Trace;
+using EpicGames.Horde.Storage;
 using Jupiter.Implementation;
 using Microsoft.Extensions.Options;
 using Serilog;
diff --git a/Engine/Source/Programs/Horde/HordeStorage/Horde.Storage/Implementation/Refs/CachedRefStore.cs b/Engine/Source/Programs/Horde/HordeStorage/Horde.Storage/Implementation/Refs/CachedRefStore.cs
--- a/Engine/Source/Programs/Horde/HordeStorage/Horde.Storage/Implementation/Refs/CachedRefStore.cs
+++ b/Engine/Source/Programs/Horde/HordeStorage/Horde.Storage/Implementation/Refs/CachedRefStore.cs
@@ -5,6 +5,7 @@ using System.Collections.Concurrent;
 using System.Collections.Generic;
 using System.Threading;
 using System.Threading.Tasks;
+using EpicGames.Horde.Storage;
 using Jupiter.Implementation;
 using Microsoft.Extensions.Caching.Memory;
 using Microsoft.Extensions.Options;
diff --git a/Engine/Source/Programs/Horde/HordeStorage/Horde.Storage/Implementation/Refs/DynamoDbRefsStore.cs b/Engine/Source/Programs/Horde/HordeStorage/Horde.Storage/Implementation/Refs/DynamoDbRefsStore.cs
--- a/Engine/Source/Programs/Horde/HordeStorage/Horde.Storage/Implementation/Refs/DynamoDbRefsStore.cs
+++ b/Engine/Source/Programs/Horde/HordeStorage/Horde.Storage/Implementation/Refs/DynamoDbRefsStore.cs
@@ -9,6 +9,7 @@ using Amazon.DynamoDBv2.DataModel;
 using Amazon.DynamoDBv2.DocumentModel;
 using Amazon.DynamoDBv2.Model;
 using Datadog.Trace;
+using EpicGames.Horde.Storage;
 using Jupiter.Implementation;
 using Microsoft.Extensions.Options;
 using Serilog;
diff --git a/Engine/Source/Programs/Horde/HordeStorage/Horde.Storage/Implementation/Refs/DynamoNamespaceStore.cs b/Engine/Source/Programs/Horde/HordeStorage/Horde.Storage/Implementation/Refs/DynamoNamespaceStore.cs
--- a/Engine/Source/Programs/Horde/HordeStorage/Horde.Storage/Implementation/Refs/DynamoNamespaceStore.cs
+++ b/Engine/Source/Programs/Horde/HordeStorage/Horde.Storage/Implementation/Refs/DynamoNamespaceStore.cs
@@ -7,6 +7,7 @@ using Amazon.DynamoDBv2;
 using Amazon.DynamoDBv2.DataModel;
 using Amazon.DynamoDBv2.DocumentModel;
 using Amazon.DynamoDBv2.Model;
+using EpicGames.Horde.Storage;
 using Jupiter.Implementation;
 using Microsoft.Extensions.Options;
 
diff --git a/Engine/Source/Programs/Horde/HordeStorage/Horde.Storage/Implementation/Refs/IRefsStore.cs b/Engine/Source/Programs/Horde/HordeStorage/Horde.Storage/Implementation/Refs/IRefsStore.cs
--- a/Engine/Source/Programs/Horde/HordeStorage/Horde.Storage/Implementation/Refs/IRefsStore.cs
+++ b/Engine/Source/Programs/Horde/HordeStorage/Horde.Storage/Implementation/Refs/IRefsStore.cs
@@ -3,6 +3,7 @@
 using System;
 using System.Collections.Generic;
 using System.Threading.Tasks;
+using EpicGames.Horde.Storage;
 using Jupiter.Implementation;
 
 namespace Horde.Storage.Implementation
diff --git a/Engine/Source/Programs/Horde/HordeStorage/Horde.Storage/Implementation/Refs/MemoryRefsStore.cs b/Engine/Source/Programs/Horde/HordeStorage/Horde.Storage/Implementation/Refs/MemoryRefsStore.cs
--- a/Engine/Source/Programs/Horde/HordeStorage/Horde.Storage/Implementation/Refs/MemoryRefsStore.cs
+++ b/Engine/Source/Programs/Horde/HordeStorage/Horde.Storage/Implementation/Refs/MemoryRefsStore.cs
@@ -8,8 +8,9 @@ using Jupiter.Implementation;
 using Serilog;
 
 using RefContainer = System.Collections.Concurrent.ConcurrentDictionary<Jupiter.Implementation.KeyId, Horde.Storage.Implementation.RefRecord>;
-using BucketContainer = System.Collections.Concurrent.ConcurrentDictionary<Jupiter.Implementation.BucketId, System.Collections.Concurrent.ConcurrentDictionary<Jupiter.Implementation.KeyId, Horde.Storage.Implementation.RefRecord>>;
-using NamespaceContainer = System.Collections.Concurrent.ConcurrentDictionary<Jupiter.Implementation.NamespaceId, System.Collections.Concurrent.ConcurrentDictionary<Jupiter.Implementation.BucketId, System.Collections.Concurrent.ConcurrentDictionary<Jupiter.Implementation.KeyId, Horde.Storage.Implementation.RefRecord>>>;
+using BucketContainer = System.Collections.Concurrent.ConcurrentDictionary<EpicGames.Horde.Storage.BucketId, System.Collections.Concurrent.ConcurrentDictionary<Jupiter.Implementation.KeyId, Horde.Storage.Implementation.RefRecord>>;
+using NamespaceContainer = System.Collections.Concurrent.ConcurrentDictionary<EpicGames.Horde.Storage.NamespaceId, System.Collections.Concurrent.ConcurrentDictionary<EpicGames.Horde.Storage.BucketId, System.Collections.Concurrent.ConcurrentDictionary<Jupiter.Implementation.KeyId, Horde.Storage.Implementation.RefRecord>>>;
+using EpicGames.Horde.Storage;
 
 namespace Horde.Storage.Implementation
 {
diff --git a/Engine/Source/Programs/Horde/HordeStorage/Horde.Storage/Implementation/Refs/MongoRefsStore.cs b/Engine/Source/Programs/Horde/HordeStorage/Horde.Storage/Implementation/Refs/MongoRefsStore.cs
--- a/Engine/Source/Programs/Horde/HordeStorage/Horde.Storage/Implementation/Refs/MongoRefsStore.cs
+++ b/Engine/Source/Programs/Horde/HordeStorage/Horde.Storage/Implementation/Refs/MongoRefsStore.cs
@@ -5,6 +5,7 @@ using System.Collections.Generic;
 using System.Linq;
 using System.Security.Authentication;
 using System.Threading.Tasks;
+using EpicGames.Horde.Storage;
 using Jupiter.Implementation;
 using Microsoft.Extensions.Options;
 using MongoDB.Bson;
diff --git a/Engine/Source/Programs/Horde/HordeStorage/Horde.Storage/Implementation/Replication/IReplicator.cs b/Engine/Source/Programs/Horde/HordeStorage/Horde.Storage/Implementation/Replication/IReplicator.cs
--- a/Engine/Source/Programs/Horde/HordeStorage/Horde.Storage/Implementation/Replication/IReplicator.cs
+++ b/Engine/Source/Programs/Horde/HordeStorage/Horde.Storage/Implementation/Replication/IReplicator.cs
@@ -2,6 +2,7 @@
 
 using System;
 using System.Threading.Tasks;
+using EpicGames.Horde.Storage;
 using Jupiter.Implementation;
 
 namespace Horde.Storage.Implementation
diff --git a/Engine/Source/Programs/Horde/HordeStorage/Horde.Storage/Implementation/Replication/RefsReplicator.cs b/Engine/Source/Programs/Horde/HordeStorage/Horde.Storage/Implementation/Replication/RefsReplicator.cs
--- a/Engine/Source/Programs/Horde/HordeStorage/Horde.Storage/Implementation/Replication/RefsReplicator.cs
+++ b/Engine/Source/Programs/Horde/HordeStorage/Horde.Storage/Implementation/Replication/RefsReplicator.cs
@@ -13,6 +13,7 @@ using System.Threading;
 using System.Threading.Tasks;
 using Dasync.Collections;
 using Datadog.Trace;
+using EpicGames.Horde.Storage;
 using Horde.Storage.Controllers;
 using Horde.Storage.Implementation.TransactionLog;
 using Jupiter.Common.Implementation;
diff --git a/Engine/Source/Programs/Horde/HordeStorage/Horde.Storage/Implementation/Replication/ReplicationService.cs b/Engine/Source/Programs/Horde/HordeStorage/Horde.Storage/Implementation/Replication/ReplicationService.cs
--- a/Engine/Source/Programs/Horde/HordeStorage/Horde.Storage/Implementation/Replication/ReplicationService.cs
+++ b/Engine/Source/Programs/Horde/HordeStorage/Horde.Storage/Implementation/Replication/ReplicationService.cs
@@ -6,6 +6,7 @@ using System.IO;
 using System.Linq;
 using System.Threading;
 using System.Threading.Tasks;
+using EpicGames.Horde.Storage;
 using Jupiter;
 using Jupiter.Implementation;
 using Microsoft.Extensions.DependencyInjection;
diff --git a/Engine/Source/Programs/Horde/HordeStorage/Horde.Storage/Implementation/Replication/ReplicationSettings.cs b/Engine/Source/Programs/Horde/HordeStorage/Horde.Storage/Implementation/Replication/ReplicationSettings.cs
--- a/Engine/Source/Programs/Horde/HordeStorage/Horde.Storage/Implementation/Replication/ReplicationSettings.cs
+++ b/Engine/Source/Programs/Horde/HordeStorage/Horde.Storage/Implementation/Replication/ReplicationSettings.cs
@@ -2,6 +2,7 @@
 
 using System.Collections.Generic;
 using System.ComponentModel.DataAnnotations;
+using EpicGames.Horde.Storage;
 using Horde.Storage.Implementation;
 using Jupiter.Implementation;
 
diff --git a/Engine/Source/Programs/Horde/HordeStorage/Horde.Storage/Implementation/Replication/Replicator.cs b/Engine/Source/Programs/Horde/HordeStorage/Horde.Storage/Implementation/Replication/Replicator.cs
--- a/Engine/Source/Programs/Horde/HordeStorage/Horde.Storage/Implementation/Replication/Replicator.cs
+++ b/Engine/Source/Programs/Horde/HordeStorage/Horde.Storage/Implementation/Replication/Replicator.cs
@@ -8,6 +8,7 @@ using System.Net;
 using System.Threading;
 using System.Threading.Tasks;
 using Datadog.Trace;
+using EpicGames.Horde.Storage;
 using Jupiter;
 using Jupiter.Implementation;
 using Microsoft.Extensions.Options;
diff --git a/Engine/Source/Programs/Horde/HordeStorage/Horde.Storage/Implementation/TransactionLog/CallistoTransactionLogWriter.cs b/Engine/Source/Programs/Horde/HordeStorage/Horde.Storage/Implementation/TransactionLog/CallistoTransactionLogWriter.cs
--- a/Engine/Source/Programs/Horde/HordeStorage/Horde.Storage/Implementation/TransactionLog/CallistoTransactionLogWriter.cs
+++ b/Engine/Source/Programs/Horde/HordeStorage/Horde.Storage/Implementation/TransactionLog/CallistoTransactionLogWriter.cs
@@ -7,6 +7,7 @@ using Serilog;
 using Microsoft.Extensions.Options;
 using RestSharp;
 using RestSharp.Serializers.NewtonsoftJson;
+using EpicGames.Horde.Storage;
 
 namespace Horde.Storage.Implementation
 {
diff --git a/Engine/Source/Programs/Horde/HordeStorage/Horde.Storage/Implementation/TransactionLog/IReplicationLog.cs b/Engine/Source/Programs/Horde/HordeStorage/Horde.Storage/Implementation/TransactionLog/IReplicationLog.cs
--- a/Engine/Source/Programs/Horde/HordeStorage/Horde.Storage/Implementation/TransactionLog/IReplicationLog.cs
+++ b/Engine/Source/Programs/Horde/HordeStorage/Horde.Storage/Implementation/TransactionLog/IReplicationLog.cs
@@ -4,6 +4,7 @@ using System;
 using System.Collections.Generic;
 using System.Reflection.Metadata;
 using System.Threading.Tasks;
+using EpicGames.Horde.Storage;
 using Horde.Storage.Implementation.TransactionLog;
 using Jupiter.Implementation;
 using Newtonsoft.Json;
diff --git a/Engine/Source/Programs/Horde/HordeStorage/Horde.Storage/Implementation/TransactionLog/ITransactionLogWriter.cs b/Engine/Source/Programs/Horde/HordeStorage/Horde.Storage/Implementation/TransactionLog/ITransactionLogWriter.cs
--- a/Engine/Source/Programs/Horde/HordeStorage/Horde.Storage/Implementation/TransactionLog/ITransactionLogWriter.cs
+++ b/Engine/Source/Programs/Horde/HordeStorage/Horde.Storage/Implementation/TransactionLog/ITransactionLogWriter.cs
@@ -1,6 +1,7 @@
 ﻿// Copyright Epic Games, Inc. All Rights Reserved.
 
 using System.Threading.Tasks;
+using EpicGames.Horde.Storage;
 using Jupiter.Implementation;
 
 namespace Horde.Storage.Implementation
diff --git a/Engine/Source/Programs/Horde/HordeStorage/Horde.Storage/Implementation/TransactionLog/MemoryReplicationLog.cs b/Engine/Source/Programs/Horde/HordeStorage/Horde.Storage/Implementation/TransactionLog/MemoryReplicationLog.cs
--- a/Engine/Source/Programs/Horde/HordeStorage/Horde.Storage/Implementation/TransactionLog/MemoryReplicationLog.cs
+++ b/Engine/Source/Programs/Horde/HordeStorage/Horde.Storage/Implementation/TransactionLog/MemoryReplicationLog.cs
@@ -7,6 +7,7 @@ using System.Linq;
 using System.Threading.Tasks;
 using Cassandra;
 using Dasync.Collections;
+using EpicGames.Horde.Storage;
 using Jupiter.Implementation;
 
 namespace Horde.Storage.Implementation
diff --git a/Engine/Source/Programs/Horde/HordeStorage/Horde.Storage/Implementation/TransactionLog/MemoryTransactionLogWriter.cs b/Engine/Source/Programs/Horde/HordeStorage/Horde.Storage/Implementation/TransactionLog/MemoryTransactionLogWriter.cs
--- a/Engine/Source/Programs/Horde/HordeStorage/Horde.Storage/Implementation/TransactionLog/MemoryTransactionLogWriter.cs
+++ b/Engine/Source/Programs/Horde/HordeStorage/Horde.Storage/Implementation/TransactionLog/MemoryTransactionLogWriter.cs
@@ -2,6 +2,7 @@
 
 using System.Collections.Concurrent;
 using System.Threading.Tasks;
+using EpicGames.Horde.Storage;
 using Jupiter.Implementation;
 
 namespace Horde.Storage.Implementation
diff --git a/Engine/Source/Programs/Horde/HordeStorage/Horde.Storage/Implementation/TransactionLog/ReplicationLogSnapshot.cs b/Engine/Source/Programs/Horde/HordeStorage/Horde.Storage/Implementation/TransactionLog/ReplicationLogSnapshot.cs
--- a/Engine/Source/Programs/Horde/HordeStorage/Horde.Storage/Implementation/TransactionLog/ReplicationLogSnapshot.cs
+++ b/Engine/Source/Programs/Horde/HordeStorage/Horde.Storage/Implementation/TransactionLog/ReplicationLogSnapshot.cs
@@ -6,6 +6,7 @@ using System.IO;
 using System.IO.Compression;
 using System.Linq;
 using System.Threading.Tasks;
+using EpicGames.Horde.Storage;
 using Jupiter.Implementation;
 using Newtonsoft.Json;
 
diff --git a/Engine/Source/Programs/Horde/HordeStorage/Horde.Storage/Implementation/TransactionLog/ReplicationLogSnapshotBuilder.cs b/Engine/Source/Programs/Horde/HordeStorage/Horde.Storage/Implementation/TransactionLog/ReplicationLogSnapshotBuilder.cs
--- a/Engine/Source/Programs/Horde/HordeStorage/Horde.Storage/Implementation/TransactionLog/ReplicationLogSnapshotBuilder.cs
+++ b/Engine/Source/Programs/Horde/HordeStorage/Horde.Storage/Implementation/TransactionLog/ReplicationLogSnapshotBuilder.cs
@@ -6,6 +6,7 @@ using System.IO;
 using System.Linq;
 using System.Threading;
 using System.Threading.Tasks;
+using EpicGames.Horde.Storage;
 using Jupiter.Implementation;
 
 namespace Horde.Storage.Implementation.TransactionLog
diff --git a/Engine/Source/Programs/Horde/HordeStorage/Horde.Storage/Implementation/TransactionLog/ReplicationSnapshotService.cs b/Engine/Source/Programs/Horde/HordeStorage/Horde.Storage/Implementation/TransactionLog/ReplicationSnapshotService.cs
--- a/Engine/Source/Programs/Horde/HordeStorage/Horde.Storage/Implementation/TransactionLog/ReplicationSnapshotService.cs
+++ b/Engine/Source/Programs/Horde/HordeStorage/Horde.Storage/Implementation/TransactionLog/ReplicationSnapshotService.cs
@@ -5,6 +5,7 @@ using System.ComponentModel.DataAnnotations;
 using System.Threading;
 using System.Threading.Tasks;
 using Dasync.Collections;
+using EpicGames.Horde.Storage;
 using Horde.Storage.Implementation.TransactionLog;
 using Jupiter;
 using Jupiter.Common;
diff --git a/Engine/Source/Programs/Horde/HordeStorage/Horde.Storage/Implementation/TransactionLog/ScyllaReplicationLog.cs b/Engine/Source/Programs/Horde/HordeStorage/Horde.Storage/Implementation/TransactionLog/ScyllaReplicationLog.cs
--- a/Engine/Source/Programs/Horde/HordeStorage/Horde.Storage/Implementation/TransactionLog/ScyllaReplicationLog.cs
+++ b/Engine/Source/Programs/Horde/HordeStorage/Horde.Storage/Implementation/TransactionLog/ScyllaReplicationLog.cs
@@ -8,6 +8,7 @@ using Cassandra;
 using Cassandra.Mapping;
 using Dasync.Collections;
 using Datadog.Trace;
+using EpicGames.Horde.Storage;
 using Jupiter.Implementation;
 using Microsoft.Extensions.Options;
 
diff --git a/Engine/Source/Programs/Horde/HordeStorage/Horde.Storage/Tests/Functional/Blob/BlobIndexTests.cs b/Engine/Source/Programs/Horde/HordeStorage/Horde.Storage/Tests/Functional/Blob/BlobIndexTests.cs
--- a/Engine/Source/Programs/Horde/HordeStorage/Horde.Storage/Tests/Functional/Blob/BlobIndexTests.cs
+++ b/Engine/Source/Programs/Horde/HordeStorage/Horde.Storage/Tests/Functional/Blob/BlobIndexTests.cs
@@ -9,6 +9,7 @@ using System.Net.Mime;
 using System.Text;
 using System.Threading.Tasks;
 using Cassandra;
+using EpicGames.Horde.Storage;
 using Horde.Storage.Implementation;
 using Horde.Storage.Implementation.Blob;
 using Jupiter;
diff --git a/Engine/Source/Programs/Horde/HordeStorage/Horde.Storage/Tests/Functional/GC/GCBlobTests.cs b/Engine/Source/Programs/Horde/HordeStorage/Horde.Storage/Tests/Functional/GC/GCBlobTests.cs
--- a/Engine/Source/Programs/Horde/HordeStorage/Horde.Storage/Tests/Functional/GC/GCBlobTests.cs
+++ b/Engine/Source/Programs/Horde/HordeStorage/Horde.Storage/Tests/Functional/GC/GCBlobTests.cs
@@ -21,6 +21,7 @@ using RestSharp;
 using Serilog;
 using Serilog.Core;
 using Moq;
+using EpicGames.Horde.Storage;
 
 namespace Horde.Storage.FunctionalTests.GC
 {
diff --git a/Engine/Source/Programs/Horde/HordeStorage/Horde.Storage/Tests/Functional/GC/GCRefTests.cs b/Engine/Source/Programs/Horde/HordeStorage/Horde.Storage/Tests/Functional/GC/GCRefTests.cs
--- a/Engine/Source/Programs/Horde/HordeStorage/Horde.Storage/Tests/Functional/GC/GCRefTests.cs
+++ b/Engine/Source/Programs/Horde/HordeStorage/Horde.Storage/Tests/Functional/GC/GCRefTests.cs
@@ -18,6 +18,7 @@ using Microsoft.VisualStudio.TestTools.UnitTesting;
 using Moq;
 using Serilog;
 using Serilog.Core;
+using EpicGames.Horde.Storage;
 
 namespace Horde.Storage.FunctionalTests.GC
 {
diff --git a/Engine/Source/Programs/Horde/HordeStorage/Horde.Storage/Tests/Functional/Objects/RefReplicatorTests.cs b/Engine/Source/Programs/Horde/HordeStorage/Horde.Storage/Tests/Functional/Objects/RefReplicatorTests.cs
--- a/Engine/Source/Programs/Horde/HordeStorage/Horde.Storage/Tests/Functional/Objects/RefReplicatorTests.cs
+++ b/Engine/Source/Programs/Horde/HordeStorage/Horde.Storage/Tests/Functional/Objects/RefReplicatorTests.cs
@@ -25,6 +25,7 @@ using Moq.Contrib.HttpClient;
 using Newtonsoft.Json;
 using Serilog;
 using Logger = Serilog.Core.Logger;
+using EpicGames.Horde.Storage;
 
 namespace Horde.Storage.FunctionalTests.Replication
 {
diff --git a/Engine/Source/Programs/Horde/HordeStorage/Horde.Storage/Tests/Functional/Objects/ReferencesTests.cs b/Engine/Source/Programs/Horde/HordeStorage/Horde.Storage/Tests/Functional/Objects/ReferencesTests.cs
--- a/Engine/Source/Programs/Horde/HordeStorage/Horde.Storage/Tests/Functional/Objects/ReferencesTests.cs
+++ b/Engine/Source/Programs/Horde/HordeStorage/Horde.Storage/Tests/Functional/Objects/ReferencesTests.cs
@@ -26,6 +26,7 @@ using Newtonsoft.Json;
 using Newtonsoft.Json.Linq;
 using Serilog;
 using Logger = Serilog.Core.Logger;
+using EpicGames.Horde.Storage;
 
 namespace Horde.Storage.FunctionalTests.References
 {
diff --git a/Engine/Source/Programs/Horde/HordeStorage/Horde.Storage/Tests/Functional/Objects/ReplicationTests.cs b/Engine/Source/Programs/Horde/HordeStorage/Horde.Storage/Tests/Functional/Objects/ReplicationTests.cs
--- a/Engine/Source/Programs/Horde/HordeStorage/Horde.Storage/Tests/Functional/Objects/ReplicationTests.cs
+++ b/Engine/Source/Programs/Horde/HordeStorage/Horde.Storage/Tests/Functional/Objects/ReplicationTests.cs
@@ -25,6 +25,7 @@ using Microsoft.Extensions.DependencyInjection;
 using Microsoft.VisualStudio.TestTools.UnitTesting;
 using Serilog;
 using Logger = Serilog.Core.Logger;
+using EpicGames.Horde.Storage;
 
 namespace Horde.Storage.FunctionalTests.References
 {
diff --git a/Engine/Source/Programs/Horde/HordeStorage/Horde.Storage/Tests/Functional/RefTests.cs b/Engine/Source/Programs/Horde/HordeStorage/Horde.Storage/Tests/Functional/RefTests.cs
--- a/Engine/Source/Programs/Horde/HordeStorage/Horde.Storage/Tests/Functional/RefTests.cs
+++ b/Engine/Source/Programs/Horde/HordeStorage/Horde.Storage/Tests/Functional/RefTests.cs
@@ -28,6 +28,7 @@ using Newtonsoft.Json;
 using Newtonsoft.Json.Linq;
 using Serilog;
 using Serilog.Core;
+using EpicGames.Horde.Storage;
 
 namespace Horde.Storage.FunctionalTests.Ref
 {
diff --git a/Engine/Source/Programs/Horde/HordeStorage/Horde.Storage/Tests/Functional/StorageTests.cs b/Engine/Source/Programs/Horde/HordeStorage/Horde.Storage/Tests/Functional/StorageTests.cs
--- a/Engine/Source/Programs/Horde/HordeStorage/Horde.Storage/Tests/Functional/StorageTests.cs
+++ b/Engine/Source/Programs/Horde/HordeStorage/Horde.Storage/Tests/Functional/StorageTests.cs
@@ -32,6 +32,7 @@ using Newtonsoft.Json;
 using Newtonsoft.Json.Linq;
 using Serilog;
 using Serilog.Core;
+using EpicGames.Horde.Storage;
 
 namespace Horde.Storage.FunctionalTests.Storage
 {
diff --git a/Engine/Source/Programs/Horde/HordeStorage/Horde.Storage/Tests/UnitTests/HierarchicalBlobStorageTest.cs b/Engine/Source/Programs/Horde/HordeStorage/Horde.Storage/Tests/UnitTests/HierarchicalBlobStorageTest.cs
--- a/Engine/Source/Programs/Horde/HordeStorage/Horde.Storage/Tests/UnitTests/HierarchicalBlobStorageTest.cs
+++ b/Engine/Source/Programs/Horde/HordeStorage/Horde.Storage/Tests/UnitTests/HierarchicalBlobStorageTest.cs
@@ -7,6 +7,7 @@ using System.Linq;
 using System.Text;
 using System.Threading.Tasks;
 using Dasync.Collections;
+using EpicGames.Horde.Storage;
 using Horde.Storage.Implementation;
 using Horde.Storage.Implementation.Blob;
 using Jupiter.Implementation;
@@ -16,6 +17,8 @@ using Moq;
 
 namespace Horde.Storage.UnitTests
 {
+    using BlobNotFoundException = Horde.Storage.Implementation.BlobNotFoundException;
+
     [TestClass]
     public class HierarchicalBlobStorageTest
     {
diff --git a/Engine/Source/Programs/Horde/HordeStorage/Horde.Storage/Tests/UnitTests/KestrelDispatcherTest.cs b/Engine/Source/Programs/Horde/HordeStorage/Horde.Storage/Tests/UnitTests/KestrelDispatcherTest.cs
--- a/Engine/Source/Programs/Horde/HordeStorage/Horde.Storage/Tests/UnitTests/KestrelDispatcherTest.cs
+++ b/Engine/Source/Programs/Horde/HordeStorage/Horde.Storage/Tests/UnitTests/KestrelDispatcherTest.cs
@@ -6,6 +6,7 @@ using System.IO;
 using System.IO.Pipelines;
 using System.Text;
 using System.Threading.Tasks;
+using EpicGames.Horde.Storage;
 using Horde.Storage.Implementation;
 using Horde.Storage.Implementation.Kestrel;
 using Jupiter.Implementation;
diff --git a/Engine/Source/Programs/Horde/HordeStorage/Horde.Storage/Tests/UnitTests/MemoryCacheBlobStoreTest.cs b/Engine/Source/Programs/Horde/HordeStorage/Horde.Storage/Tests/UnitTests/MemoryCacheBlobStoreTest.cs
--- a/Engine/Source/Programs/Horde/HordeStorage/Horde.Storage/Tests/UnitTests/MemoryCacheBlobStoreTest.cs
+++ b/Engine/Source/Programs/Horde/HordeStorage/Horde.Storage/Tests/UnitTests/MemoryCacheBlobStoreTest.cs
@@ -4,6 +4,7 @@ using System;
 using System.IO;
 using System.Text;
 using System.Threading.Tasks;
+using EpicGames.Horde.Storage;
 using Horde.Storage.Implementation;
 using Jupiter.Implementation;
 using Microsoft.Extensions.Options;
diff --git a/Engine/Source/Programs/Horde/HordeStorage/Horde.Storage/Tests/UnitTests/ReplicatorTests.cs b/Engine/Source/Programs/Horde/HordeStorage/Horde.Storage/Tests/UnitTests/ReplicatorTests.cs
--- a/Engine/Source/Programs/Horde/HordeStorage/Horde.Storage/Tests/UnitTests/ReplicatorTests.cs
+++ b/Engine/Source/Programs/Horde/HordeStorage/Horde.Storage/Tests/UnitTests/ReplicatorTests.cs
@@ -16,6 +16,7 @@ using Jupiter.Implementation;
 using Microsoft.Extensions.Options;
 using Moq;
 using RestSharp;
+using EpicGames.Horde.Storage;
 
 namespace Horde.Storage.UnitTests
 {
diff --git a/Engine/Source/Programs/Horde/HordeStorage/Horde.Storage/Tests/UnitTests/S3StoreTests.cs b/Engine/Source/Programs/Horde/HordeStorage/Horde.Storage/Tests/UnitTests/S3StoreTests.cs
--- a/Engine/Source/Programs/Horde/HordeStorage/Horde.Storage/Tests/UnitTests/S3StoreTests.cs
+++ b/Engine/Source/Programs/Horde/HordeStorage/Horde.Storage/Tests/UnitTests/S3StoreTests.cs
@@ -7,6 +7,7 @@ using System.Text;
 using System.Threading.Tasks;
 using Amazon.S3;
 using Amazon.S3.Model;
+using EpicGames.Horde.Storage;
 using Horde.Storage.Implementation;
 using Horde.Storage.Implementation.Blob;
 using Jupiter;
diff --git a/Engine/Source/Programs/Horde/HordeStorage/Jupiter.Common/Authentication/NamespaceAuthorizationHandler.cs b/Engine/Source/Programs/Horde/HordeStorage/Jupiter.Common/Authentication/NamespaceAuthorizationHandler.cs
--- a/Engine/Source/Programs/Horde/HordeStorage/Jupiter.Common/Authentication/NamespaceAuthorizationHandler.cs
+++ b/Engine/Source/Programs/Horde/HordeStorage/Jupiter.Common/Authentication/NamespaceAuthorizationHandler.cs
@@ -4,6 +4,7 @@ using System;
 using System.Collections.Generic;
 using System.ComponentModel.DataAnnotations;
 using System.Threading.Tasks;
+using EpicGames.Horde.Storage;
 using Jupiter.Common;
 using Jupiter.Implementation;
 using Microsoft.AspNetCore.Authorization;
diff --git a/Engine/Source/Programs/Horde/HordeStorage/Jupiter.Common/Implementation/BatchWriter.cs b/Engine/Source/Programs/Horde/HordeStorage/Jupiter.Common/Implementation/BatchWriter.cs
--- a/Engine/Source/Programs/Horde/HordeStorage/Jupiter.Common/Implementation/BatchWriter.cs
+++ b/Engine/Source/Programs/Horde/HordeStorage/Jupiter.Common/Implementation/BatchWriter.cs
@@ -7,6 +7,7 @@ using System.Text;
 using System.Threading;
 using System.Threading.Tasks;
 using Datadog.Trace;
+using EpicGames.Horde.Storage;
 using Serilog;
 
 namespace Jupiter.Implementation
diff --git a/Engine/Source/Programs/Horde/HordeStorage/Jupiter.Common/Implementation/BucketId.cs b/Engine/Source/Programs/Horde/HordeStorage/Jupiter.Common/Implementation/BucketId.cs
deleted file mode 100644
diff --git a/Engine/Source/Programs/Horde/HordeStorage/Jupiter.Common/Implementation/CallistoReader.cs b/Engine/Source/Programs/Horde/HordeStorage/Jupiter.Common/Implementation/CallistoReader.cs
--- a/Engine/Source/Programs/Horde/HordeStorage/Jupiter.Common/Implementation/CallistoReader.cs
+++ b/Engine/Source/Programs/Horde/HordeStorage/Jupiter.Common/Implementation/CallistoReader.cs
@@ -7,6 +7,7 @@ using System.Net;
 using System.Runtime.CompilerServices;
 using System.Threading;
 using System.Threading.Tasks;
+using EpicGames.Horde.Storage;
 using Newtonsoft.Json;
 using Newtonsoft.Json.Linq;
 using RestSharp;
diff --git a/Engine/Source/Programs/Horde/HordeStorage/Jupiter.Common/Implementation/NamespaceId.cs b/Engine/Source/Programs/Horde/HordeStorage/Jupiter.Common/Implementation/NamespaceId.cs
deleted file mode 100644
diff --git a/Engine/Source/Programs/Horde/HordeStorage/Jupiter.Common/Implementation/NamespacePolicyResolver.cs b/Engine/Source/Programs/Horde/HordeStorage/Jupiter.Common/Implementation/NamespacePolicyResolver.cs
--- a/Engine/Source/Programs/Horde/HordeStorage/Jupiter.Common/Implementation/NamespacePolicyResolver.cs
+++ b/Engine/Source/Programs/Horde/HordeStorage/Jupiter.Common/Implementation/NamespacePolicyResolver.cs
@@ -1,6 +1,7 @@
 // Copyright Epic Games, Inc. All Rights Reserved.
 
 using System;
+using EpicGames.Horde.Storage;
 using Jupiter.Implementation;
 using Microsoft.Extensions.Options;
 
diff --git a/Engine/Source/Programs/Horde/HordeStorage/Jupiter.Common/Jupiter.Common.csproj b/Engine/Source/Programs/Horde/HordeStorage/Jupiter.Common/Jupiter.Common.csproj
--- a/Engine/Source/Programs/Horde/HordeStorage/Jupiter.Common/Jupiter.Common.csproj
+++ b/Engine/Source/Programs/Horde/HordeStorage/Jupiter.Common/Jupiter.Common.csproj
@@ -17,7 +17,7 @@
     <PackageReference Include="AWSSDK.Extensions.NETCore.Setup" Version="3.7.0.1" />
     <PackageReference Include="AWSSDK.SecretsManager" Version="3.7.1.39" />
     <PackageReference Include="AWSSDK.SecurityToken" Version="3.7.1.4" />
-    <PackageReference Include="Blake3" Version="0.2.1" />
+    <PackageReference Include="Blake3" Version="0.3.0" />
     <PackageReference Include="Datadog.Trace" Version="2.1.0" />
     <PackageReference Include="IdentityServer4.AccessTokenValidation" Version="3.0.1" />
     <PackageReference Include="JsonSubTypes" Version="1.8.0" />
@@ -51,5 +51,9 @@
     <None Include="Shared\Oodle\2.9.3\win\redist\oo2core_9_win32.dll" Link="oo2core_9_win32.dll" CopyToOutputDirectory="PreserveNewest" />
   </ItemGroup>
 
+  <ItemGroup>
+    <ProjectReference Include="..\..\..\Shared\EpicGames.Horde\EpicGames.Horde.csproj" />
+  </ItemGroup>
+
 </Project>
  
\ No newline at end of file
