Horde: Suppress datadog health check errors when running locally.

diff --git a/Engine/Source/Programs/Horde/HordeStorage/Jupiter.Common/BaseStartup.cs b/Engine/Source/Programs/Horde/HordeStorage/Jupiter.Common/BaseStartup.cs
--- a/Engine/Source/Programs/Horde/HordeStorage/Jupiter.Common/BaseStartup.cs
+++ b/Engine/Source/Programs/Horde/HordeStorage/Jupiter.Common/BaseStartup.cs
@@ -209,10 +209,14 @@ namespace Jupiter
         private void OnAddHealthChecks(IServiceCollection services)
         {
             IHealthChecksBuilder healthChecks = services.AddHealthChecks()
-                .AddCheck("self", () => HealthCheckResult.Healthy(), tags: new[] {"self"});
+                .AddCheck("self", () => HealthCheckResult.Healthy(), tags: new[] { "self" });
             OnAddHealthChecks(services, healthChecks);
 
-            healthChecks.AddDatadogPublisher("jupiter.healthchecks");
+            string? ddAgentHost = System.Environment.GetEnvironmentVariable("DD_AGENT_HOST");
+            if (!String.IsNullOrEmpty(ddAgentHost))
+            {
+                healthChecks.AddDatadogPublisher("jupiter.healthchecks");
+            }
         }
 
         /// <summary>
