Horde: Capture up to max warning/error log events
#skipci
#rnx

diff --git a/Engine/Source/Programs/Horde/Horde.Agent/Utility/JsonRpcLogger.cs b/Engine/Source/Programs/Horde/Horde.Agent/Utility/JsonRpcLogger.cs
--- a/Engine/Source/Programs/Horde/Horde.Agent/Utility/JsonRpcLogger.cs
+++ b/Engine/Source/Programs/Horde/Horde.Agent/Utility/JsonRpcLogger.cs
@@ -215,7 +215,7 @@ namespace HordeAgent.Parser
 					{
 						if (Data.Event != null)
 						{
-							if ((Data.Event.Severity == EventSeverity.Warning && ++NumWarnings < MaxWarnings) || (Data.Event.Severity == EventSeverity.Error && ++NumErrors < MaxErrors))
+							if ((Data.Event.Severity == EventSeverity.Warning && ++NumWarnings <= MaxWarnings) || (Data.Event.Severity == EventSeverity.Error && ++NumErrors <= MaxErrors))
 							{
 								Data.Event.LineIndex = LineIndex;
 								Events.Add(Data.Event);
