Horde: Add systemic issue handler
#skipci
#rnx

diff --git a/Engine/Source/Programs/Horde/Horde.Agent.Tests/LogParserTests.cs b/Engine/Source/Programs/Horde/Horde.Agent.Tests/LogParserTests.cs
--- a/Engine/Source/Programs/Horde/Horde.Agent.Tests/LogParserTests.cs
+++ b/Engine/Source/Programs/Horde/Horde.Agent.Tests/LogParserTests.cs
@@ -126,7 +126,7 @@ namespace HordeAgentTests
 			};
 
 			List<LogEvent> Events = Parse(Lines);
-			CheckEventGroup(Events, 1, 3, LogLevel.Error, KnownLogEvents.Generic);
+			CheckEventGroup(Events, 1, 3, LogLevel.Error, KnownLogEvents.ExitCode);
 		}
 
 		[TestMethod]
@@ -964,13 +964,15 @@ namespace HordeAgentTests
 				@"    Tasks execution is impeded due to low agent responsiveness",
 				@"-------------------------------------------------------------------------------",
 				@"    LogXGEController: Warning: XGE's background service (BuildService.exe) is not running - service is likely disabled on this machine.",
+				@"BUILD FAILED: Command failed (Result:1): C:\Program Files (x86)\IncrediBuild\xgConsole.exe ""d:\build\Sync\Engine\Programs\AutomationTool\Saved\Logs\UAT_XGE.xml"" /Rebuild /NoLogo /ShowAgent /ShowTime"
 			};
 
 			List<LogEvent> Events = Parse(Lines);
-			Assert.AreEqual(7, Events.Count);
+			Assert.AreEqual(8, Events.Count);
 			CheckEventGroup(Events.Slice(0, 3), 0, 3, LogLevel.Information, KnownLogEvents.Systemic_Xge_Standalone);
 			CheckEventGroup(Events.Slice(3, 3), 3, 3, LogLevel.Information, KnownLogEvents.Systemic_Xge);
 			CheckEventGroup(Events.Slice(6, 1), 7, 1, LogLevel.Information, KnownLogEvents.Systemic_Xge_ServiceNotRunning);
+			CheckEventGroup(Events.Slice(7, 1), 8, 1, LogLevel.Information, KnownLogEvents.Systemic_Xge_BuildFailed);
 		}
 
 		string GetSubProperty(LogEvent Event, string SpanName, string Name)
diff --git a/Engine/Source/Programs/Horde/Horde.Agent/Parser/Matchers/ExitCodeEventMatcher.cs b/Engine/Source/Programs/Horde/Horde.Agent/Parser/Matchers/ExitCodeEventMatcher.cs
--- a/Engine/Source/Programs/Horde/Horde.Agent/Parser/Matchers/ExitCodeEventMatcher.cs
+++ b/Engine/Source/Programs/Horde/Horde.Agent/Parser/Matchers/ExitCodeEventMatcher.cs
@@ -39,7 +39,7 @@ namespace HordeAgent.Parser.Matchers
 			{
 				LogEventBuilder Builder = new LogEventBuilder(Cursor);
 				Builder.MoveNext(NumLines - 1);
-				return Builder.ToMatch(LogEventPriority.Low, LogLevel.Error, KnownLogEvents.Generic);
+				return Builder.ToMatch(LogEventPriority.Low, LogLevel.Error, KnownLogEvents.ExitCode);
 			}
 			return null;
 		}
diff --git a/Engine/Source/Programs/Horde/Horde.Agent/Parser/Matchers/XoreaxEventMatcher.cs b/Engine/Source/Programs/Horde/Horde.Agent/Parser/Matchers/XoreaxEventMatcher.cs
--- a/Engine/Source/Programs/Horde/Horde.Agent/Parser/Matchers/XoreaxEventMatcher.cs
+++ b/Engine/Source/Programs/Horde/Horde.Agent/Parser/Matchers/XoreaxEventMatcher.cs
@@ -24,6 +24,12 @@ namespace HordeAgent.Parser.Matchers
 				return Builder.ToMatch(LogEventPriority.High, LogLevel.Information, KnownLogEvents.Systemic_Xge_ServiceNotRunning);
 			}
 
+			if (Cursor.IsMatch(@"BUILD FAILED: (.*)xgConsole\.exe(.*)"))
+			{
+				LogEventBuilder Builder = new LogEventBuilder(Cursor);
+				return Builder.ToMatch(LogEventPriority.High, LogLevel.Information, KnownLogEvents.Systemic_Xge_BuildFailed);
+			}
+
 			if (Cursor.IsMatch(@"^\s*--------------------Build System Warning[- ]"))
 			{
 				LogEventBuilder Builder = new LogEventBuilder(Cursor);
diff --git a/Engine/Source/Programs/Horde/Horde.Build/Issues/Impl/Handlers/SystemicIssueHandler.cs b/Engine/Source/Programs/Horde/Horde.Build/Issues/Impl/Handlers/SystemicIssueHandler.cs
new file mode 100644
--- /dev/null
+++ b/Engine/Source/Programs/Horde/Horde.Build/Issues/Impl/Handlers/SystemicIssueHandler.cs
@@ -0,0 +1,82 @@
+// Copyright Epic Games, Inc. All Rights Reserved.
+
+using HordeCommon;
+using HordeServer.Collections;
+using HordeCommon.Rpc;
+using HordeServer.Services;
+using HordeServer.Utilities;
+using Microsoft.Extensions.Logging;
+using MongoDB.Bson;
+using MongoDB.Bson.Serialization;
+using MongoDB.Driver;
+using System;
+using System.Collections.Generic;
+using System.Diagnostics.CodeAnalysis;
+using System.IO;
+using System.Linq;
+using System.Security.Policy;
+using System.Text;
+using System.Threading.Tasks;
+using System.Text.Json;
+using HordeServer.Models;
+using EpicGames.Core;
+
+namespace HordeServer.IssueHandlers.Impl
+{
+	/// <summary>
+	/// Instance of a particular systemic error
+	/// </summary>
+	class SystemicIssueHandler : IIssueHandler
+	{
+		/// <inheritdoc/>
+		public string Type => "Systemic";
+
+		/// <inheritdoc/>
+		public int Priority => 10;
+
+		/// <summary>
+		///  Known systemic errors
+		/// </summary>
+		static HashSet<EventId?> KnownSystemic = new HashSet<EventId?> {  KnownLogEvents.Systemic, KnownLogEvents.Systemic_Xge, KnownLogEvents.Systemic_Xge_Standalone, 
+																		KnownLogEvents.Systemic_Xge_ServiceNotRunning, KnownLogEvents.Systemic_Xge_BuildFailed, 
+																		KnownLogEvents.Systemic_SlowDDC, KnownLogEvents.Systemic_Horde, KnownLogEvents.Systemic_Horde_ArtifactUpload,
+																		KnownLogEvents.Horde, KnownLogEvents.Horde_InvalidPreflight, KnownLogEvents.ExitCode};
+
+		/// <summary>
+		/// Determines if the given event id matches
+		/// </summary>
+		/// <param name="EventId">The event id to compare</param>
+		/// <returns>True if the given event id matches</returns>
+		public static bool IsMatchingEventId(EventId? EventId)
+		{
+			return KnownSystemic.Contains(EventId);
+		}
+
+		/// <inheritdoc/>
+		public void RankSuspects(IIssueFingerprint Fingerprint, List<SuspectChange> Suspects)
+		{
+		}
+
+		/// <inheritdoc/>
+		public bool TryGetFingerprint(IJob Job, INode Node, ILogEventData EventData, [NotNullWhen(true)] out NewIssueFingerprint? Fingerprint)
+		{
+			if (!IsMatchingEventId(EventData.EventId))
+			{
+				Fingerprint = null;
+				return false;
+			}
+
+			Fingerprint = new NewIssueFingerprint(Type, new[] { Node.Name }, null);
+			return true;
+		}
+
+		/// <inheritdoc/>
+		public string GetSummary(IIssueFingerprint Fingerprint, IssueSeverity Severity)
+		{
+			string Type = (Severity == IssueSeverity.Warning) ? "Systemic warnings" : "Systemic errors";
+			string NodeName = Fingerprint.Keys.FirstOrDefault() ?? "(unknown)";
+			return $"{Type} in {NodeName}";
+		}
+
+	}
+}
diff --git a/Engine/Source/Programs/Horde/Horde.Build/Issues/Impl/IssueService.cs b/Engine/Source/Programs/Horde/Horde.Build/Issues/Impl/IssueService.cs
--- a/Engine/Source/Programs/Horde/Horde.Build/Issues/Impl/IssueService.cs
+++ b/Engine/Source/Programs/Horde/Horde.Build/Issues/Impl/IssueService.cs
@@ -681,6 +681,17 @@ namespace HordeServer.Services.Impl
 					EventGroup.Events.Add(new NewEvent(StepEvent, StepEventData));
 				}
 			}
+			
+			List<NewIssueFingerprint> SystemicFingerprints = FingerprintToEventGroup.Keys.Where(x => x.Type == "Systemic").ToList();
+
+			// Only generate user issues for systemic issues where there are no other failures, otherwise this leads to many false positives (Experimental)
+			if ((SystemicFingerprints.Count != 0) && (SystemicFingerprints.Count != FingerprintToEventGroup.Values.Count))
+			{
+				for (int i = 0; i < SystemicFingerprints.Count; i++)
+				{
+					FingerprintToEventGroup.Remove(SystemicFingerprints[i]);
+				}
+			}
 
 			// Print the list of new events
 			HashSet<NewEventGroup> EventGroups = new HashSet<NewEventGroup>(FingerprintToEventGroup.Values);
diff --git a/Engine/Source/Programs/Shared/EpicGames.Core/KnownLogEvents.cs b/Engine/Source/Programs/Shared/EpicGames.Core/KnownLogEvents.cs
--- a/Engine/Source/Programs/Shared/EpicGames.Core/KnownLogEvents.cs
+++ b/Engine/Source/Programs/Shared/EpicGames.Core/KnownLogEvents.cs
@@ -29,6 +29,11 @@ namespace EpicGames.Core
 		/// </summary>
 		public static EventId Exception { get; } = new EventId(2);
 
+		/// <summary>
+		/// Exit code error
+		/// </summary>
+		public static EventId ExitCode { get; } = new EventId(3);
+
 		/// <summary>
 		/// Compiler error
 		/// </summary>
@@ -139,6 +144,11 @@ namespace EpicGames.Core
 		/// </summary>
 		public static EventId Systemic_Xge_ServiceNotRunning { get; } = new EventId(712);
 
+		/// <summary>
+		/// General build failed error
+		/// </summary>
+		public static EventId Systemic_Xge_BuildFailed { get; } = new EventId(713);
+
 		/// <summary>
 		/// DDC is slow
 		/// </summary>
