Horde: Audit fingerprint changes
#skipci
#rnx

diff --git a/Engine/Source/Programs/Horde/Horde.Build/Issues/Impl/IssueCollection.cs b/Engine/Source/Programs/Horde/Horde.Build/Issues/Impl/IssueCollection.cs
--- a/Engine/Source/Programs/Horde/Horde.Build/Issues/Impl/IssueCollection.cs
+++ b/Engine/Source/Programs/Horde/Horde.Build/Issues/Impl/IssueCollection.cs
@@ -137,6 +137,21 @@ namespace HordeServer.Collections.Impl
 				}
 				return null;
 			}
+
+			public string FingerprintsDesc
+			{
+				get
+				{
+					if (Fingerprints == null || Fingerprints.Count == 0)
+					{
+						return string.Empty;
+					}
+
+					return string.Join(", ", Fingerprints.Select(x => {
+					   return $"(Type: {x.Type} / Keys: {string.Join(", ", x.Keys)} / RejectKeys: {string.Join(", ", x.RejectKeys ?? new CaseInsensitiveStringSet(new string[] {"No Reject Keys"}))})";
+				   }));
+				}
+			}
 		}
 
 		class IssueStream : IIssueStream
@@ -533,6 +548,14 @@ namespace HordeServer.Collections.Impl
 				}
 			}
 
+			string OldFingerprints = OldIssue.FingerprintsDesc;
+			string NewFingerprints = NewIssue.FingerprintsDesc;
+			if (OldFingerprints != NewFingerprints)
+			{
+				IssueLogger.LogInformation("Fingerprints changed {Fingerprints}", NewFingerprints);
+			}
+
+
 			HashSet<StreamId> OldFixStreams = new HashSet<StreamId>(OldIssue.Streams.Where(x => x.ContainsFix ?? false).Select(x => x.StreamId));
 			HashSet<StreamId> NewFixStreams = new HashSet<StreamId>(NewIssue.Streams.Where(x => x.ContainsFix ?? false).Select(x => x.StreamId));
 			foreach (StreamId StreamId in NewFixStreams.Where(x => !OldFixStreams.Contains(x)))
