Horde Storage - Fixed up callisto which was broken due to NamespaceId definition changing in CL 18793433
This will fix the broken CI builds.

diff --git a/Engine/Source/Programs/Horde/HordeStorage/Callisto/Controllers/TransactionLogController.cs b/Engine/Source/Programs/Horde/HordeStorage/Callisto/Controllers/TransactionLogController.cs
--- a/Engine/Source/Programs/Horde/HordeStorage/Callisto/Controllers/TransactionLogController.cs
+++ b/Engine/Source/Programs/Horde/HordeStorage/Callisto/Controllers/TransactionLogController.cs
@@ -6,6 +6,7 @@ using System.IO;
 using System.Linq;
 using System.Threading.Tasks;
 using Callisto.Implementation;
+using EpicGames.Horde.Storage;
 using Jupiter;
 using Jupiter.Implementation;
 using Microsoft.AspNetCore.Authorization;
diff --git a/Engine/Source/Programs/Horde/HordeStorage/Callisto/Implementation/FileTransactionLog.cs b/Engine/Source/Programs/Horde/HordeStorage/Callisto/Implementation/FileTransactionLog.cs
--- a/Engine/Source/Programs/Horde/HordeStorage/Callisto/Implementation/FileTransactionLog.cs
+++ b/Engine/Source/Programs/Horde/HordeStorage/Callisto/Implementation/FileTransactionLog.cs
@@ -10,6 +10,7 @@ using System.Security.Cryptography;
 using System.Text;
 using System.Threading.Tasks;
 using Datadog.Trace;
+using EpicGames.Horde.Storage;
 using Jupiter;
 using Jupiter.Implementation;
 using Microsoft.Extensions.Options;
diff --git a/Engine/Source/Programs/Horde/HordeStorage/Callisto/Implementation/FileTransactionLogProxy.cs b/Engine/Source/Programs/Horde/HordeStorage/Callisto/Implementation/FileTransactionLogProxy.cs
--- a/Engine/Source/Programs/Horde/HordeStorage/Callisto/Implementation/FileTransactionLogProxy.cs
+++ b/Engine/Source/Programs/Horde/HordeStorage/Callisto/Implementation/FileTransactionLogProxy.cs
@@ -4,6 +4,7 @@ using System;
 using System.Collections.Concurrent;
 using System.IO;
 using System.Linq;
+using EpicGames.Horde.Storage;
 using Jupiter.Implementation;
 using Microsoft.Extensions.Options;
 
diff --git a/Engine/Source/Programs/Horde/HordeStorage/Callisto/Implementation/ITransactionLog.cs b/Engine/Source/Programs/Horde/HordeStorage/Callisto/Implementation/ITransactionLog.cs
--- a/Engine/Source/Programs/Horde/HordeStorage/Callisto/Implementation/ITransactionLog.cs
+++ b/Engine/Source/Programs/Horde/HordeStorage/Callisto/Implementation/ITransactionLog.cs
@@ -3,6 +3,7 @@
 using System;
 using System.Collections.Generic;
 using System.Threading.Tasks;
+using EpicGames.Horde.Storage;
 using Jupiter.Implementation;
 
 namespace Callisto.Implementation
diff --git a/Engine/Source/Programs/Horde/HordeStorage/Callisto/Implementation/MemoryTransactionLog.cs b/Engine/Source/Programs/Horde/HordeStorage/Callisto/Implementation/MemoryTransactionLog.cs
--- a/Engine/Source/Programs/Horde/HordeStorage/Callisto/Implementation/MemoryTransactionLog.cs
+++ b/Engine/Source/Programs/Horde/HordeStorage/Callisto/Implementation/MemoryTransactionLog.cs
@@ -5,6 +5,7 @@ using System.Collections.Concurrent;
 using System.Collections.Generic;
 using System.Linq;
 using System.Threading.Tasks;
+using EpicGames.Horde.Storage;
 using Jupiter.Implementation;
 
 namespace Callisto.Implementation
diff --git a/Engine/Source/Programs/Horde/HordeStorage/Callisto/Tests/CallisoFunctionalTests/CallistoTests.cs b/Engine/Source/Programs/Horde/HordeStorage/Callisto/Tests/CallisoFunctionalTests/CallistoTests.cs
--- a/Engine/Source/Programs/Horde/HordeStorage/Callisto/Tests/CallisoFunctionalTests/CallistoTests.cs
+++ b/Engine/Source/Programs/Horde/HordeStorage/Callisto/Tests/CallisoFunctionalTests/CallistoTests.cs
@@ -9,6 +9,7 @@ using System.Net.Http;
 using System.Threading.Tasks;
 using Callisto.Controllers;
 using Callisto.Implementation;
+using EpicGames.Horde.Storage;
 using Jupiter.Implementation;
 using Microsoft.AspNetCore.TestHost;
 using Microsoft.AspNetCore.Hosting;
diff --git a/Engine/Source/Programs/Horde/HordeStorage/Callisto/Tests/CallisoUnitTests/TransactionLogTests.cs b/Engine/Source/Programs/Horde/HordeStorage/Callisto/Tests/CallisoUnitTests/TransactionLogTests.cs
--- a/Engine/Source/Programs/Horde/HordeStorage/Callisto/Tests/CallisoUnitTests/TransactionLogTests.cs
+++ b/Engine/Source/Programs/Horde/HordeStorage/Callisto/Tests/CallisoUnitTests/TransactionLogTests.cs
@@ -8,6 +8,7 @@ using System.Linq;
 using Microsoft.VisualStudio.TestTools.UnitTesting;
 using System.Threading.Tasks;
 using Callisto.Implementation;
+using EpicGames.Horde.Storage;
 using Jupiter.Implementation;
 
 namespace Calliso.UnitTests
