Horde Storage - Added a cache deployment mode which caches blobs and ref mappings locally while relying on a upstream Horde Storage deployment to get new data. This is essentially a read thru http cache but with the ability to act on a blob level rather then http request level.

This is only partially tested right now as we need some of the changes in this changelist deployed to our upstream to fully test it out.

diff --git a/Engine/Source/Programs/Horde/HordeStorage/Horde.Storage/Controllers/ContentIdController.cs b/Engine/Source/Programs/Horde/HordeStorage/Horde.Storage/Controllers/ContentIdController.cs
new file mode 100644
--- /dev/null
+++ b/Engine/Source/Programs/Horde/HordeStorage/Horde.Storage/Controllers/ContentIdController.cs
@@ -0,0 +1,100 @@
+// Copyright Epic Games, Inc. All Rights Reserved.
+
+using System;
+using System.Linq;
+using System.Net.Mime;
+using System.Threading.Tasks;
+using Datadog.Trace;
+using EpicGames.Horde.Storage;
+using Horde.Storage.Implementation;
+using Jupiter;
+using Jupiter.Implementation;
+using Microsoft.AspNetCore.Authorization;
+using Microsoft.AspNetCore.Mvc;
+
+namespace Horde.Storage.Controllers
+{
+    [ApiController]
+    [FormatFilter]
+    [Produces(MediaTypeNames.Application.Json, MediaTypeNames.Application.Octet, CustomMediaTypeNames.UnrealCompactBinary)]
+    [Route("api/v1/content-id")]
+    public class ContentIdController : ControllerBase
+    {
+        private readonly IAuthorizationService _authorizationService;
+        private IContentIdStore _contentIdStore;
+
+        public ContentIdController(IAuthorizationService authorizationService, IContentIdStore contentIdStore)
+        {
+            _authorizationService = authorizationService;
+            _contentIdStore = contentIdStore;
+        }
+
+        /// <summary>
+        /// Returns which blobs a content id maps to
+        /// </summary>
+        /// <param name="ns">Namespace. Each namespace is completely separated from each other. Use for different types of data that is never expected to be similar (between two different games for instance). Example: `uc4.ddc`</param>
+        /// <param name="contentId">The content id to resolve </param>
+        /// <param name="format">Optional specifier to set which output format is used json/raw/cb</param>
+        [HttpGet("{ns}/{contentId}.{format?}", Order = 500)]
+        [ProducesDefaultResponseType]
+        [ProducesResponseType(type: typeof(ProblemDetails), 400)]
+        [Authorize("Cache.read")]
+        public async Task<IActionResult> Resolve(NamespaceId ns, ContentId contentId, string? format = null)
+        {
+            using (IScope _ = Tracer.Instance.StartActive("authorize"))
+            {
+                AuthorizationResult authorizationResult = await _authorizationService.AuthorizeAsync(User, ns, NamespaceAccessRequirement.Name);
+
+                if (!authorizationResult.Succeeded)
+                {
+                    return Forbid();
+                }
+            }
+            BlobIdentifier[]? blobs = await _contentIdStore.Resolve(ns, contentId);
+
+            if (blobs == null)
+            {
+                return NotFound(new ProblemDetails { Title = $"Unable to resolve content id {contentId} ({ns})." });
+            }
+
+            return Ok(new ResolvedContentIdResponse(blobs));
+        }
+
+        /// <summary>
+        /// Update a content id mapping
+        /// </summary>
+        /// <param name="ns">Namespace. Each namespace is completely separated from each other. Use for different types of data that is never expected to be similar (between two different games for instance). Example: `uc4.ddc`</param>
+        /// <param name="contentId">The content id to resolve </param>
+        /// <param name="blobIdentifier">The blob identifier to map the content id to</param>
+        /// <param name="contentWeight">The weight of this mapping, higher means more important</param>
+        [HttpPut("{ns}/{contentId}/update/{blobIdentifier}/{contentWeight}", Order = 500)]
+        [ProducesDefaultResponseType]
+        [ProducesResponseType(type: typeof(ProblemDetails), 400)]
+        [Authorize("Cache.write")]
+        public async Task<IActionResult> UpdateContentIdMapping(NamespaceId ns, ContentId contentId, BlobIdentifier blobIdentifier, int contentWeight)
+        {
+            using (IScope _ = Tracer.Instance.StartActive("authorize"))
+            {
+                AuthorizationResult authorizationResult = await _authorizationService.AuthorizeAsync(User, ns, NamespaceAccessRequirement.Name);
+
+                if (!authorizationResult.Succeeded)
+                {
+                    return Forbid();
+                }
+            }
+            await _contentIdStore.Put(ns, contentId, blobIdentifier, contentWeight);
+
+            return Ok();
+        }
+    }
+
+    public class ResolvedContentIdResponse
+    {
+        public BlobIdentifier[] Blobs { get; set; }
+
+        public ResolvedContentIdResponse(BlobIdentifier[] blobs)
+        {
+            Blobs = blobs;
+        }
+    }
+}
diff --git a/Engine/Source/Programs/Horde/HordeStorage/Horde.Storage/Controllers/ReferencesController.cs b/Engine/Source/Programs/Horde/HordeStorage/Horde.Storage/Controllers/ReferencesController.cs
--- a/Engine/Source/Programs/Horde/HordeStorage/Horde.Storage/Controllers/ReferencesController.cs
+++ b/Engine/Source/Programs/Horde/HordeStorage/Horde.Storage/Controllers/ReferencesController.cs
@@ -210,6 +210,54 @@ namespace Horde.Storage.Controllers
             }
         }
 
+          /// <summary>
+        /// Returns the metadata about a ref key
+        /// </summary>
+        /// <param name="ns">Namespace. Each namespace is completely separated from each other. Use for different types of data that is never expected to be similar (between two different games for instance). Example: `uc4.ddc`</param>
+        /// <param name="bucket">The category/type of record you are caching. Is a clustered key together with the actual key, but all records in the same bucket can be dropped easily. Example: `terrainTexture` </param>
+        /// <param name="key">The unique name of this particular key. `iAmAVeryValidKey`</param>
+        /// <param name="fields">The fields to include in the response, omit this to include everything.</param>
+        [HttpGet("{ns}/{bucket}/{key}/metadata", Order = 500)]
+        [Authorize("Object.read")]
+        public async Task<IActionResult> GetMetadata(
+            [FromRoute] [Required] NamespaceId ns,
+            [FromRoute] [Required] BucketId bucket,
+            [FromRoute] [Required] IoHashKey key,
+            [FromQuery] string[] fields)
+        {
+            if (ShouldDoAuth())
+            {
+                using (IScope _ = Tracer.Instance.StartActive("authorize"))
+                {
+                    AuthorizationResult authorizationResult = await _authorizationService.AuthorizeAsync(User, ns, NamespaceAccessRequirement.Name);
+
+                    if (!authorizationResult.Succeeded)
+                    {
+                        return Forbid();
+                    }
+                }
+            }
+
+            try
+            {
+                (ObjectRecord objectRecord, BlobContents? _) = await _objectService.Get(ns, bucket, key, fields);
+
+                return Ok(new RefMetadataResponse(objectRecord));
+            }
+            catch (NamespaceNotFoundException e)
+            {
+                return NotFound(new ProblemDetails {Title = $"Namespace {e.Namespace} did not exist"});
+            }
+            catch (ObjectNotFoundException e)
+            {
+                return NotFound(new ProblemDetails { Title = $"Object {e.Bucket} {e.Key} did not exist" });
+            }
+            catch (BlobNotFoundException e)
+            {
+                return NotFound(new ProblemDetails { Title = $"Object {e.Blob} in {e.Ns} not found" });
+            }
+        }
+
 
         /// <summary>
         /// Checks if a object exists
@@ -537,6 +585,26 @@ namespace Horde.Storage.Controllers
         }
     }
 
+    public class RefMetadataResponse
+    {
+        public RefMetadataResponse(ObjectRecord objectRecord)
+        {
+            Namespace = objectRecord.Namespace;
+            Bucket = objectRecord.Bucket;
+            Name = objectRecord.Name;
+            PayloadIdentifier = objectRecord.BlobIdentifier;
+            LastAccess = objectRecord.LastAccess;
+            IsFinalized = objectRecord.IsFinalized;
+        }
+
+        public NamespaceId Namespace { get; set; }
+        public BucketId Bucket { get; set; }
+        public IoHashKey Name { get; set; }
+        public BlobIdentifier PayloadIdentifier { get; set; }
+        public DateTime LastAccess { get; set; }
+        public bool IsFinalized { get; set; }
+    }
+
     public class PutObjectResponse
     {
         public PutObjectResponse(ContentHash[] missingReferences)
diff --git a/Engine/Source/Programs/Horde/HordeStorage/Horde.Storage/Horde.Storage.Startup.cs b/Engine/Source/Programs/Horde/HordeStorage/Horde.Storage/Horde.Storage.Startup.cs
--- a/Engine/Source/Programs/Horde/HordeStorage/Horde.Storage/Horde.Storage.Startup.cs
+++ b/Engine/Source/Programs/Horde/HordeStorage/Horde.Storage/Horde.Storage.Startup.cs
@@ -170,6 +170,8 @@ namespace Horde.Storage
 
             services.AddOptions<ConsistencyCheckSettings>().Bind(Configuration.GetSection("ConsistencyCheck")).ValidateDataAnnotations();
 
+            services.AddOptions<UpstreamRelaySettings>().Configure(o => Configuration.GetSection("Upstream").Bind(o)).ValidateDataAnnotations();
+
             services.AddOptions<GCSettings>().Configure(o => Configuration.GetSection("GC").Bind(o)).ValidateDataAnnotations();
             services.AddOptions<ReplicationSettings>().Configure(o => Configuration.GetSection("Replication").Bind(o)).ValidateDataAnnotations();
             services.AddOptions<ServiceCredentialSettings>().Configure(o => Configuration.GetSection("ServiceCredentials").Bind(o)).ValidateDataAnnotations();
@@ -214,6 +216,7 @@ namespace Horde.Storage
             services.AddSingleton<AzureBlobStore>();
             services.AddSingleton<MemoryCacheBlobStore>();
             services.AddSingleton<MemoryBlobStore>();
+            services.AddSingleton<RelayBlobStore>();
 
             services.AddSingleton<OrphanBlobCleanup>();
             services.AddSingleton<OrphanBlobCleanupRefs>();
@@ -289,6 +292,8 @@ namespace Horde.Storage
                     return ActivatorUtilities.CreateInstance<ScyllaBlobIndex>(provider);
                 case HordeStorageSettings.BlobIndexImplementations.Mongo:
                     return ActivatorUtilities.CreateInstance<MongoBlobIndex>(provider);
+                case HordeStorageSettings.BlobIndexImplementations.Cache:
+                    return ActivatorUtilities.CreateInstance<CachedBlobIndex>(provider);
                 default:
                     throw new ArgumentOutOfRangeException();
             }
@@ -327,6 +332,8 @@ namespace Horde.Storage
                     return ActivatorUtilities.CreateInstance<ScyllaContentIdStore>(provider);
                 case HordeStorageSettings.ContentIdStoreImplementations.Mongo:
                     return ActivatorUtilities.CreateInstance<MongoContentIdStore>(provider);
+                case HordeStorageSettings.ContentIdStoreImplementations.Cache:
+                    return ActivatorUtilities.CreateInstance<CacheContentIdStore>(provider);
                 default:
                     throw new ArgumentOutOfRangeException();
             }
@@ -394,6 +401,8 @@ namespace Horde.Storage
                     return ActivatorUtilities.CreateInstance<ScyllaReferencesStore>(provider);
                 case HordeStorageSettings.ReferencesDbImplementations.Mongo:
                     return ActivatorUtilities.CreateInstance<MongoReferencesStore>(provider);
+                case HordeStorageSettings.ReferencesDbImplementations.Cache:
+                    return ActivatorUtilities.CreateInstance<CacheReferencesStore>(provider);
                 default:
                     throw new ArgumentOutOfRangeException();
             }
@@ -620,15 +629,16 @@ namespace Horde.Storage
                         break;
 
                     case HordeStorageSettings.StorageBackendImplementations.FileSystem:
-                        FilesystemSettings filesystemSettings = provider.GetService<IOptionsMonitor<FilesystemSettings>>()!.CurrentValue;
                         healthChecks.AddDiskStorageHealthCheck(options =>
                         {
+                            FilesystemSettings filesystemSettings = provider.GetService<IOptionsMonitor<FilesystemSettings>>()!.CurrentValue;
                             string? driveRoot = Path.GetPathRoot(PathUtil.ResolvePath(filesystemSettings.RootDir));
                             options.AddDrive(driveRoot);
                         });
                         break;
                     case HordeStorageSettings.StorageBackendImplementations.Memory:
                     case HordeStorageSettings.StorageBackendImplementations.MemoryBlobStore:
+                    case HordeStorageSettings.StorageBackendImplementations.Relay:
                         break;
                     default:
                         throw new ArgumentOutOfRangeException("Unhandled storage impl " + impl);
diff --git a/Engine/Source/Programs/Horde/HordeStorage/Horde.Storage/Horde.Storage.csproj b/Engine/Source/Programs/Horde/HordeStorage/Horde.Storage/Horde.Storage.csproj
--- a/Engine/Source/Programs/Horde/HordeStorage/Horde.Storage/Horde.Storage.csproj
+++ b/Engine/Source/Programs/Horde/HordeStorage/Horde.Storage/Horde.Storage.csproj
@@ -55,6 +55,16 @@
     <None Remove="Tests\**" />
   </ItemGroup>
 
+  <ItemGroup>
+    <None Remove="appsettings.Cache.yaml" />
+  </ItemGroup>
+
+  <ItemGroup>
+    <Content Include="appsettings.Cache.yaml">
+      <CopyToOutputDirectory>PreserveNewest</CopyToOutputDirectory>
+    </Content>
+  </ItemGroup>
+
   <ItemGroup>
     <PackageReference Include="akarnokd.async-enumerable-dotnet" Version="0.0.4" />
     <PackageReference Include="AspNetCore.HealthChecks.Aws.S3" Version="5.0.1" />
diff --git a/Engine/Source/Programs/Horde/HordeStorage/Horde.Storage/Implementation/Blob/CachedBlobIndex.cs b/Engine/Source/Programs/Horde/HordeStorage/Horde.Storage/Implementation/Blob/CachedBlobIndex.cs
new file mode 100644
--- /dev/null
+++ b/Engine/Source/Programs/Horde/HordeStorage/Horde.Storage/Implementation/Blob/CachedBlobIndex.cs
@@ -0,0 +1,62 @@
+// Copyright Epic Games, Inc. All Rights Reserved.
+
+using System;
+using System.Collections.Generic;
+using System.Threading.Tasks;
+using EpicGames.Horde.Storage;
+using Jupiter;
+using Jupiter.Implementation;
+using Microsoft.Extensions.Options;
+
+namespace Horde.Storage.Implementation.Blob;
+
+public class CachedBlobIndex : IBlobIndex
+{
+    private readonly IOptionsMonitor<JupiterSettings> _jupiterSettings;
+    private readonly FileSystemStore _fileSystemStore;
+
+    public CachedBlobIndex(IOptionsMonitor<JupiterSettings> jupiterSettings, FileSystemStore fileSystemStore)
+    {
+        _jupiterSettings = jupiterSettings;
+        _fileSystemStore = fileSystemStore;
+    }
+
+    public async Task<IBlobIndex.BlobInfo?> GetBlobInfo(NamespaceId ns, BlobIdentifier id)
+    {
+        if (await BlobExistsInRegion(ns, id))
+        {
+            return new IBlobIndex.BlobInfo
+            {
+                Namespace = ns,
+                References = new List<(BucketId, IoHashKey)>(),
+                Regions = new HashSet<string> {_jupiterSettings.CurrentValue.CurrentSite}
+            };
+        }
+
+        return null;
+    }
+
+    public async Task AddBlobToIndex(NamespaceId ns, BlobIdentifier id)
+    {
+        // We do not actually track any blob information when running in cached mode
+        await Task.CompletedTask;
+    }
+
+    public async Task<bool> RemoveBlobFromIndex(NamespaceId ns, BlobIdentifier id)
+    {
+        // We do not actually track any blob information when running in cached mode
+        await Task.CompletedTask;
+        return false;
+    }
+
+    public async Task<bool> BlobExistsInRegion(NamespaceId ns, BlobIdentifier blobIdentifier)
+    {
+        return await _fileSystemStore.Exists(ns, blobIdentifier);
+    }
+
+    public async Task AddRefToBlobs(NamespaceId ns, BucketId bucket, IoHashKey key, BlobIdentifier[] blobs)
+    {
+        // We do not actually track any blob information when running in cached mode
+        await Task.CompletedTask;
+    }
+}
diff --git a/Engine/Source/Programs/Horde/HordeStorage/Horde.Storage/Implementation/Blob/IBlobService.cs b/Engine/Source/Programs/Horde/HordeStorage/Horde.Storage/Implementation/Blob/IBlobService.cs
--- a/Engine/Source/Programs/Horde/HordeStorage/Horde.Storage/Implementation/Blob/IBlobService.cs
+++ b/Engine/Source/Programs/Horde/HordeStorage/Horde.Storage/Implementation/Blob/IBlobService.cs
@@ -78,6 +78,7 @@ public class BlobService : IBlobService
             HordeStorageSettings.StorageBackendImplementations.FileSystem => provider.GetService<FileSystemStore>(),
             HordeStorageSettings.StorageBackendImplementations.Memory => provider.GetService<MemoryCacheBlobStore>(),
             HordeStorageSettings.StorageBackendImplementations.MemoryBlobStore => provider.GetService<MemoryBlobStore>(),
+            HordeStorageSettings.StorageBackendImplementations.Relay => provider.GetService<RelayBlobStore>(),
             _ => throw new ArgumentOutOfRangeException()
         };
         if (store == null)
diff --git a/Engine/Source/Programs/Horde/HordeStorage/Horde.Storage/Implementation/Blob/RelayBlobStore.cs b/Engine/Source/Programs/Horde/HordeStorage/Horde.Storage/Implementation/Blob/RelayBlobStore.cs
new file mode 100644
--- /dev/null
+++ b/Engine/Source/Programs/Horde/HordeStorage/Horde.Storage/Implementation/Blob/RelayBlobStore.cs
@@ -0,0 +1,98 @@
+// Copyright Epic Games, Inc. All Rights Reserved.
+
+using System;
+using System.Collections.Generic;
+using System.IO;
+using System.Net;
+using System.Net.Http;
+using System.Net.Http.Headers;
+using System.Net.Mime;
+using System.Threading.Tasks;
+using EpicGames.Horde.Storage;
+using Jupiter.Implementation;
+using Microsoft.Extensions.Options;
+
+namespace Horde.Storage.Implementation
+{
+    public class RelayBlobStore : RelayStore, IBlobStore
+    {
+        public RelayBlobStore(IOptionsMonitor<UpstreamRelaySettings> settings, IHttpClientFactory httpClientFactory, IServiceCredentials serviceCredentials) : base(settings, httpClientFactory, serviceCredentials)
+        {
+        }
+
+        public async Task<BlobIdentifier> PutObject(NamespaceId ns, byte[] blob, BlobIdentifier identifier)
+        {
+            HttpRequestMessage putObjectRequest = BuildHttpRequest(HttpMethod.Put, $"api/v1/blobs/{ns}/{identifier}");
+
+            putObjectRequest.Content = new ByteArrayContent(blob);
+            putObjectRequest.Content.Headers.ContentType = new MediaTypeHeaderValue(MediaTypeNames.Application.Octet);
+            putObjectRequest.Content.Headers.Add(CommonHeaders.HashHeaderName, identifier.ToString());
+
+            HttpResponseMessage response = await HttpClient.SendAsync(putObjectRequest);
+
+            response.EnsureSuccessStatusCode();
+
+            return identifier;
+        }
+
+        public Task<BlobIdentifier> PutObject(NamespaceId ns, ReadOnlyMemory<byte> blob, BlobIdentifier identifier)
+        {
+            return PutObject(ns, blob.ToArray(), identifier);
+        }
+
+        public async Task<BlobIdentifier> PutObject(NamespaceId ns, Stream content, BlobIdentifier identifier)
+        {
+            HttpRequestMessage putObjectRequest = BuildHttpRequest(HttpMethod.Put, $"api/v1/blobs/{ns}/{identifier}");
+            putObjectRequest.Content = new StreamContent(content);
+            putObjectRequest.Content.Headers.ContentType = new MediaTypeHeaderValue(MediaTypeNames.Application.Octet);
+            putObjectRequest.Content.Headers.Add(CommonHeaders.HashHeaderName, identifier.ToString());
+            HttpResponseMessage response = await HttpClient.SendAsync(putObjectRequest);
+
+            response.EnsureSuccessStatusCode();
+
+            return identifier;
+        }
+
+        public async Task<BlobContents> GetObject(NamespaceId ns, BlobIdentifier blob)
+        {
+            HttpRequestMessage getObjectRequest = BuildHttpRequest(HttpMethod.Get, $"api/v1/blobs/{ns}/{blob}");
+            getObjectRequest.Headers.Add("Accept", MediaTypeNames.Application.Octet);
+            HttpResponseMessage response = await HttpClient.SendAsync(getObjectRequest);
+            if (response.StatusCode == HttpStatusCode.NotFound)
+                throw new BlobNotFoundException(ns, blob);
+            response.EnsureSuccessStatusCode();
+
+            long? contentLength = response.Content.Headers.ContentLength;
+            if (contentLength == null)
+                throw new Exception($"Content length missing in response from upstream blob store. This is not supported");
+            return new BlobContents(await response.Content.ReadAsStreamAsync(), contentLength.Value);
+        }
+
+        public async Task<bool> Exists(NamespaceId ns, BlobIdentifier blob)
+        {
+            HttpRequestMessage headObjectRequest = BuildHttpRequest(HttpMethod.Head, $"api/v1/blobs/{ns}/{blob}");
+            HttpResponseMessage response = await HttpClient.SendAsync(headObjectRequest);
+            if (response.StatusCode == HttpStatusCode.NotFound)
+                return false;
+
+            response.EnsureSuccessStatusCode();
+
+            return true;
+        }
+
+        public Task DeleteObject(NamespaceId ns, BlobIdentifier blob)
+        {
+            throw new NotImplementedException("DeleteObjects is not supported on the relay blob store");
+        }
+
+        public Task DeleteNamespace(NamespaceId ns)
+        {
+            throw new NotImplementedException("DeleteNamespace is not supported on the relay blob store");
+        }
+
+        public IAsyncEnumerable<(BlobIdentifier, DateTime)> ListObjects(NamespaceId ns)
+        {
+            throw new NotImplementedException("ListObjects is not supported on the relay blob store");
+        }
+    }
+}
diff --git a/Engine/Source/Programs/Horde/HordeStorage/Horde.Storage/Implementation/ContentId/CacheContentIdStore.cs b/Engine/Source/Programs/Horde/HordeStorage/Horde.Storage/Implementation/ContentId/CacheContentIdStore.cs
new file mode 100644
--- /dev/null
+++ b/Engine/Source/Programs/Horde/HordeStorage/Horde.Storage/Implementation/ContentId/CacheContentIdStore.cs
@@ -0,0 +1,42 @@
+// Copyright Epic Games, Inc. All Rights Reserved.
+
+using System.Net;
+using System.Net.Http;
+using System.Threading.Tasks;
+using EpicGames.Horde.Storage;
+using Horde.Storage.Controllers;
+using Jupiter.Implementation;
+using Microsoft.Extensions.Options;
+
+namespace Horde.Storage.Implementation
+{
+    public class CacheContentIdStore : RelayStore, IContentIdStore
+    {
+        
+        public CacheContentIdStore(IOptionsMonitor<UpstreamRelaySettings> settings, IHttpClientFactory httpClientFactory, IServiceCredentials serviceCredentials) : base(settings, httpClientFactory, serviceCredentials)
+        {
+        }
+
+        public async Task<BlobIdentifier[]?> Resolve(NamespaceId ns, ContentId contentId)
+        {
+            HttpRequestMessage getContentIdRequest = BuildHttpRequest(HttpMethod.Get, $"api/v1/content-id/{ns}/{contentId}");
+            HttpResponseMessage response = await HttpClient.SendAsync(getContentIdRequest);
+
+            if (response.StatusCode == HttpStatusCode.NotFound)
+                return null;
+
+            response.EnsureSuccessStatusCode();
+            ResolvedContentIdResponse resolvedContentId = await response.Content.ReadAsAsync<ResolvedContentIdResponse>();
+
+            return resolvedContentId.Blobs;
+        }
+
+        public async Task Put(NamespaceId ns, ContentId contentId, BlobIdentifier blobIdentifier, int contentWeight)
+        {
+            HttpRequestMessage putContentIdRequest = BuildHttpRequest(HttpMethod.Put, $"api/v1/content-id/{ns}/{contentId}/update/{blobIdentifier}/{contentWeight}");
+            HttpResponseMessage response = await HttpClient.SendAsync(putContentIdRequest);
+
+            response.EnsureSuccessStatusCode();
+        }
+    }
+}
diff --git a/Engine/Source/Programs/Horde/HordeStorage/Horde.Storage/Implementation/ContentId/MemoryContentIdStore.cs b/Engine/Source/Programs/Horde/HordeStorage/Horde.Storage/Implementation/ContentId/MemoryContentIdStore.cs
--- a/Engine/Source/Programs/Horde/HordeStorage/Horde.Storage/Implementation/ContentId/MemoryContentIdStore.cs
+++ b/Engine/Source/Programs/Horde/HordeStorage/Horde.Storage/Implementation/ContentId/MemoryContentIdStore.cs
@@ -1,12 +1,8 @@
 ﻿// Copyright Epic Games, Inc. All Rights Reserved.
 
-using System;
 using System.Collections.Concurrent;
 using System.Collections.Generic;
-using System.Linq;
 using System.Threading.Tasks;
-using Cassandra;
-using Cassandra.Mapping;
 using EpicGames.Horde.Storage;
 using Jupiter.Implementation;
 
diff --git a/Engine/Source/Programs/Horde/HordeStorage/Horde.Storage/Implementation/MongoStore.cs b/Engine/Source/Programs/Horde/HordeStorage/Horde.Storage/Implementation/MongoStore.cs
--- a/Engine/Source/Programs/Horde/HordeStorage/Horde.Storage/Implementation/MongoStore.cs
+++ b/Engine/Source/Programs/Horde/HordeStorage/Horde.Storage/Implementation/MongoStore.cs
@@ -13,10 +13,12 @@ namespace Horde.Storage.Implementation;
 
 public class MongoStore
 {
+    private readonly string? _overrideDatabaseName;
     protected readonly MongoClient _client;
 
-    public MongoStore(IOptionsMonitor<MongoSettings> settings)
+    public MongoStore(IOptionsMonitor<MongoSettings> settings, string? overrideDatabaseName = null)
     {
+        _overrideDatabaseName = overrideDatabaseName;
         MongoClientSettings mongoClientSettings = MongoClientSettings.FromUrl(
             new MongoUrl(settings.CurrentValue.ConnectionString)
         );
@@ -77,11 +79,7 @@ public class MongoStore
 
     private string GetDatabaseName()
     {
-        // TODO: Database per namespace seems a bit much?
-        //string cleanedNs = ns.ToString().Replace(".", "_");
-        //string dbName = $"HordeStorage_{cleanedNs}";
-        string dbName = $"HordeStorage";
-        return dbName;
+        return _overrideDatabaseName ?? "HordeStorage";
     }
 }
 
diff --git a/Engine/Source/Programs/Horde/HordeStorage/Horde.Storage/Implementation/ObjectService.cs b/Engine/Source/Programs/Horde/HordeStorage/Horde.Storage/Implementation/ObjectService.cs
--- a/Engine/Source/Programs/Horde/HordeStorage/Horde.Storage/Implementation/ObjectService.cs
+++ b/Engine/Source/Programs/Horde/HordeStorage/Horde.Storage/Implementation/ObjectService.cs
@@ -132,7 +132,19 @@ namespace Horde.Storage.Implementation
             
             if (missingReferences.Length == 0 && missingBlobs.Length == 0)
             {
-                await _referencesStore.Finalize(ns, bucket, key);
+                try
+                {
+                    await _referencesStore.Finalize(ns, bucket, key);
+                }
+                catch (PartialReferenceResolveException e)
+                {
+                    missingReferences = e.UnresolvedReferences.ToArray();
+                }
+                catch (ReferenceIsMissingBlobsException e)
+                {
+                    missingBlobs = e.MissingBlobs.ToArray();
+                }
+
                 await _replicationLog.InsertAddEvent(ns, bucket, key, blobHash);
             }
 
diff --git a/Engine/Source/Programs/Horde/HordeStorage/Horde.Storage/Implementation/Objects/CacheReferencesStore.cs b/Engine/Source/Programs/Horde/HordeStorage/Horde.Storage/Implementation/Objects/CacheReferencesStore.cs
new file mode 100644
--- /dev/null
+++ b/Engine/Source/Programs/Horde/HordeStorage/Horde.Storage/Implementation/Objects/CacheReferencesStore.cs
@@ -0,0 +1,178 @@
+// Copyright Epic Games, Inc. All Rights Reserved.
+
+using System;
+using System.Collections.Generic;
+using System.Linq;
+using System.Net;
+using System.Net.Http;
+using System.Net.Http.Headers;
+using System.Net.Mime;
+using System.Threading.Tasks;
+using EpicGames.Horde.Storage;
+using Horde.Storage.Controllers;
+using Jupiter.Implementation;
+using Microsoft.Extensions.DependencyInjection;
+using Microsoft.Extensions.Options;
+
+namespace Horde.Storage.Implementation
+{
+    public class CacheReferencesStore : IReferencesStore
+    {
+        private readonly MongoReferencesStore _mongoReferenceStore;
+        private readonly UpstreamReferenceStore _upstreamReferenceStore;
+
+        public CacheReferencesStore(IServiceProvider provider)
+        {
+            _mongoReferenceStore = ActivatorUtilities.CreateInstance<MongoReferencesStore>(provider, "HordeStorage_Cache");
+            _mongoReferenceStore.SetTTLDuration(TimeSpan.FromDays(7.0));
+
+            _upstreamReferenceStore = ActivatorUtilities.CreateInstance<UpstreamReferenceStore>(provider);
+        }
+
+        public async Task<ObjectRecord> Get(NamespaceId ns, BucketId bucket, IoHashKey key, IReferencesStore.FieldFlags flags)
+        {
+            try
+            {
+                ObjectRecord cachedRecord = await _mongoReferenceStore.Get(ns, bucket, key, flags);
+                return cachedRecord;
+            }
+            catch (ObjectNotFoundException)
+            {
+                // not cached, we check the upstream for it
+            }
+ 
+            return await _upstreamReferenceStore.Get(ns, bucket, key, flags);
+        }
+
+        public async Task Put(NamespaceId ns, BucketId bucket, IoHashKey key, BlobIdentifier blobHash, byte[] blob, bool isFinalized)
+        {
+            Task cachePut = _mongoReferenceStore.Put(ns, bucket, key, blobHash, blob, isFinalized);
+            Task upstreamPut = _upstreamReferenceStore.Put(ns, bucket, key, blobHash, blob, isFinalized);
+
+            await Task.WhenAll(cachePut, upstreamPut);
+        }
+
+        public async Task Finalize(NamespaceId ns, BucketId bucket, IoHashKey key)
+        {
+            Task cacheFinalize = _mongoReferenceStore.Finalize(ns, bucket, key);
+            Task upstreamFinalize = _upstreamReferenceStore.Finalize(ns, bucket, key);
+
+            await Task.WhenAll(cacheFinalize, upstreamFinalize);
+        }
+
+        public async Task UpdateLastAccessTime(NamespaceId ns, BucketId bucket, IoHashKey key, DateTime newLastAccessTime)
+        {
+            Task cacheUpdateLastAccess = _mongoReferenceStore.UpdateLastAccessTime(ns, bucket, key, newLastAccessTime);
+            Task upstreamUpdateLastAccess = _upstreamReferenceStore.UpdateLastAccessTime(ns, bucket, key, newLastAccessTime);
+
+            await Task.WhenAll(cacheUpdateLastAccess, upstreamUpdateLastAccess);
+        }
+
+        public IAsyncEnumerable<(BucketId, IoHashKey, DateTime)> GetRecords(NamespaceId ns)
+        {
+            throw new NotImplementedException("GetRecords not supported on a cached reference store");
+        }
+
+        public IAsyncEnumerable<NamespaceId> GetNamespaces()
+        {
+            throw new NotImplementedException("GetNamespaces not supported on a cached reference store");
+        }
+
+        public Task<bool> Delete(NamespaceId ns, BucketId bucket, IoHashKey key)
+        {
+            throw new NotImplementedException("Deletes are not supported on a cached reference store");
+        }
+
+        public Task<long> DropNamespace(NamespaceId ns)
+        {
+            throw new NotImplementedException("DropNamespace is not supported on a cached reference store");
+
+        }
+
+        public Task<long> DeleteBucket(NamespaceId ns, BucketId bucket)
+        {
+            throw new NotImplementedException("DeleteBucket is not supported on a cached reference store");
+        }
+    }
+
+    class UpstreamReferenceStore : RelayStore, IReferencesStore
+    {
+        public UpstreamReferenceStore(IOptionsMonitor<UpstreamRelaySettings> settings, IHttpClientFactory httpClientFactory, IServiceCredentials serviceCredentials) : base(settings, httpClientFactory, serviceCredentials)
+        {
+        }
+
+        public async Task<ObjectRecord> Get(NamespaceId ns, BucketId bucket, IoHashKey key, IReferencesStore.FieldFlags flags)
+        {
+            HttpRequestMessage getObjectRequest = BuildHttpRequest(HttpMethod.Get, $"api/v1/refs/{ns}/{bucket}/{key}/metadata");
+            getObjectRequest.Headers.Add("Accept", MediaTypeNames.Application.Json);
+            HttpResponseMessage response = await HttpClient.SendAsync(getObjectRequest);
+            if (response.StatusCode == HttpStatusCode.NotFound)
+                throw new ObjectNotFoundException(ns, bucket, key);
+            response.EnsureSuccessStatusCode();
+
+            RefMetadataResponse metadataResponse = await response.Content.ReadAsAsync<RefMetadataResponse>();
+
+            return new ObjectRecord(metadataResponse.Namespace, metadataResponse.Bucket, metadataResponse.Name, metadataResponse.LastAccess, null, metadataResponse.PayloadIdentifier, metadataResponse.IsFinalized);
+        }
+
+        public async Task Put(NamespaceId ns, BucketId bucket, IoHashKey key, BlobIdentifier blobHash, byte[] blob, bool isFinalized)
+        {
+            HttpRequestMessage putObjectRequest = BuildHttpRequest(HttpMethod.Put, $"api/v1/refs/{ns}/{bucket}/{key}");
+            putObjectRequest.Headers.Add("Accept", MediaTypeNames.Application.Json);
+            putObjectRequest.Headers.Add(CommonHeaders.HashHeaderName, blobHash.ToString());
+            putObjectRequest.Content = new ByteArrayContent(blob);
+            putObjectRequest.Content.Headers.ContentType = new MediaTypeHeaderValue(MediaTypeNames.Application.Octet);
+
+            HttpResponseMessage response = await HttpClient.SendAsync(putObjectRequest);
+            response.EnsureSuccessStatusCode();
+
+            // if this put returns needs, we cant really do anything about it. we should be calling put in the upstream blob store as the operation continues and we should be able to catch the missing blobs during the finalize call
+        }
+
+        public async Task Finalize(NamespaceId ns, BucketId bucket, IoHashKey key)
+        {
+            HttpRequestMessage putObjectRequest = BuildHttpRequest(HttpMethod.Put, $"api/v1/refs/{ns}/{bucket}/{key}");
+            putObjectRequest.Headers.Add("Accept", MediaTypeNames.Application.Json);
+            HttpResponseMessage response = await HttpClient.SendAsync(putObjectRequest);
+            response.EnsureSuccessStatusCode();
+
+            PutObjectResponse putObjectResponse = await response.Content.ReadAsAsync<PutObjectResponse>();
+            if (putObjectResponse.Needs.Length != 0)
+            {
+                // TODO: we assume here that all the objects that are missing are content ids, that may not be true, but as needs only contains a list with mixed hashes we can not detect the difference
+                throw new PartialReferenceResolveException(putObjectResponse.Needs.Select(hash => new ContentId(hash.HashData)).ToList());
+            }
+        }
+
+        public async Task UpdateLastAccessTime(NamespaceId ns, BucketId bucket, IoHashKey key, DateTime newLastAccessTime)
+        {
+            // there is no endpoint to update last access time externally, but fetching a object will in turn update its last access time, though to a different time then newLastAccessTime
+            await Get(ns, bucket, key, IReferencesStore.FieldFlags.None);
+        }
+
+        public IAsyncEnumerable<(BucketId, IoHashKey, DateTime)> GetRecords(NamespaceId ns)
+        {
+            throw new NotImplementedException("GetRecords is not supported on a upstream reference store");
+        }
+
+        public IAsyncEnumerable<NamespaceId> GetNamespaces()
+        {
+            throw new NotImplementedException("GetNamespaces is not supported on a upstream reference store");
+        }
+
+        public Task<bool> Delete(NamespaceId ns, BucketId bucket, IoHashKey key)
+        {
+            throw new NotImplementedException("Delete is not supported on a upstream reference store");
+        }
+
+        public Task<long> DropNamespace(NamespaceId ns)
+        {
+            throw new NotImplementedException("DropNamespace is not supported on a upstream reference store");
+        }
+
+        public Task<long> DeleteBucket(NamespaceId ns, BucketId bucket)
+        {
+            throw new NotImplementedException("DeleteBucket is not supported on a upstream reference store");
+        }
+    }
+}
diff --git a/Engine/Source/Programs/Horde/HordeStorage/Horde.Storage/Implementation/Objects/MongoReferencesStore.cs b/Engine/Source/Programs/Horde/HordeStorage/Horde.Storage/Implementation/Objects/MongoReferencesStore.cs
--- a/Engine/Source/Programs/Horde/HordeStorage/Horde.Storage/Implementation/Objects/MongoReferencesStore.cs
+++ b/Engine/Source/Programs/Horde/HordeStorage/Horde.Storage/Implementation/Objects/MongoReferencesStore.cs
@@ -17,7 +17,7 @@ namespace Horde.Storage.Implementation
     {
         private readonly ILogger _logger = Log.ForContext<MongoReferencesStore>();
 
-        public MongoReferencesStore(IOptionsMonitor<MongoSettings> settings) : base(settings)
+        public MongoReferencesStore(IOptionsMonitor<MongoSettings> settings, string? overrideDatabaseName = null) : base(settings, overrideDatabaseName)
         {
             CreateCollectionIfNotExists<MongoReferencesModelV0>().Wait();
             CreateCollectionIfNotExists<MongoNamespacesModelV0>().Wait();
@@ -180,6 +180,16 @@ namespace Horde.Storage.Implementation
             // failed to delete
             return 0L;
         }
+
+        public void SetTTLDuration(TimeSpan duration)
+        {
+            IndexKeysDefinitionBuilder<MongoReferencesModelV0> indexKeysDefinitionBuilder = Builders<MongoReferencesModelV0>.IndexKeys;
+            CreateIndexModel<MongoReferencesModelV0> indexTTL = new CreateIndexModel<MongoReferencesModelV0>(
+                indexKeysDefinitionBuilder.Ascending(m => m.LastAccessTime), new CreateIndexOptions {Name = "LastAccessTTL", ExpireAfter = duration}
+            );
+
+            AddIndexFor<MongoReferencesModelV0>().CreateOne(indexTTL);
+        }
     }
 
     
diff --git a/Engine/Source/Programs/Horde/HordeStorage/Horde.Storage/Implementation/RelayStore.cs b/Engine/Source/Programs/Horde/HordeStorage/Horde.Storage/Implementation/RelayStore.cs
new file mode 100644
--- /dev/null
+++ b/Engine/Source/Programs/Horde/HordeStorage/Horde.Storage/Implementation/RelayStore.cs
@@ -0,0 +1,38 @@
+using System;
+using System.Net.Http;
+using Microsoft.Extensions.Options;
+
+namespace Horde.Storage.Implementation
+{
+    public abstract class RelayStore
+    {
+        private readonly IOptionsMonitor<UpstreamRelaySettings> _settings;
+        private readonly IServiceCredentials _serviceCredentials;
+        private readonly HttpClient _httpClient;
+
+        protected HttpClient HttpClient
+        {
+            get { return _httpClient; }
+        }
+
+        protected RelayStore(IOptionsMonitor<UpstreamRelaySettings> settings, IHttpClientFactory httpClientFactory, IServiceCredentials serviceCredentials)
+        {
+            _settings = settings;
+            _serviceCredentials = serviceCredentials;
+
+            _httpClient = httpClientFactory.CreateClient();
+            _httpClient.BaseAddress = new Uri(_settings.CurrentValue.ConnectionString);
+        }
+
+        protected HttpRequestMessage BuildHttpRequest(HttpMethod method, string uri)
+        {
+            string? token = _serviceCredentials.GetToken();
+            HttpRequestMessage request = new HttpRequestMessage(method, uri);
+            if (!string.IsNullOrEmpty(token))
+                request.Headers.Add("Authorization", "Bearer " + token);
+
+            return request;
+        }
+
+    }
+}
diff --git a/Engine/Source/Programs/Horde/HordeStorage/Horde.Storage/Implementation/ServiceCredentials.cs b/Engine/Source/Programs/Horde/HordeStorage/Horde.Storage/Implementation/ServiceCredentials.cs
--- a/Engine/Source/Programs/Horde/HordeStorage/Horde.Storage/Implementation/ServiceCredentials.cs
+++ b/Engine/Source/Programs/Horde/HordeStorage/Horde.Storage/Implementation/ServiceCredentials.cs
@@ -20,20 +20,21 @@ namespace Horde.Storage.Implementation
 
     public class ServiceCredentials : IServiceCredentials
     {
+        private readonly IOptionsMonitor<ServiceCredentialSettings> _settings;
         private readonly ClientCredentialOAuthAuthenticator? _authenticator;
 
-        public ServiceCredentials(IOptionsMonitor<ServiceCredentialSettings> settingsMonitor, ISecretResolver secretResolver)
+        public ServiceCredentials(IOptionsMonitor<ServiceCredentialSettings> settings, ISecretResolver secretResolver)
         {
-            ServiceCredentialSettings settings = settingsMonitor.CurrentValue;
-            if (!string.IsNullOrEmpty(settings.OAuthLoginUrl))
+            _settings = settings;
+            if (!string.IsNullOrEmpty(settings.CurrentValue.OAuthLoginUrl))
             {
-                string? clientId = secretResolver.Resolve(settings.OAuthClientId);
+                string? clientId = secretResolver.Resolve(settings.CurrentValue.OAuthClientId);
                 if (string.IsNullOrEmpty(clientId))
                     throw new ArgumentException("ClientId must be set when using a service credential");
-                string? clientSecret = secretResolver.Resolve(settings.OAuthClientSecret);
+                string? clientSecret = secretResolver.Resolve(settings.CurrentValue.OAuthClientSecret);
                 if (string.IsNullOrEmpty(clientSecret))
                     throw new ArgumentException("ClientSecret must be set when using a service credential");
-                _authenticator = new ClientCredentialOAuthAuthenticator(settings.OAuthLoginUrl, clientId, clientSecret, settings.OAuthScope);
+                _authenticator = new ClientCredentialOAuthAuthenticator(settings.CurrentValue.OAuthLoginUrl, clientId, clientSecret, settings.CurrentValue.OAuthScope);
             }
         }
 
diff --git a/Engine/Source/Programs/Horde/HordeStorage/Horde.Storage/Settings.cs b/Engine/Source/Programs/Horde/HordeStorage/Horde.Storage/Settings.cs
--- a/Engine/Source/Programs/Horde/HordeStorage/Horde.Storage/Settings.cs
+++ b/Engine/Source/Programs/Horde/HordeStorage/Horde.Storage/Settings.cs
@@ -50,28 +50,32 @@ namespace Horde.Storage
             Azure,
             FileSystem,
             Memory,
-            MemoryBlobStore
+            MemoryBlobStore,
+            Relay
         }
 
         public enum ReferencesDbImplementations
         {
             Memory,
             Scylla,
-            Mongo
+            Mongo,
+            Cache
         }
 
         public enum ContentIdStoreImplementations
         {
             Memory,
             Scylla,
-            Mongo
+            Mongo,
+            Cache
         }
 
         public enum BlobIndexImplementations
         {
             Memory,
             Scylla,
-            Mongo
+            Mongo,
+            Cache
         }
 
         public enum LeaderElectionImplementations
@@ -257,4 +261,10 @@ namespace Horde.Storage
         public List<string> CleanNamespacesV1 { get; set; } = new List<string>();
         public List<string> CleanNamespaces { get; set; }  = new List<string>();
     }
+
+    public class UpstreamRelaySettings
+    {
+        [Required] public string ConnectionString { get; set; } = null!;
+    }
+
 }
diff --git a/Engine/Source/Programs/Horde/HordeStorage/Horde.Storage/appsettings.Production.json b/Engine/Source/Programs/Horde/HordeStorage/Horde.Storage/appsettings.Production.json
--- a/Engine/Source/Programs/Horde/HordeStorage/Horde.Storage/appsettings.Production.json
+++ b/Engine/Source/Programs/Horde/HordeStorage/Horde.Storage/appsettings.Production.json
@@ -27,10 +27,6 @@
     "AutoscaleCapacity": true
   },
 
-  "Horde_Storage": {
-    "TransactionLogWriterImplementation": "callisto"
-  },
-
   "Replication": {
     "ReplicationPollFrequencySeconds": 15
   },
diff --git a/Engine/Source/Programs/Horde/HordeStorage/Jupiter.Common/BaseProgram.cs b/Engine/Source/Programs/Horde/HordeStorage/Jupiter.Common/BaseProgram.cs
--- a/Engine/Source/Programs/Horde/HordeStorage/Jupiter.Common/BaseProgram.cs
+++ b/Engine/Source/Programs/Horde/HordeStorage/Jupiter.Common/BaseProgram.cs
@@ -48,6 +48,7 @@ namespace Jupiter
         private static IConfiguration GetConfigurationBuilder()
         {
             string env = Environment.GetEnvironmentVariable("ASPNETCORE_ENVIRONMENT") ?? "Production";
+            string mode = Environment.GetEnvironmentVariable("HORDESTORAGE_MODE") ?? "DefaultMode";
             string configRoot = "/config";
             // set the config root to config under the current directory for windows
             if (OperatingSystem.IsWindows())
@@ -61,6 +62,10 @@ namespace Jupiter
                     path:
                     $"appsettings.{env}.json",
                     true)
+                .AddYamlFile(
+                    path:
+                    $"appsettings.{mode}.yaml",
+                    false)
                 .AddYamlFile(Path.Combine(configRoot, "appsettings.Local.yaml"), optional: true, reloadOnChange: true)
                 .AddEnvironmentVariables()
                 .Build();
