Horde Storage - Added missing copyright header

diff --git a/Engine/Source/Programs/Horde/HordeStorage/Horde.Storage/Implementation/RelayStore.cs b/Engine/Source/Programs/Horde/HordeStorage/Horde.Storage/Implementation/RelayStore.cs
--- a/Engine/Source/Programs/Horde/HordeStorage/Horde.Storage/Implementation/RelayStore.cs
+++ b/Engine/Source/Programs/Horde/HordeStorage/Horde.Storage/Implementation/RelayStore.cs
@@ -1,3 +1,5 @@
+// Copyright Epic Games, Inc. All Rights Reserved.
+
 using System;
 using System.Net.Http;
 using Microsoft.Extensions.Options;
