Horde Storage - Fixed docker build setup to handle shared libraries being required

diff --git a/Engine/Source/Programs/Horde/BuildHorde.xml b/Engine/Source/Programs/Horde/BuildHorde.xml
--- a/Engine/Source/Programs/Horde/BuildHorde.xml
+++ b/Engine/Source/Programs/Horde/BuildHorde.xml
@@ -48,13 +48,14 @@
 		<Node Name="Build HordeStorage" Requires="Test HordeStorage">
 			<!-- Tag all the files that need to be staged to build -->
 			<Property Name="StagePaths">
+				Engine/Source/Programs/Shared/...
 				Engine/Source/Programs/Horde/HordeStorage/...
 			</Property>
 			<Tag Files="$(StagePaths)" Except=".../.vs/...;.../.git/...;.../bin/...;.../obj/..." With="#InputFiles"/>
 
 			<!-- Build the Docker images -->
-			<Docker-Build BaseDir="Engine/Source/Programs/Horde/HordeStorage" Files="#InputFiles" Tag="horde_storage:$(InformationalVersion)" DockerFile="Engine/Source/Programs/Horde/HordeStorage/dockerfile" Arguments="--build-arg Project=Horde.Storage" Environment="UE_P4_CHANGELIST=(Change)"/>
-			<Docker-Build BaseDir="Engine/Source/Programs/Horde/HordeStorage" Files="#InputFiles" Tag="callisto:$(InformationalVersion)" DockerFile="Engine/Source/Programs/Horde/HordeStorage/dockerfile" Arguments="--build-arg Project=Callisto" Environment="UE_P4_CHANGELIST=(Change)" />
+			<Docker-Build BaseDir="Engine/Source" Files="#InputFiles" Tag="horde_storage:$(InformationalVersion)" DockerFile="Engine/Source/Programs/Horde/HordeStorage/dockerfile" Arguments="--build-arg Project=Horde.Storage" Environment="UE_P4_CHANGELIST=(Change)"/>
+			<Docker-Build BaseDir="Engine/Source" Files="#InputFiles" Tag="callisto:$(InformationalVersion)" DockerFile="Engine/Source/Programs/Horde/HordeStorage/dockerfile" Arguments="--build-arg Project=Callisto" Environment="UE_P4_CHANGELIST=(Change)" />
 		</Node>
 	</Agent>
 </BuildGraph>
