Horde Storage - Fixed crash in Dev by making sure the mode yaml file is considered optional

diff --git a/Engine/Source/Programs/Horde/HordeStorage/Jupiter.Common/BaseProgram.cs b/Engine/Source/Programs/Horde/HordeStorage/Jupiter.Common/BaseProgram.cs
--- a/Engine/Source/Programs/Horde/HordeStorage/Jupiter.Common/BaseProgram.cs
+++ b/Engine/Source/Programs/Horde/HordeStorage/Jupiter.Common/BaseProgram.cs
@@ -65,7 +65,7 @@ namespace Jupiter
                 .AddYamlFile(
                     path:
                     $"appsettings.{mode}.yaml",
-                    false)
+                    true)
                 .AddYamlFile(Path.Combine(configRoot, "appsettings.Local.yaml"), optional: true, reloadOnChange: true)
                 .AddEnvironmentVariables()
                 .Build();
