Horde Storage - Renamed Namespace to Ns in metadata response to make it easier to deserialize to a C# identifier (this is also more in line with how we generally serialize the namespace).

diff --git a/Engine/Source/Programs/Horde/HordeStorage/Horde.Storage/Controllers/ReferencesController.cs b/Engine/Source/Programs/Horde/HordeStorage/Horde.Storage/Controllers/ReferencesController.cs
--- a/Engine/Source/Programs/Horde/HordeStorage/Horde.Storage/Controllers/ReferencesController.cs
+++ b/Engine/Source/Programs/Horde/HordeStorage/Horde.Storage/Controllers/ReferencesController.cs
@@ -9,6 +9,7 @@ using System.IO;
 using System.Linq;
 using System.Net.Mime;
 using System.Text;
+using System.Text.Json.Serialization;
 using System.Threading.Tasks;
 using Dasync.Collections;
 using Datadog.Trace;
@@ -587,9 +588,20 @@ namespace Horde.Storage.Controllers
 
     public class RefMetadataResponse
     {
+        [JsonConstructor]
+        public RefMetadataResponse(NamespaceId ns, BucketId bucket, IoHashKey name, BlobIdentifier payloadIdentifier, DateTime lastAccess, bool isFinalized)
+        {
+            Ns = ns;
+            Bucket = bucket;
+            Name = name;
+            PayloadIdentifier = payloadIdentifier;
+            LastAccess = lastAccess;
+            IsFinalized = isFinalized;
+        }
+
         public RefMetadataResponse(ObjectRecord objectRecord)
         {
-            Namespace = objectRecord.Namespace;
+            Ns = objectRecord.Namespace;
             Bucket = objectRecord.Bucket;
             Name = objectRecord.Name;
             PayloadIdentifier = objectRecord.BlobIdentifier;
@@ -597,7 +609,7 @@ namespace Horde.Storage.Controllers
             IsFinalized = objectRecord.IsFinalized;
         }
 
-        public NamespaceId Namespace { get; set; }
+        public NamespaceId Ns { get; set; }
         public BucketId Bucket { get; set; }
         public IoHashKey Name { get; set; }
         public BlobIdentifier PayloadIdentifier { get; set; }
