Horde Storage - Fixed build error from renaming of field

diff --git a/Engine/Source/Programs/Horde/HordeStorage/Horde.Storage/Implementation/Objects/CacheReferencesStore.cs b/Engine/Source/Programs/Horde/HordeStorage/Horde.Storage/Implementation/Objects/CacheReferencesStore.cs
--- a/Engine/Source/Programs/Horde/HordeStorage/Horde.Storage/Implementation/Objects/CacheReferencesStore.cs
+++ b/Engine/Source/Programs/Horde/HordeStorage/Horde.Storage/Implementation/Objects/CacheReferencesStore.cs
@@ -112,7 +112,7 @@ namespace Horde.Storage.Implementation
 
             RefMetadataResponse metadataResponse = await response.Content.ReadAsAsync<RefMetadataResponse>();
 
-            return new ObjectRecord(metadataResponse.Namespace, metadataResponse.Bucket, metadataResponse.Name, metadataResponse.LastAccess, null, metadataResponse.PayloadIdentifier, metadataResponse.IsFinalized);
+            return new ObjectRecord(metadataResponse.Ns, metadataResponse.Bucket, metadataResponse.Name, metadataResponse.LastAccess, null, metadataResponse.PayloadIdentifier, metadataResponse.IsFinalized);
         }
 
         public async Task Put(NamespaceId ns, BucketId bucket, IoHashKey key, BlobIdentifier blobHash, byte[] blob, bool isFinalized)
