Horde Storage - Fixes to cache mode

We are now able to succesfully cache all our structured cache api calls

diff --git a/Engine/Source/Programs/Horde/HordeStorage/Horde.Storage/Controllers/ReferencesController.cs b/Engine/Source/Programs/Horde/HordeStorage/Horde.Storage/Controllers/ReferencesController.cs
--- a/Engine/Source/Programs/Horde/HordeStorage/Horde.Storage/Controllers/ReferencesController.cs
+++ b/Engine/Source/Programs/Horde/HordeStorage/Horde.Storage/Controllers/ReferencesController.cs
@@ -9,7 +9,6 @@ using System.IO;
 using System.Linq;
 using System.Net.Mime;
 using System.Text;
-using System.Text.Json.Serialization;
 using System.Threading.Tasks;
 using Dasync.Collections;
 using Datadog.Trace;
@@ -24,6 +23,7 @@ using Microsoft.AspNetCore.Http;
 using Microsoft.AspNetCore.Http.Extensions;
 using Microsoft.AspNetCore.Mvc;
 using Microsoft.Extensions.Options;
+using Newtonsoft.Json;
 using Serilog;
 
 namespace Horde.Storage.Controllers
@@ -589,7 +589,7 @@ namespace Horde.Storage.Controllers
     public class RefMetadataResponse
     {
         [JsonConstructor]
-        public RefMetadataResponse(NamespaceId ns, BucketId bucket, IoHashKey name, BlobIdentifier payloadIdentifier, DateTime lastAccess, bool isFinalized)
+        public RefMetadataResponse(NamespaceId ns, BucketId bucket, IoHashKey name, BlobIdentifier payloadIdentifier, DateTime lastAccess, bool isFinalized, byte[]? inlinePayload)
         {
             Ns = ns;
             Bucket = bucket;
@@ -597,6 +597,7 @@ namespace Horde.Storage.Controllers
             PayloadIdentifier = payloadIdentifier;
             LastAccess = lastAccess;
             IsFinalized = isFinalized;
+            InlinePayload = inlinePayload;
         }
 
         public RefMetadataResponse(ObjectRecord objectRecord)
@@ -607,6 +608,7 @@ namespace Horde.Storage.Controllers
             PayloadIdentifier = objectRecord.BlobIdentifier;
             LastAccess = objectRecord.LastAccess;
             IsFinalized = objectRecord.IsFinalized;
+            InlinePayload = objectRecord.InlinePayload;
         }
 
         public NamespaceId Ns { get; set; }
@@ -615,6 +617,7 @@ namespace Horde.Storage.Controllers
         public BlobIdentifier PayloadIdentifier { get; set; }
         public DateTime LastAccess { get; set; }
         public bool IsFinalized { get; set; }
+        public byte[]? InlinePayload { get; set; }
     }
 
     public class PutObjectResponse
diff --git a/Engine/Source/Programs/Horde/HordeStorage/Horde.Storage/Implementation/ObjectService.cs b/Engine/Source/Programs/Horde/HordeStorage/Horde.Storage/Implementation/ObjectService.cs
--- a/Engine/Source/Programs/Horde/HordeStorage/Horde.Storage/Implementation/ObjectService.cs
+++ b/Engine/Source/Programs/Horde/HordeStorage/Horde.Storage/Implementation/ObjectService.cs
@@ -134,7 +134,7 @@ namespace Horde.Storage.Implementation
             {
                 try
                 {
-                    await _referencesStore.Finalize(ns, bucket, key);
+                    await _referencesStore.Finalize(ns, bucket, key, blobHash);
                 }
                 catch (PartialReferenceResolveException e)
                 {
@@ -186,7 +186,7 @@ namespace Horde.Storage.Implementation
 
             if (missingReferences.Length == 0 && missingBlobs.Length == 0)
             {
-                await _referencesStore.Finalize(ns, bucket, key);
+                await _referencesStore.Finalize(ns, bucket, key, blobHash);
                 await _replicationLog.InsertAddEvent(ns, bucket, key, blobHash);
             }
 
diff --git a/Engine/Source/Programs/Horde/HordeStorage/Horde.Storage/Implementation/Objects/CacheReferencesStore.cs b/Engine/Source/Programs/Horde/HordeStorage/Horde.Storage/Implementation/Objects/CacheReferencesStore.cs
--- a/Engine/Source/Programs/Horde/HordeStorage/Horde.Storage/Implementation/Objects/CacheReferencesStore.cs
+++ b/Engine/Source/Programs/Horde/HordeStorage/Horde.Storage/Implementation/Objects/CacheReferencesStore.cs
@@ -11,8 +11,10 @@ using System.Threading.Tasks;
 using EpicGames.Horde.Storage;
 using Horde.Storage.Controllers;
 using Jupiter.Implementation;
+using Microsoft.AspNetCore.Mvc;
 using Microsoft.Extensions.DependencyInjection;
 using Microsoft.Extensions.Options;
+using Jupiter;
 
 namespace Horde.Storage.Implementation
 {
@@ -41,7 +43,9 @@ namespace Horde.Storage.Implementation
                 // not cached, we check the upstream for it
             }
  
-            return await _upstreamReferenceStore.Get(ns, bucket, key, flags);
+            ObjectRecord record = await _upstreamReferenceStore.Get(ns, bucket, key, flags);
+            await _mongoReferenceStore.Put(record.Namespace, record.Bucket, record.Name, record.BlobIdentifier, record.InlinePayload, record.IsFinalized);
+            return record;
         }
 
         public async Task Put(NamespaceId ns, BucketId bucket, IoHashKey key, BlobIdentifier blobHash, byte[] blob, bool isFinalized)
@@ -52,10 +56,10 @@ namespace Horde.Storage.Implementation
             await Task.WhenAll(cachePut, upstreamPut);
         }
 
-        public async Task Finalize(NamespaceId ns, BucketId bucket, IoHashKey key)
+        public async Task Finalize(NamespaceId ns, BucketId bucket, IoHashKey key, BlobIdentifier blobIdentifier)
         {
-            Task cacheFinalize = _mongoReferenceStore.Finalize(ns, bucket, key);
-            Task upstreamFinalize = _upstreamReferenceStore.Finalize(ns, bucket, key);
+            Task cacheFinalize = _mongoReferenceStore.Finalize(ns, bucket, key, blobIdentifier);
+            Task upstreamFinalize = _upstreamReferenceStore.Finalize(ns, bucket, key, blobIdentifier);
 
             await Task.WhenAll(cacheFinalize, upstreamFinalize);
         }
@@ -112,7 +116,7 @@ namespace Horde.Storage.Implementation
 
             RefMetadataResponse metadataResponse = await response.Content.ReadAsAsync<RefMetadataResponse>();
 
-            return new ObjectRecord(metadataResponse.Ns, metadataResponse.Bucket, metadataResponse.Name, metadataResponse.LastAccess, null, metadataResponse.PayloadIdentifier, metadataResponse.IsFinalized);
+            return new ObjectRecord(metadataResponse.Ns, metadataResponse.Bucket, metadataResponse.Name, metadataResponse.LastAccess, metadataResponse.InlinePayload, metadataResponse.PayloadIdentifier, metadataResponse.IsFinalized);
         }
 
         public async Task Put(NamespaceId ns, BucketId bucket, IoHashKey key, BlobIdentifier blobHash, byte[] blob, bool isFinalized)
@@ -121,19 +125,32 @@ namespace Horde.Storage.Implementation
             putObjectRequest.Headers.Add("Accept", MediaTypeNames.Application.Json);
             putObjectRequest.Headers.Add(CommonHeaders.HashHeaderName, blobHash.ToString());
             putObjectRequest.Content = new ByteArrayContent(blob);
-            putObjectRequest.Content.Headers.ContentType = new MediaTypeHeaderValue(MediaTypeNames.Application.Octet);
+            // a blob sent to the reference store is always by definition a compact binary, so we upload it as such to avoid any automatic conversions that we would get with octet-stream
+            putObjectRequest.Content.Headers.ContentType = new MediaTypeHeaderValue(CustomMediaTypeNames.UnrealCompactBinary);
 
             HttpResponseMessage response = await HttpClient.SendAsync(putObjectRequest);
+            if (response.StatusCode == HttpStatusCode.BadRequest)
+            {
+                ProblemDetails? problem = await response.Content.ReadAsAsync<ProblemDetails>();
+                throw new Exception("Upstream returned 400 (BadRequest) with error: " + problem.Title);
+            }
             response.EnsureSuccessStatusCode();
 
             // if this put returns needs, we cant really do anything about it. we should be calling put in the upstream blob store as the operation continues and we should be able to catch the missing blobs during the finalize call
         }
 
-        public async Task Finalize(NamespaceId ns, BucketId bucket, IoHashKey key)
+        public async Task Finalize(NamespaceId ns, BucketId bucket, IoHashKey key, BlobIdentifier blobIdentifier)
         {
-            HttpRequestMessage putObjectRequest = BuildHttpRequest(HttpMethod.Put, $"api/v1/refs/{ns}/{bucket}/{key}");
+            HttpRequestMessage putObjectRequest = BuildHttpRequest(HttpMethod.Post, $"api/v1/refs/{ns}/{bucket}/{key}/finalize/{blobIdentifier}");
             putObjectRequest.Headers.Add("Accept", MediaTypeNames.Application.Json);
+
             HttpResponseMessage response = await HttpClient.SendAsync(putObjectRequest);
+
+            if (response.StatusCode == HttpStatusCode.BadRequest)
+            {
+                ProblemDetails? problem = await response.Content.ReadAsAsync<ProblemDetails>();
+                throw new Exception("Upstream returned 400 (BadRequest) with error: " + problem.Title);
+            }
             response.EnsureSuccessStatusCode();
 
             PutObjectResponse putObjectResponse = await response.Content.ReadAsAsync<PutObjectResponse>();
diff --git a/Engine/Source/Programs/Horde/HordeStorage/Horde.Storage/Implementation/Objects/IReferencesStore.cs b/Engine/Source/Programs/Horde/HordeStorage/Horde.Storage/Implementation/Objects/IReferencesStore.cs
--- a/Engine/Source/Programs/Horde/HordeStorage/Horde.Storage/Implementation/Objects/IReferencesStore.cs
+++ b/Engine/Source/Programs/Horde/HordeStorage/Horde.Storage/Implementation/Objects/IReferencesStore.cs
@@ -21,7 +21,7 @@ namespace Horde.Storage.Implementation
         }
 
         Task Put(NamespaceId ns, BucketId bucket, IoHashKey key, BlobIdentifier blobHash, byte[] blob, bool isFinalized);
-        Task Finalize(NamespaceId ns, BucketId bucket, IoHashKey key);
+        Task Finalize(NamespaceId ns, BucketId bucket, IoHashKey key, BlobIdentifier blobIdentifier);
 
 
         Task UpdateLastAccessTime(NamespaceId ns, BucketId bucket, IoHashKey key, DateTime newLastAccessTime);
diff --git a/Engine/Source/Programs/Horde/HordeStorage/Horde.Storage/Implementation/Objects/MemoryReferencesStore.cs b/Engine/Source/Programs/Horde/HordeStorage/Horde.Storage/Implementation/Objects/MemoryReferencesStore.cs
--- a/Engine/Source/Programs/Horde/HordeStorage/Horde.Storage/Implementation/Objects/MemoryReferencesStore.cs
+++ b/Engine/Source/Programs/Horde/HordeStorage/Horde.Storage/Implementation/Objects/MemoryReferencesStore.cs
@@ -45,7 +45,7 @@ namespace Horde.Storage.Implementation
             return Task.FromResult(o);
         }
 
-        public Task Finalize(NamespaceId ns, BucketId bucket, IoHashKey key)
+        public Task Finalize(NamespaceId ns, BucketId bucket, IoHashKey key, BlobIdentifier blobIdentifier)
         {
             if (!_objects.TryGetValue(BuildKey(ns, bucket, key), out MemoryStoreObject? o))
             {
diff --git a/Engine/Source/Programs/Horde/HordeStorage/Horde.Storage/Implementation/Objects/MongoReferencesStore.cs b/Engine/Source/Programs/Horde/HordeStorage/Horde.Storage/Implementation/Objects/MongoReferencesStore.cs
--- a/Engine/Source/Programs/Horde/HordeStorage/Horde.Storage/Implementation/Objects/MongoReferencesStore.cs
+++ b/Engine/Source/Programs/Horde/HordeStorage/Horde.Storage/Implementation/Objects/MongoReferencesStore.cs
@@ -54,12 +54,12 @@ namespace Horde.Storage.Implementation
             return model.ToObjectRecord();
         }
 
-        public async Task Put(NamespaceId ns, BucketId bucket, IoHashKey key, BlobIdentifier blobHash, byte[] blob, bool isFinalized)
+        public async Task Put(NamespaceId ns, BucketId bucket, IoHashKey key, BlobIdentifier blobHash, byte[]? blob, bool isFinalized)
         {
             IMongoCollection<MongoReferencesModelV0> collection = GetCollection<MongoReferencesModelV0>();
 
             //if (blob.LongLength > _settings.CurrentValue.InlineBlobMaxSize)
-            if (blob.LongLength > 64_000) // TODO: add a setting to control this
+            if (blob != null && blob.LongLength > 64_000) // TODO: add a setting to control this
             {
                 // do not inline large blobs
                 blob = Array.Empty<byte>();
@@ -78,7 +78,7 @@ namespace Horde.Storage.Implementation
             await addNamespaceTask;
         }
 
-        public async Task Finalize(NamespaceId ns, BucketId bucket, IoHashKey key)
+        public async Task Finalize(NamespaceId ns, BucketId bucket, IoHashKey key, BlobIdentifier blobIdentifier)
         {
             IMongoCollection<MongoReferencesModelV0> collection = GetCollection<MongoReferencesModelV0>();
 
@@ -211,7 +211,7 @@ namespace Horde.Storage.Implementation
             LastAccessTime = lastAccessTime;
         }
 
-        public MongoReferencesModelV0(NamespaceId ns, BucketId bucket, IoHashKey key, BlobIdentifier blobIdentifier, byte[] blob, bool isFinalized, DateTime lastAccessTime)
+        public MongoReferencesModelV0(NamespaceId ns, BucketId bucket, IoHashKey key, BlobIdentifier blobIdentifier, byte[]? blob, bool isFinalized, DateTime lastAccessTime)
         {
             Ns = ns.ToString();
             Bucket = bucket.ToString();
@@ -238,7 +238,7 @@ namespace Horde.Storage.Implementation
         [BsonRequired]
         public DateTime LastAccessTime { get; set; }
 
-        public byte[] InlineBlob { get; set; }
+        public byte[]? InlineBlob { get; set; }
 
         public ObjectRecord ToObjectRecord()
         {
diff --git a/Engine/Source/Programs/Horde/HordeStorage/Horde.Storage/Implementation/Objects/ScyllaReferencesStore.cs b/Engine/Source/Programs/Horde/HordeStorage/Horde.Storage/Implementation/Objects/ScyllaReferencesStore.cs
--- a/Engine/Source/Programs/Horde/HordeStorage/Horde.Storage/Implementation/Objects/ScyllaReferencesStore.cs
+++ b/Engine/Source/Programs/Horde/HordeStorage/Horde.Storage/Implementation/Objects/ScyllaReferencesStore.cs
@@ -111,7 +111,7 @@ namespace Horde.Storage.Implementation
         }
 
 
-        public async Task Finalize(NamespaceId ns, BucketId bucket, IoHashKey name)
+        public async Task Finalize(NamespaceId ns, BucketId bucket, IoHashKey name, BlobIdentifier blobIdentifier)
         {
             using IScope _ = Tracer.Instance.StartActive("scylla.finalize");
 
