Horde: Simplified version of user edge perfoce server support
#skipci
#rnx

diff --git a/Engine/Source/Programs/Horde/Horde.Build/Models/Globals.cs b/Engine/Source/Programs/Horde/Horde.Build/Models/Globals.cs
--- a/Engine/Source/Programs/Horde/Horde.Build/Models/Globals.cs
+++ b/Engine/Source/Programs/Horde/Horde.Build/Models/Globals.cs
@@ -343,13 +343,20 @@ namespace HordeServer.Models
 		}
 
 		/// <summary>
-		/// Finds a perforce cluster with the given name
+		/// Finds a perforce cluster with the given name or that contains the provided server
 		/// </summary>
 		/// <param name="Name">Name of the cluster</param>
+		/// <param name="ServerAndPort">Find cluster which contains server</param>
 		/// <returns></returns>
-		public PerforceCluster? FindPerforceCluster(string? Name)
+		public PerforceCluster? FindPerforceCluster(string? Name, string? ServerAndPort = null)
 		{
 			List<PerforceCluster> Clusters = PerforceClusters;
+
+			if (ServerAndPort != null)
+			{
+				return Clusters.FirstOrDefault(x => x.Servers.FirstOrDefault( server => string.Equals(server.ServerAndPort, ServerAndPort, StringComparison.OrdinalIgnoreCase)) != null);
+			}
+
 			if (Clusters.Count == 0)
 			{
 				Clusters = DefaultClusters;
diff --git a/Engine/Source/Programs/Horde/Horde.Build/Perforce/Impl/PerforceService.cs b/Engine/Source/Programs/Horde/Horde.Build/Perforce/Impl/PerforceService.cs
--- a/Engine/Source/Programs/Horde/Horde.Build/Perforce/Impl/PerforceService.cs
+++ b/Engine/Source/Programs/Horde/Horde.Build/Perforce/Impl/PerforceService.cs
@@ -137,11 +137,11 @@ namespace HordeServer.Services
 			return User!;
 		}
 
-		async Task<PerforceCluster> GetClusterAsync(string? ClusterName)
+		async Task<PerforceCluster> GetClusterAsync(string? ClusterName, string? ServerAndPort = null)
 		{
 			Globals Globals = await CachedGlobals.GetCached();
 
-			PerforceCluster? Cluster = Globals.FindPerforceCluster(ClusterName);
+			PerforceCluster? Cluster = Globals.FindPerforceCluster(ClusterName, ServerAndPort);
 			if (Cluster == null)
 			{
 				throw new Exception($"Unknown Perforce cluster '{ClusterName}'");
@@ -200,6 +200,7 @@ namespace HordeServer.Services
 			return Repository;
 		}
 
+
 		/// <summary>
 		/// 
 		/// </summary>
@@ -338,17 +339,87 @@ namespace HordeServer.Services
 			CachedTicketInfo TicketInfo = await GetImpersonateCredential(Cluster, ImpersonateUser);
 			return CreateConnection(TicketInfo.Server, TicketInfo.UserName, TicketInfo.Ticket.Ticket);
 		}
+		
+		async Task<string> GetPortFromChange(PerforceCluster Cluster, int Change)
+		{
+			using IScope Scope = GlobalTracer.Instance.BuildSpan("PerforceService.GetPortFromChange").StartActive();
+			Scope.Span.SetTag("ClusterName", Cluster.Name);
+			Scope.Span.SetTag("Change", Change);
+
+			string? ClientPort = null;
+
+			using (P4.Repository Repository = await GetConnection(Cluster, NoClient: true))
+			{				
+				P4.Changelist Changelist = Repository.GetChangelist(Change);
+				if (Changelist == null || string.IsNullOrEmpty(Changelist.ClientId))
+				{
+					throw new Exception($"GetPortFromChange - Unable to get changelist for {Change}");
+				}
+
+				P4.Client Client = Repository.GetClient(Changelist.ClientId);
+
+				if (Client == null || string.IsNullOrEmpty(Client.ServerID))
+				{
+					throw new Exception($"GetPortFromChange - Unable to get client for {Change}");
+				}
+
+				List<string> Args = new List<string> { "-o", Client.ServerID };
+
+				using (P4.P4Command Command = new P4.P4Command(Repository, "server", true, Args.ToArray()))
+				{
+					P4.P4CommandResult Result = Command.Run();
+
+					if (!Result.Success || Result.TaggedOutput == null || Result.TaggedOutput.Count == 0)
+					{
+						throw new Exception($"Unable to get tagged output for {Client.ServerID}");
+					}
+
+					foreach (P4.TaggedObject TaggedObject in Result.TaggedOutput)
+					{
+						if (TaggedObject.TryGetValue("Address", out ClientPort))
+						{
+							break;
+						}
+					}
+				}
+			}
+
+			if (string.IsNullOrEmpty(ClientPort))
+			{
+				throw new Exception($"GetPortFromChange - Unable to get client port for {Change}");
+			}
+
+			return ClientPort;
+		}
 
 		[SuppressMessage("Microsoft.Reliability", "CA2000:DisposeObjectsBeforeLosingScope")]
-		async Task<P4.Repository> GetConnection(PerforceCluster Cluster, string? Stream = null, string? Username = null, bool ReadOnly = true, bool CreateChange = false, int? ClientFromChange = null, bool UseClientFromChange = false, bool UsePortFromChange = false, string? ClientId = null, bool NoClient = false)
+		async Task<P4.Repository> GetConnection(PerforceCluster Cluster, string? Stream = null, string? Username = null, bool ReadOnly = true, bool CreateChange = false, int? ClientFromChange = null, bool UseClientFromChange = false, int? UsePortFromChange = null, string? ClientId = null, bool NoClient = false)
 		{
 			using IScope Scope = GlobalTracer.Instance.BuildSpan("PerforceService.GetConnection").StartActive();
 			Scope.Span.SetTag("ClusterName", Cluster.Name);
 			Scope.Span.SetTag("Stream", Stream);
 			Scope.Span.SetTag("Username", Username);
 			Scope.Span.SetTag("NoClient", NoClient);
-			Scope.Span.SetTag("ClientId", ClientId);
-			
+			Scope.Span.SetTag("ClientId", ClientId);			
+
+			if (UsePortFromChange.HasValue)
+			{
+				try
+				{
+					string? ClientPort = await GetPortFromChange(Cluster, UsePortFromChange.Value);
+					PerforceCluster ClientCluster = await GetClusterAsync(null, ClientPort);
+					if (!string.Equals(ClientCluster.Name, Cluster.Name, StringComparison.OrdinalIgnoreCase))
+					{
+						Logger.LogInformation("Perforce: Overriding Cluster {Cluster} with Client Cluster {ClientCluster} for Change {Change}", Cluster.Name, ClientCluster.Name, UsePortFromChange.Value);
+						Cluster = ClientCluster;
+					}
+				}
+				catch (Exception Ex)
+				{
+					Logger.LogError(Ex, "Unable to get client port for changelist {UsePortFromChange}", UsePortFromChange);
+				}
+			}
+
 			P4.Repository Repository;
 			if (Username == null || !Cluster.CanImpersonate || Username.Equals(Cluster.ServiceAccount, StringComparison.OrdinalIgnoreCase))
 			{
@@ -369,7 +440,7 @@ namespace HordeServer.Services
 					}
 					else
 					{
-						P4.Client Client = GetOrCreateClient(Cluster.ServiceAccount, Repository, Stream, Username, ReadOnly, CreateChange, ClientFromChange, UseClientFromChange, UsePortFromChange);
+						P4.Client Client = GetOrCreateClient(Cluster.ServiceAccount, Repository, Stream, Username, ReadOnly, CreateChange, ClientFromChange, UseClientFromChange);
 						Repository.Connection.SetClient(Client.Name);
 
 					}
@@ -381,40 +452,7 @@ namespace HordeServer.Services
 						{
 							Repository.Connection.getP4Server().SetConnectionHost(ClientHost);
 						}
-					}
-
-					/*
-
-					This is disabled until https://jira.it.epicgames.com/servicedesk/customer/portal/1/ITH-144069 is resolved and edge server is sorted (required for things like deleting shelves)
-
-					string? ClientPort = null;
-
-					if (UsePortFromChange)
-					{
-						if (!Repository.Connection.Client.Name.StartsWith("horde-p4bridge-", StringComparison.OrdinalIgnoreCase))
-						{
-							using (P4.P4Command Cmd = new P4.P4Command(Repository, "info", true))
-							{
-								P4.P4CommandResult Results = Cmd.Run();
-
-								if (!Results.Success)
-								{
-									throw new Exception("Unable to get server info");
-								}
-
-								P4.TaggedObject? Tag = Results.TaggedOutput.Find(T => T.ContainsKey("changeServer"));
-
-								if (Tag != null)
-								{
-									if (!Tag.TryGetValue("changeServer", out ClientPort))
-									{
-										throw new Exception("Unable to get change server from tagged server info output");
-									}
-								}
-							}
-						}
-					}
-					*/
+					}					
 
 				}
 				catch
@@ -472,7 +510,7 @@ namespace HordeServer.Services
 			}
 		}
 
-		static P4.Client GetOrCreateClient(string? ServiceUserName, P4.Repository Repository, string? Stream, string? Username = null, bool ReadOnly = true, bool CreateChange = false, int? ClientFromChange = null, bool UseClientFromChange = false, bool UsePortFromChange = false)
+		static P4.Client GetOrCreateClient(string? ServiceUserName, P4.Repository Repository, string? Stream, string? Username = null, bool ReadOnly = true, bool CreateChange = false, int? ClientFromChange = null, bool UseClientFromChange = false)
 		{
 			using IScope Scope = GlobalTracer.Instance.BuildSpan("PerforceService.GetOrCreateClient").StartActive();
 			Scope.Span.SetTag("ServiceUserName", ServiceUserName);
@@ -500,12 +538,10 @@ namespace HordeServer.Services
 					throw new Exception($"Unable to get client for id {Changelist.ClientId}");
 				}
 
-
 				Stream = Client.Stream;
 				Username = Client.OwnerName;
 				ClientHost = Client.Host;
 
-
 			}
 
 			if (!UseClientFromChange)
@@ -554,7 +590,7 @@ namespace HordeServer.Services
 
 			if (Client == null)
 			{
-				throw new Exception($"Unable to create client for Stream:{Stream} Username:{Username} ReadOnly:{ReadOnly} CreateChange:{CreateChange} ClientFromChange:{ClientFromChange} UseClientFromChange:{UseClientFromChange} UsePortFromChange:{UsePortFromChange}");
+				throw new Exception($"Unable to create client for Stream:{Stream} Username:{Username} ReadOnly:{ReadOnly} CreateChange:{CreateChange} ClientFromChange:{ClientFromChange} UseClientFromChange:{UseClientFromChange}");
 			}
 
 			return Client;
@@ -1157,7 +1193,7 @@ namespace HordeServer.Services
 				throw new Exception($"Unable to get owner for shelved change {ShelvedChange}");
 			}
 
-			using (P4.Repository Repository = await GetConnection(Cluster, ReadOnly: false, ClientFromChange: ShelvedChange, UseClientFromChange: false, Username: ChangeOwner))
+			using (P4.Repository Repository = await GetConnection(Cluster, ReadOnly: false, UsePortFromChange: ShelvedChange, ClientFromChange: ShelvedChange, UseClientFromChange: false, Username: ChangeOwner))
 			{
 				return ReshelveChange(Repository, ShelvedChange);
 			}
@@ -1190,7 +1226,7 @@ namespace HordeServer.Services
 					throw new Exception($"Unable to get shelved changelist client id for delete: {ShelvedChange}");
 				}
 
-				using (P4.Repository ClientRepository = await GetConnection(Cluster, Changelist.Stream, Changelist.OwnerName, ClientId: ClientId, UseClientFromChange: true, UsePortFromChange: true))
+				using (P4.Repository ClientRepository = await GetConnection(Cluster, Changelist.Stream, Changelist.OwnerName, ClientId: ClientId, UseClientFromChange: true, UsePortFromChange: ShelvedChange))
 				{
 
 					if (Changelist.Shelved)
@@ -1213,6 +1249,7 @@ namespace HordeServer.Services
 			try
 			{
 				PerforceCluster Cluster = await GetClusterAsync(ClusterName);
+
 				using (P4.Repository Repository = await GetConnection(Cluster, NoClient: true))
 				{
 					P4.Changelist Changelist = Repository.GetChangelist(Change);
@@ -1220,7 +1257,7 @@ namespace HordeServer.Services
 					Repository.Connection.Disconnect();
 
 					// the client must exist for the change list, otherwise will fail (for example, CreateNewChangeAsync deletes the client before returning)
-					using (P4.Repository UpdateRepository = await GetConnection(Cluster, ClientFromChange: Change, UseClientFromChange: true, Username: Changelist.OwnerName))
+					using (P4.Repository UpdateRepository = await GetConnection(Cluster, ClientFromChange: Change, UseClientFromChange: true, UsePortFromChange: Change, Username: Changelist.OwnerName))
 					{
 						P4.Changelist UpdatedChangelist = UpdateRepository.GetChangelist(Change);
 						UpdatedChangelist.Description = Description;
@@ -1462,7 +1499,7 @@ namespace HordeServer.Services
 
 					Repository.Connection.Disconnect();
 
-					using (P4.Repository SubmitRepository = await GetConnection(Cluster, ReadOnly: false, ClientFromChange: Change, UseClientFromChange: true, Username: Changelist.OwnerName))
+					using (P4.Repository SubmitRepository = await GetConnection(Cluster, ReadOnly: false, ClientFromChange: Change, UseClientFromChange: true, UsePortFromChange: Change, Username: Changelist.OwnerName))
 					{
 						// we might not need a client here, possibly -e below facilitates this, check!
 						using (P4.P4Command SubmitCommand = new P4.P4Command(SubmitRepository, "submit", null, true, new string[] { "-e", Change.ToString(CultureInfo.InvariantCulture) }))
