Horde: Fix for unresrtricted workspaces
#skipci
#rnx

diff --git a/Engine/Source/Programs/Horde/Horde.Build/Perforce/Impl/PerforceService.cs b/Engine/Source/Programs/Horde/Horde.Build/Perforce/Impl/PerforceService.cs
--- a/Engine/Source/Programs/Horde/Horde.Build/Perforce/Impl/PerforceService.cs
+++ b/Engine/Source/Programs/Horde/Horde.Build/Perforce/Impl/PerforceService.cs
@@ -340,7 +340,7 @@ namespace HordeServer.Services
 			return CreateConnection(TicketInfo.Server, TicketInfo.UserName, TicketInfo.Ticket.Ticket);
 		}
 		
-		async Task<string> GetPortFromChange(PerforceCluster Cluster, int Change)
+		async Task<string?> GetPortFromChange(PerforceCluster Cluster, int Change)
 		{
 			using IScope Scope = GlobalTracer.Instance.BuildSpan("PerforceService.GetPortFromChange").StartActive();
 			Scope.Span.SetTag("ClusterName", Cluster.Name);
@@ -358,11 +358,17 @@ namespace HordeServer.Services
 
 				P4.Client Client = Repository.GetClient(Changelist.ClientId);
 
-				if (Client == null || string.IsNullOrEmpty(Client.ServerID))
+				if (Client == null)
 				{
 					throw new Exception($"GetPortFromChange - Unable to get client for {Change}");
 				}
 
+				// Check whether not restricted to a specific server
+				if (string.IsNullOrEmpty(Client.ServerID))
+				{
+					return null;
+				}
+
 				List<string> Args = new List<string> { "-o", Client.ServerID };
 
 				using (P4.P4Command Command = new P4.P4Command(Repository, "server", true, Args.ToArray()))
@@ -400,23 +406,30 @@ namespace HordeServer.Services
 			Scope.Span.SetTag("Stream", Stream);
 			Scope.Span.SetTag("Username", Username);
 			Scope.Span.SetTag("NoClient", NoClient);
-			Scope.Span.SetTag("ClientId", ClientId);			
+			Scope.Span.SetTag("ClientId", ClientId);
 
 			if (UsePortFromChange.HasValue)
 			{
 				try
 				{
 					string? ClientPort = await GetPortFromChange(Cluster, UsePortFromChange.Value);
-					PerforceCluster ClientCluster = await GetClusterAsync(null, ClientPort);
-					if (!string.Equals(ClientCluster.Name, Cluster.Name, StringComparison.OrdinalIgnoreCase))
+					if (ClientPort != null)
+					{
+						PerforceCluster ClientCluster = await GetClusterAsync(null, ClientPort);
+						if (!string.Equals(ClientCluster.Name, Cluster.Name, StringComparison.OrdinalIgnoreCase))
+						{
+							Logger.LogInformation("Perforce: Overriding Cluster {Cluster} with Client Cluster {ClientCluster} for server restricted change {Change}", Cluster.Name, ClientCluster.Name, UsePortFromChange.Value);
+							Cluster = ClientCluster;
+						}
+					}
+					else
 					{
-						Logger.LogInformation("Perforce: Overriding Cluster {Cluster} with Client Cluster {ClientCluster} for Change {Change}", Cluster.Name, ClientCluster.Name, UsePortFromChange.Value);
-						Cluster = ClientCluster;
+						Logger.LogInformation("Change {UsePortFromChange} is not restricted to a server", UsePortFromChange.Value);
 					}
 				}
 				catch (Exception Ex)
 				{
-					Logger.LogError(Ex, "Unable to get client port for changelist {UsePortFromChange}", UsePortFromChange);
+					Logger.LogError(Ex, "Unable to get client port for changelist {Change} using cluster {ClusterName}", UsePortFromChange, Cluster.Name);
 				}
 			}
 
