Horde: Include jobs in find timings endpoint

To further reduce chattiness for the timings panel.

diff --git a/Engine/Source/Programs/Horde/Horde.Build.Tests/JobsControllerDbTest.cs b/Engine/Source/Programs/Horde/Horde.Build.Tests/JobsControllerDbTest.cs
--- a/Engine/Source/Programs/Horde/Horde.Build.Tests/JobsControllerDbTest.cs
+++ b/Engine/Source/Programs/Horde/Horde.Build.Tests/JobsControllerDbTest.cs
@@ -69,6 +69,7 @@ namespace HordeServerTests
 	        object Obj = (await JobsController.FindJobTimingsAsync(Fixture.Stream!.Id.ToString(), Templates)).Value!;
 	        FindJobTimingsResponse Res = (Obj as FindJobTimingsResponse)!;
 	        Assert.AreEqual(1, Res.Timings.Count);
+	        Assert.AreEqual(Job.Name, Res.Timings[Job.Id.ToString()].Job.Name);
         }
     }
 }
\ No newline at end of file
diff --git a/Engine/Source/Programs/Horde/Horde.Build/Api/Jobs.cs b/Engine/Source/Programs/Horde/Horde.Build/Api/Jobs.cs
--- a/Engine/Source/Programs/Horde/Horde.Build/Api/Jobs.cs
+++ b/Engine/Source/Programs/Horde/Horde.Build/Api/Jobs.cs
@@ -383,6 +383,11 @@ namespace HordeServer.Api
 	/// </summary>
 	public class GetJobTimingResponse
 	{
+		/// <summary>
+		/// The job object
+		/// </summary>
+		public IJob Job { get; set; }
+		
 		/// <summary>
 		/// Timing info for each step
 		/// </summary>
@@ -396,10 +401,12 @@ namespace HordeServer.Api
 		/// <summary>
 		/// Constructor
 		/// </summary>
+		/// <param name="Job">The job object</param>
 		/// <param name="Steps">Timing info for each steps</param>
 		/// <param name="Labels">Timing info for each label</param>
-		public GetJobTimingResponse(Dictionary<string, GetStepTimingInfoResponse> Steps, List<GetLabelTimingInfoResponse> Labels)
+		public GetJobTimingResponse(IJob Job, Dictionary<string, GetStepTimingInfoResponse> Steps, List<GetLabelTimingInfoResponse> Labels)
 		{
+			this.Job = Job;
 			this.Steps = Steps;
 			this.Labels = Labels;
 		}
diff --git a/Engine/Source/Programs/Horde/Horde.Build/Controllers/JobsController.cs b/Engine/Source/Programs/Horde/Horde.Build/Controllers/JobsController.cs
--- a/Engine/Source/Programs/Horde/Horde.Build/Controllers/JobsController.cs
+++ b/Engine/Source/Programs/Horde/Horde.Build/Controllers/JobsController.cs
@@ -588,7 +588,7 @@ namespace HordeServer.Controllers
 				Labels.Add(new GetLabelTimingInfoResponse(Label, TimingInfo));
 			}
 
-			return new GetJobTimingResponse(Steps, Labels);
+			return new GetJobTimingResponse(Job, Steps, Labels);
 		}
 		
 		/// <summary>
