Horde Storage - Handle exception when no incremental state is available when attempting to build the snapshot (is nothing to do as there is no state to snapshot)

diff --git a/Engine/Source/Programs/Horde/HordeStorage/Horde.Storage/Implementation/TransactionLog/ReplicationSnapshotService.cs b/Engine/Source/Programs/Horde/HordeStorage/Horde.Storage/Implementation/TransactionLog/ReplicationSnapshotService.cs
--- a/Engine/Source/Programs/Horde/HordeStorage/Horde.Storage/Implementation/TransactionLog/ReplicationSnapshotService.cs
+++ b/Engine/Source/Programs/Horde/HordeStorage/Horde.Storage/Implementation/TransactionLog/ReplicationSnapshotService.cs
@@ -74,8 +74,16 @@ namespace Horde.Storage.Implementation
                     }
                 }
                 ReplicationLogSnapshotBuilder builder = ActivatorUtilities.CreateInstance<ReplicationLogSnapshotBuilder>(_provider);
-                BlobIdentifier snapshotBlob = await builder.BuildSnapshot(ns, _settings.CurrentValue.SnapshotStorageNamespace, _cancellationTokenSource.Token);
-                _logger.Information("Snapshot built for {Namespace} with id {Id}", ns, snapshotBlob);
+                try
+                {
+                    BlobIdentifier snapshotBlob = await builder.BuildSnapshot(ns, _settings.CurrentValue.SnapshotStorageNamespace, _cancellationTokenSource.Token);
+                    _logger.Information("Snapshot built for {Namespace} with id {Id}", ns, snapshotBlob);
+
+                }
+                catch (IncrementalLogNotAvailableException)
+                {
+                    _logger.Error("Unable to generate a snapshot for {Namespace} as there was no incremental state available", ns);
+                }
 
                 ran = true;
             }, _settings.CurrentValue.MaxCountOfNamespacesToSnapshotInParallel, _cancellationTokenSource.Token);
