Horde Storage - Fixes to refs replication
We were not correctly verifying that our blob download was succesful.
Fixed issue were we were attempting to download the blob from the incorrect namespace. Also updated the tests which had incorrect mocks setup for this as well.
Fixed that our service accounts did not have access to download the snapshot blobs.
Added support for requiring values to be of a particular value in the namespace claim mapping using A=B

diff --git a/Engine/Source/Programs/Horde/HordeStorage/Horde.Storage/Implementation/Replication/RefsReplicator.cs b/Engine/Source/Programs/Horde/HordeStorage/Horde.Storage/Implementation/Replication/RefsReplicator.cs
--- a/Engine/Source/Programs/Horde/HordeStorage/Horde.Storage/Implementation/Replication/RefsReplicator.cs
+++ b/Engine/Source/Programs/Horde/HordeStorage/Horde.Storage/Implementation/Replication/RefsReplicator.cs
@@ -16,6 +16,7 @@ using Datadog.Trace;
 using EpicGames.Horde.Storage;
 using Horde.Storage.Controllers;
 using Horde.Storage.Implementation.TransactionLog;
+using Jupiter.Common;
 using Jupiter.Common.Implementation;
 using Jupiter.Implementation;
 using Microsoft.AspNetCore.Mvc;
@@ -198,7 +199,7 @@ namespace Horde.Storage.Implementation
 
                         // if we fail to replicate incrementally we revert to using a snapshot
                         haveAttemptedSnapshot = true;
-                        (string eventBucket, Guid eventId, int countOfEventsReplicated) = await ReplicateFromSnapshot(ns, replicationToken, useSnapshotException.SnapshotBlob);
+                        (string eventBucket, Guid eventId, int countOfEventsReplicated) = await ReplicateFromSnapshot(ns, replicationToken, useSnapshotException.SnapshotBlob, INamespacePolicyResolver.JupiterInternalNamespace);
                         countOfReplicationsDone += countOfEventsReplicated;
 
                         // resume from these new events instead
@@ -223,7 +224,7 @@ namespace Horde.Storage.Implementation
             return hasRun;
         }
 
-        private async Task<(string, Guid, int)> ReplicateFromSnapshot(NamespaceId ns, CancellationToken cancellationToken, BlobIdentifier? snapshotBlob = null)
+        private async Task<(string, Guid, int)> ReplicateFromSnapshot(NamespaceId ns, CancellationToken cancellationToken, BlobIdentifier? snapshotBlob = null, NamespaceId? blobNamespace = null)
         {
             // determine latest snapshot if no specific blob was specified
             if (snapshotBlob == null)
@@ -244,14 +245,15 @@ namespace Horde.Storage.Implementation
                 }
 
                 snapshotBlob = snapshotInfo.SnapshotBlob;
+                blobNamespace = snapshotInfo.BlobNamespace;
             }
 
             // process snapshot
             // fetch the snapshot from the remote blob store
-            HttpRequestMessage request = BuildHttpRequest(HttpMethod.Get, $"api/v1/blobs/{ns}/{snapshotBlob}");
+            HttpRequestMessage request = BuildHttpRequest(HttpMethod.Get, $"api/v1/blobs/{blobNamespace}/{snapshotBlob}");
             HttpResponseMessage response = await _httpClient.SendAsync(request, cancellationToken);
+            response.EnsureSuccessStatusCode();
             Stream blobStream = await response.Content.ReadAsStreamAsync(cancellationToken);
-
             ReplicationLogSnapshot snapshot = await ReplicationLogSnapshot.DeserializeSnapshot(blobStream);
 
             if (!snapshot.LastEvent.HasValue)
diff --git a/Engine/Source/Programs/Horde/HordeStorage/Horde.Storage/Tests/Functional/Objects/RefReplicatorTests.cs b/Engine/Source/Programs/Horde/HordeStorage/Horde.Storage/Tests/Functional/Objects/RefReplicatorTests.cs
--- a/Engine/Source/Programs/Horde/HordeStorage/Horde.Storage/Tests/Functional/Objects/RefReplicatorTests.cs
+++ b/Engine/Source/Programs/Horde/HordeStorage/Horde.Storage/Tests/Functional/Objects/RefReplicatorTests.cs
@@ -207,7 +207,7 @@ namespace Horde.Storage.FunctionalTests.Replication
             // after processing a snapshot it will attempt to incrementally replicate from there, which should be empty
             handler.SetupRequest($"http://localhost/api/v1/replication-log/incremental/{TestNamespace}?lastBucket={replicationEvents.Last().TimeBucket}&lastEvent={replicationEvents.Last().EventId}").ReturnsResponse(JsonConvert.SerializeObject(new ReplicationLogEvents(new List<ReplicationLogEvent>())), "application/json");
 
-            handler.SetupRequest($"http://localhost/api/v1/blobs/{TestNamespace}/{snapshotBlob}").ReturnsResponse(snapshotContent, "application/octet-stream").Verifiable();
+            handler.SetupRequest($"http://localhost/api/v1/blobs/{SnapshotNamespace}/{snapshotBlob}").ReturnsResponse(snapshotContent, "application/octet-stream").Verifiable();
 
             foreach (BlobIdentifier blob in blobs.Keys)
             {
@@ -295,7 +295,7 @@ namespace Horde.Storage.FunctionalTests.Replication
             // after processing the incremental events there is nothing more to find
             handler.SetupRequest($"http://localhost/api/v1/replication-log/incremental/{TestNamespace}?lastBucket={incrementalEvents.Last().TimeBucket}&lastEvent={incrementalEvents.Last().EventId}").ReturnsResponse(JsonConvert.SerializeObject(new ReplicationLogEvents(new List<ReplicationLogEvent>())), "application/json");
 
-            handler.SetupRequest($"http://localhost/api/v1/blobs/{TestNamespace}/{snapshotBlob}").ReturnsResponse(snapshotContent, "application/octet-stream").Verifiable();
+            handler.SetupRequest($"http://localhost/api/v1/blobs/{SnapshotNamespace}/{snapshotBlob}").ReturnsResponse(snapshotContent, "application/octet-stream").Verifiable();
 
             foreach (BlobIdentifier blob in blobs.Keys)
             {
@@ -384,7 +384,7 @@ namespace Horde.Storage.FunctionalTests.Replication
             // after processing the snapshot we do not replicate anything more
             handler.SetupRequest($"http://localhost/api/v1/replication-log/incremental/{TestNamespace}?lastBucket={snapshotEvents.Last().TimeBucket}&lastEvent={snapshotEvents.Last().EventId}").ReturnsResponse(JsonConvert.SerializeObject(new ReplicationLogEvents(new List<ReplicationLogEvent>())), "application/json");
 
-            handler.SetupRequest($"http://localhost/api/v1/blobs/{TestNamespace}/{snapshotBlob}").ReturnsResponse(snapshotContent, "application/octet-stream").Verifiable();
+            handler.SetupRequest($"http://localhost/api/v1/blobs/{SnapshotNamespace}/{snapshotBlob}").ReturnsResponse(snapshotContent, "application/octet-stream").Verifiable();
 
             foreach (BlobIdentifier blob in blobs.Keys)
             {
diff --git a/Engine/Source/Programs/Horde/HordeStorage/Jupiter.Common/Authentication/NamespaceAuthorizationHandler.cs b/Engine/Source/Programs/Horde/HordeStorage/Jupiter.Common/Authentication/NamespaceAuthorizationHandler.cs
--- a/Engine/Source/Programs/Horde/HordeStorage/Jupiter.Common/Authentication/NamespaceAuthorizationHandler.cs
+++ b/Engine/Source/Programs/Horde/HordeStorage/Jupiter.Common/Authentication/NamespaceAuthorizationHandler.cs
@@ -39,8 +39,20 @@ namespace Jupiter
                 if (expectedClaim == "*")
                 {
                     context.Succeed(requirement);
+                    continue;
                 }
 
+                if (expectedClaim.Contains('='))
+                {
+                    int separatorIndex = expectedClaim.IndexOf('=');
+                    string claimName = expectedClaim.Substring(0, separatorIndex);
+                    string claimValue = expectedClaim.Substring(separatorIndex + 1);
+                    if (context.User.HasClaim(claim => claim.Type == claimName && claim.Value == claimValue))
+                    {
+                        context.Succeed(requirement);
+                        continue;
+                    }
+                }
                 if (context.User.HasClaim(claim => claim.Type == expectedClaim))
                 {
                     context.Succeed(requirement);
diff --git a/Engine/Source/Programs/Horde/HordeStorage/Jupiter.Common/Implementation/NamespacePolicyResolver.cs b/Engine/Source/Programs/Horde/HordeStorage/Jupiter.Common/Implementation/NamespacePolicyResolver.cs
--- a/Engine/Source/Programs/Horde/HordeStorage/Jupiter.Common/Implementation/NamespacePolicyResolver.cs
+++ b/Engine/Source/Programs/Horde/HordeStorage/Jupiter.Common/Implementation/NamespacePolicyResolver.cs
@@ -31,8 +31,8 @@ namespace Jupiter.Common
             {
                 return new NamespaceSettings.PerNamespaceSettings()
                 {
-                    // the internal jupiter namespace is sensitive so only admins should have access, but we can still use the default storage pool
-                    Claims = new string[] { "Admin", "admin" },
+                    // we grant the full storage access claim access to the internal namespace, this is typically set of admins and service accounts that can access everything
+                    Claims = new string[] {"Storage=full" },
                     StoragePool = ""
                 };
             }
