Horde Storage - Support storing the replication state in scylla rather then the filesystem. Allows us to remove the extra filesystem on worker nodes which is easier for deployment.

Also fixed issues with the replicator not decrementing the number of in flight requests (logging only issue) when requests failed.
Some error handling when failing to replicate blobs in the refs replication.
Lastly I did some improvements to how we read the snapshot blobs to not read it into memory which improves performance over slow connections

diff --git a/Engine/Source/Programs/Horde/HordeStorage/Horde.Storage/Implementation/Replication/RefsReplicator.cs b/Engine/Source/Programs/Horde/HordeStorage/Horde.Storage/Implementation/Replication/RefsReplicator.cs
--- a/Engine/Source/Programs/Horde/HordeStorage/Horde.Storage/Implementation/Replication/RefsReplicator.cs
+++ b/Engine/Source/Programs/Horde/HordeStorage/Horde.Storage/Implementation/Replication/RefsReplicator.cs
@@ -8,7 +8,6 @@ using System.Net;
 using System.Net.Http;
 using System.Runtime.CompilerServices;
 using System.Text;
-using System.Text.Json;
 using System.Threading;
 using System.Threading.Tasks;
 using Dasync.Collections;
@@ -58,13 +57,25 @@ namespace Horde.Storage.Implementation
             DirectoryInfo stateRoot = new DirectoryInfo(replicationSettings.CurrentValue.GetStateRoot());
             _stateFile = new FileInfo(Path.Combine(stateRoot.FullName, stateFileName));
 
-            if (_stateFile.Exists)
+            ReplicatorState? replicatorState = replicationLog.GetReplicatorState(replicatorSettings.NamespaceToReplicate, _name).Result;
+            if (replicatorState == null)
             {
-                _refsState = ReadState(_stateFile)!;
+                if (_stateFile.Exists)
+                {
+                    _refsState = ReadState(_stateFile)!;
+                }
+                else
+                {
+                    _refsState = new RefsState();
+                }
             }
             else
             {
-                _refsState = new RefsState();
+                _refsState = new RefsState()
+                {
+                    LastBucket = replicatorState.LastBucket,
+                    LastEvent = replicatorState.LastEvent,
+                };
             }
 
             Info = new IReplicator.ReplicatorInfo(replicatorSettings.ReplicatorName, replicatorSettings.NamespaceToReplicate, _refsState);
@@ -89,11 +100,17 @@ namespace Horde.Storage.Implementation
             return state;
         }
 
-        private static void SaveState(FileInfo stateFile, RefsState newState)
+        private void SaveState(FileInfo stateFile, RefsState newState)
         {
             using StreamWriter writer = stateFile.CreateText();
             JsonSerializer serializer = new JsonSerializer();
             serializer.Serialize(writer, newState);
+
+            _replicationLog.UpdateReplicatorState(Info.NamespaceToReplicate, _name, new ReplicatorState
+            {
+                LastBucket = newState.LastBucket,
+                LastEvent = newState.LastEvent
+            }).Wait();
         }
 
         public async Task<bool> TriggerNewReplications()
@@ -251,7 +268,7 @@ namespace Horde.Storage.Implementation
             // process snapshot
             // fetch the snapshot from the remote blob store
             HttpRequestMessage request = BuildHttpRequest(HttpMethod.Get, $"api/v1/blobs/{blobNamespace}/{snapshotBlob}");
-            HttpResponseMessage response = await _httpClient.SendAsync(request, cancellationToken);
+            HttpResponseMessage response = await _httpClient.SendAsync(request, HttpCompletionOption.ResponseHeadersRead, cancellationToken);
             response.EnsureSuccessStatusCode();
             Stream blobStream = await response.Content.ReadAsStreamAsync(cancellationToken);
             ReplicationLogSnapshot snapshot = await ReplicationLogSnapshot.DeserializeSnapshot(blobStream);
@@ -268,22 +285,29 @@ namespace Horde.Storage.Implementation
             int maxParallelism = _replicatorSettings.MaxParallelReplications != -1 ? _replicatorSettings.MaxParallelReplications : 0;
             await snapshot.LiveObjects.ParallelForEachAsync(async snapshotLiveObject =>
             {
-                using IScope scope = Tracer.Instance.StartActive("replicator.replicate_op_snapshot");
-                scope.Span.ResourceName = $"{ns}.{snapshotLiveObject.Blob}";
-                Interlocked.Increment(ref countOfObjectsCurrentlyReplicating);
-                LogReplicationHeartbeat(countOfObjectsCurrentlyReplicating);
-
-                Info.CountOfRunningReplications = countOfObjectsCurrentlyReplicating;
-                Info.LastRun = DateTime.Now;
-
-                await ReplicateOp(ns, snapshotLiveObject.Blob, cancellationToken);
-                // TODO: Avoid adding to the replication log as we will get into infinite recursion if we do with the remote site
-                // we could add to the replication log only if the op was missing, but that could cause issue for a 3rd site replicating from us
-                // and the 3rd site is the whole reason we even add this to replication log in the first place
-                //await AddToReplicationLog(ns, snapshotLiveObject.Bucket, snapshotLiveObject.Key, snapshotLiveObject.Blob);
-                
-                Interlocked.Increment(ref countOfObjectsReplicated);
-                Interlocked.Decrement(ref countOfObjectsCurrentlyReplicating);
+                try
+                {
+                    using IScope scope = Tracer.Instance.StartActive("replicator.replicate_op_snapshot");
+                    scope.Span.ResourceName = $"{ns}.{snapshotLiveObject.Blob}";
+                    Interlocked.Increment(ref countOfObjectsCurrentlyReplicating);
+                    LogReplicationHeartbeat(countOfObjectsCurrentlyReplicating);
+
+                    Info.CountOfRunningReplications = countOfObjectsCurrentlyReplicating;
+                    Info.LastRun = DateTime.Now;
+
+                    await ReplicateOp(ns, snapshotLiveObject.Blob, cancellationToken);
+                    // TODO: Avoid adding to the replication log as we will get into infinite recursion if we do with the remote site
+                    // we could add to the replication log only if the op was missing, but that could cause issue for a 3rd site replicating from us
+                    // and the 3rd site is the whole reason we even add this to replication log in the first place
+                    //await AddToReplicationLog(ns, snapshotLiveObject.Bucket, snapshotLiveObject.Key, snapshotLiveObject.Blob);
+
+                }
+                finally
+                {
+                    Interlocked.Increment(ref countOfObjectsReplicated);
+                    Interlocked.Decrement(ref countOfObjectsCurrentlyReplicating);
+                }
+
             }, maxParallelism, cancellationToken);
 
             // snapshot processed, we proceed to reading the incremental state
@@ -428,8 +452,13 @@ namespace Horde.Storage.Implementation
                 blobReplicationTasks[i] = Task.Run(async () =>
                 {
                     HttpRequestMessage blobRequest = BuildHttpRequest(HttpMethod.Get, $"api/v1/blobs/{ns}/{blobToReplicate}");
-                    HttpResponseMessage blobResponse = await _httpClient.SendAsync(blobRequest, cancellationToken);
-                    
+                    HttpResponseMessage blobResponse = await _httpClient.SendAsync(blobRequest, HttpCompletionOption.ResponseHeadersRead, cancellationToken);
+
+                    if (blobResponse.StatusCode != HttpStatusCode.OK)
+                    {
+                        _logger.Error("Bad http response when replicating {Blob} in {Namespace} . Status code: {StatusCode}", blobToReplicate, ns, blobResponse.StatusCode);
+                        return;
+                    }
                     using FilesystemBufferedPayload payload = await FilesystemBufferedPayload.Create(await blobResponse.Content.ReadAsStreamAsync(cancellationToken));
                     await _blobService.PutObject(ns, payload, blobToReplicate);
                 });
diff --git a/Engine/Source/Programs/Horde/HordeStorage/Horde.Storage/Implementation/TransactionLog/IReplicationLog.cs b/Engine/Source/Programs/Horde/HordeStorage/Horde.Storage/Implementation/TransactionLog/IReplicationLog.cs
--- a/Engine/Source/Programs/Horde/HordeStorage/Horde.Storage/Implementation/TransactionLog/IReplicationLog.cs
+++ b/Engine/Source/Programs/Horde/HordeStorage/Horde.Storage/Implementation/TransactionLog/IReplicationLog.cs
@@ -2,10 +2,8 @@
 
 using System;
 using System.Collections.Generic;
-using System.Reflection.Metadata;
 using System.Threading.Tasks;
 using EpicGames.Horde.Storage;
-using Horde.Storage.Implementation.TransactionLog;
 using Jupiter.Implementation;
 using Newtonsoft.Json;
 
@@ -22,6 +20,16 @@ namespace Horde.Storage.Implementation
         Task AddSnapshot(SnapshotInfo snapshotHeader);
         Task<SnapshotInfo?> GetLatestSnapshot(NamespaceId ns);
         IAsyncEnumerable<SnapshotInfo> GetSnapshots(NamespaceId ns);
+
+
+        Task UpdateReplicatorState(NamespaceId ns, string replicatorName, ReplicatorState newState);
+        Task<ReplicatorState?> GetReplicatorState(NamespaceId ns, string replicatorName);
+    }
+
+    public class ReplicatorState
+    {
+        public string? LastBucket { get; set; }
+        public Guid? LastEvent { get; set; }
     }
 
     public class SnapshotInfo
diff --git a/Engine/Source/Programs/Horde/HordeStorage/Horde.Storage/Implementation/TransactionLog/MemoryReplicationLog.cs b/Engine/Source/Programs/Horde/HordeStorage/Horde.Storage/Implementation/TransactionLog/MemoryReplicationLog.cs
--- a/Engine/Source/Programs/Horde/HordeStorage/Horde.Storage/Implementation/TransactionLog/MemoryReplicationLog.cs
+++ b/Engine/Source/Programs/Horde/HordeStorage/Horde.Storage/Implementation/TransactionLog/MemoryReplicationLog.cs
@@ -17,6 +17,8 @@ namespace Horde.Storage.Implementation
         private readonly ConcurrentDictionary<NamespaceId, SortedList<string, SortedList<TimeUuid, ReplicationLogEvent>>> _replicationEvents = new();
         private readonly ConcurrentDictionary<NamespaceId, List<SnapshotInfo>>  _snapshots = new();
 
+        private readonly ConcurrentDictionary<NamespaceId, ConcurrentDictionary<string, ReplicatorState>>  _replicatorState = new();
+
         public IAsyncEnumerable<NamespaceId> GetNamespaces()
         {
             return _replicationEvents.Keys.ToAsyncEnumerable();
@@ -172,6 +174,35 @@ namespace Horde.Storage.Implementation
             }
 
         }
+
+        public Task UpdateReplicatorState(NamespaceId ns, string replicatorName,
+            ReplicatorState newState)
+        {
+            _replicatorState.AddOrUpdate(ns,
+                _ =>
+                new ConcurrentDictionary<string, ReplicatorState>
+                {
+                    [replicatorName] = newState
+                },
+                (_, states) =>
+                {
+                    states[replicatorName] = newState;
+                    return states;
+                });
+
+            return Task.CompletedTask;
+        }
+
+        public Task<ReplicatorState?> GetReplicatorState(NamespaceId ns, string replicatorName)
+        {
+            if (_replicatorState.TryGetValue(ns, out ConcurrentDictionary<string, ReplicatorState>? replicationState))
+            {
+                replicationState.TryGetValue(replicatorName, out ReplicatorState? state);
+                return Task.FromResult(state);
+            }
+
+            return Task.FromResult<ReplicatorState?>(null);
+        }
     }
 
     public static class DateTimeUtils
diff --git a/Engine/Source/Programs/Horde/HordeStorage/Horde.Storage/Implementation/TransactionLog/ScyllaReplicationLog.cs b/Engine/Source/Programs/Horde/HordeStorage/Horde.Storage/Implementation/TransactionLog/ScyllaReplicationLog.cs
--- a/Engine/Source/Programs/Horde/HordeStorage/Horde.Storage/Implementation/TransactionLog/ScyllaReplicationLog.cs
+++ b/Engine/Source/Programs/Horde/HordeStorage/Horde.Storage/Implementation/TransactionLog/ScyllaReplicationLog.cs
@@ -52,6 +52,15 @@ namespace Horde.Storage.Implementation
                 PRIMARY KEY ((namespace))
             );"  
             ));
+
+            _session.Execute(new SimpleStatement(@"CREATE TABLE IF NOT EXISTS replication_state (
+                namespace varchar,
+                name varchar,
+                last_bucket varchar,
+                last_event uuid,
+                PRIMARY KEY ((namespace), name)
+            );"  
+            ));
         }
 
         public async IAsyncEnumerable<NamespaceId> GetNamespaces()
@@ -263,6 +272,23 @@ namespace Horde.Storage.Implementation
                 yield return new SnapshotInfo(ns, new NamespaceId(snapshot.BlobNamespace), snapshot.BlobSnapshot.AsBlobIdentifier(), snapshot.Id.GetDate().DateTime);
             }
         }
+
+        public async Task UpdateReplicatorState(NamespaceId ns, string replicatorName, ReplicatorState newState)
+        {
+            if (newState.LastBucket == null || newState.LastEvent == null)
+            {
+                // no point in saving a invalid state
+                return;
+            }
+            await _mapper.UpdateAsync<ScyllaReplicationState>(new ScyllaReplicationState(ns, replicatorName, newState.LastBucket, newState.LastEvent.Value));
+        }
+
+        public async Task<ReplicatorState?> GetReplicatorState(NamespaceId ns, string name)
+        {
+            ScyllaReplicationState? replicationState = await _mapper.FirstOrDefaultAsync<ScyllaReplicationState>("WHERE namespace = ? AND name = ?", ns.ToString(), name);
+
+            return replicationState?.ToReplicatorState();
+        }
     }
 
     
@@ -379,4 +405,46 @@ namespace Horde.Storage.Implementation
         [Cassandra.Mapping.Attributes.PartitionKey]
         public string Namespace { get;set; }
     }
+
+    [Cassandra.Mapping.Attributes.Table("replication_state")]
+    class ScyllaReplicationState
+    {
+        public ScyllaReplicationState()
+        {
+            Namespace = null!;
+            Name = null!;
+            LastBucket = null!;
+        }
+
+        public ScyllaReplicationState(NamespaceId ns, string name, string lastBucket, Guid lastEvent)
+        {
+            Namespace = ns.ToString();
+            Name = name;
+            LastBucket = lastBucket;
+            LastEvent = lastEvent;
+        }
+
+
+        [Cassandra.Mapping.Attributes.PartitionKey]
+        public string Namespace { get; set; }
+
+
+        [Cassandra.Mapping.Attributes.ClusteringKey]
+        public string Name { get; set; }
+
+        [Cassandra.Mapping.Attributes.Column("last_bucket")]
+        public string LastBucket { get; set; }
+
+        [Cassandra.Mapping.Attributes.Column("last_event")]
+        public Guid LastEvent { get; set; }
+
+        public ReplicatorState ToReplicatorState()
+        {
+            return new ReplicatorState
+            {
+                LastBucket = LastBucket,
+                LastEvent = LastEvent
+            };
+        }
+    }
 }
