Horde: Device service checkout expiration
#rnx
#skipci

diff --git a/Engine/Source/Programs/Horde/Horde.Build/Devices/IDeviceCollection.cs b/Engine/Source/Programs/Horde/Horde.Build/Devices/IDeviceCollection.cs
--- a/Engine/Source/Programs/Horde/Horde.Build/Devices/IDeviceCollection.cs
+++ b/Engine/Source/Programs/Horde/Horde.Build/Devices/IDeviceCollection.cs
@@ -217,6 +217,11 @@ namespace HordeServer.Collections
 		/// </summary>
 		public Task<bool> ExpireReservationsAsync();
 
+		/// <summary>
+		/// Deletes user device checkouts
+		/// </summary>
+		public Task<List<(UserId, IDevice)>?> ExpireCheckedOutAsync();
+
 	}
 
 }
\ No newline at end of file
diff --git a/Engine/Source/Programs/Horde/Horde.Build/Devices/Impl/DeviceCollection.cs b/Engine/Source/Programs/Horde/Horde.Build/Devices/Impl/DeviceCollection.cs
--- a/Engine/Source/Programs/Horde/Horde.Build/Devices/Impl/DeviceCollection.cs
+++ b/Engine/Source/Programs/Horde/Horde.Build/Devices/Impl/DeviceCollection.cs
@@ -696,6 +696,45 @@ namespace HordeServer.Collections.Impl
 
 		}
 
+		/// <summary>
+		/// Deletes expired user checkouts
+		/// </summary>
+		public async Task<List<(UserId, IDevice)>?> ExpireCheckedOutAsync()
+		{
+			FilterDefinition<DeviceDocument> Filter = Builders<DeviceDocument>.Filter.Empty;
+			Filter &= Builders<DeviceDocument>.Filter.Where(x => x.CheckedOutByUser != null && x.CheckOutTime != null);
+
+			List<DeviceDocument> CheckedOutDevices = await Devices.Find(Filter).ToListAsync();
+
+			DateTime UtcNow = DateTime.UtcNow;
+			List<DeviceDocument> ExpiredDevices = CheckedOutDevices.FindAll(x => (UtcNow - x.CheckOutTime!.Value).TotalMinutes > 5).ToList();
+
+
+			if (ExpiredDevices.Count > 0)
+			{
+
+				FilterDefinition<DeviceDocument> UpdateFilter = Builders<DeviceDocument>.Filter.In(x => x.Id, ExpiredDevices.Select(y => y.Id));
+
+				UpdateDefinitionBuilder<DeviceDocument> DeviceBuilder = Builders<DeviceDocument>.Update;
+				List<UpdateDefinition<DeviceDocument>> DeviceUpdates = new List<UpdateDefinition<DeviceDocument>>();
+
+				DeviceUpdates.Add(DeviceBuilder.Set(x => x.CheckedOutByUser, null));
+				DeviceUpdates.Add(DeviceBuilder.Set(x => x.CheckOutTime, null));
+
+				UpdateResult Result = await Devices.UpdateManyAsync(Builders<DeviceDocument>.Filter.In(x => x.Id, ExpiredDevices.Select(y => y.Id)), DeviceBuilder.Combine(DeviceUpdates));
+
+				if (Result.ModifiedCount > 0)
+				{
+					return ExpiredDevices.Select(x => (UserId.Parse(x.CheckedOutByUser!), (IDevice) x)).ToList();
+				}
+
+			}
+
+			return null;
+
+		}
+
+
 		/// <inheritdoc/>
 		public async Task<List<IDeviceReservation>> FindAllReservationsAsync()
 		{
diff --git a/Engine/Source/Programs/Horde/Horde.Build/Devices/Impl/DeviceService.cs b/Engine/Source/Programs/Horde/Horde.Build/Devices/Impl/DeviceService.cs
--- a/Engine/Source/Programs/Horde/Horde.Build/Devices/Impl/DeviceService.cs
+++ b/Engine/Source/Programs/Horde/Horde.Build/Devices/Impl/DeviceService.cs
@@ -116,6 +116,11 @@ namespace HordeServer.Services
 		/// </summary>
 		ProjectService ProjectService;
 
+		/// <summary>
+		/// The user collection instance
+		/// </summary>
+		IUserCollection UserCollection { get; set; }
+
 		/// <summary>
 		/// Log output writer
 		/// </summary>
@@ -136,8 +141,9 @@ namespace HordeServer.Services
 		/// <summary>
 		/// Device service constructor
 		/// </summary>
-		public DeviceService(IDeviceCollection Devices, ISingletonDocument<DevicePlatformMapV1> PlatformMapSingleton, JobService JobService, ProjectService ProjectService, StreamService StreamService, AclService AclService, INotificationService NotificationService, IClock Clock, ILogger<DeviceService> Logger)
+		public DeviceService(IDeviceCollection Devices, ISingletonDocument<DevicePlatformMapV1> PlatformMapSingleton, IUserCollection UserCollection, JobService JobService, ProjectService ProjectService, StreamService StreamService, AclService AclService, INotificationService NotificationService, IClock Clock, ILogger<DeviceService> Logger)
 		{
+			this.UserCollection = UserCollection;
 			this.Devices = Devices;
             this.JobService = JobService;
 			this.ProjectService = ProjectService;
@@ -169,6 +175,14 @@ namespace HordeServer.Services
 			if (!StoppingToken.IsCancellationRequested)
 			{
 				await Devices.ExpireReservationsAsync();
+				List<(UserId, IDevice)>? ExpireCheckouts = await Devices.ExpireCheckedOutAsync();
+				if (ExpireCheckouts != null && ExpireCheckouts.Count > 0)
+				{
+					foreach ((UserId, IDevice) ExpiredDevice in ExpireCheckouts)
+					{
+						await NotifyDeviceServiceAsync($"Device {ExpiredDevice.Item2.Name}:{ExpiredDevice.Item2.PlatformId} checkout has automatically expired and it has been returned to the shared pool.", null, null, null, ExpiredDevice.Item1);
+					}
+				}
 			}
 
 		}
@@ -358,7 +372,8 @@ namespace HordeServer.Services
 		/// <param name="DeviceId"></param>
 		/// <param name="JobId"></param>
 		/// <param name="StepId"></param>
-		public async Task NotifyDeviceServiceAsync(string Message, DeviceId? DeviceId = null, string? JobId = null, string? StepId = null)
+		/// <param name="UserId"></param>
+		public async Task NotifyDeviceServiceAsync(string Message, DeviceId? DeviceId = null, string? JobId = null, string? StepId = null, UserId? UserId = null)
 		{
 			try 
 			{
@@ -368,6 +383,17 @@ namespace HordeServer.Services
 				IJobStep? Step = null;
 				INode? Node = null;
 				IStream? Stream = null;
+				IUser? User = null;
+
+				if (UserId.HasValue)
+				{
+					User = await UserCollection.GetUserAsync(UserId.Value);
+					if (User == null)
+					{
+						Logger.LogError("Unable to send device notification, can't find User {UserId}", UserId.Value);
+						return;
+					}
+				}
 
 				if (DeviceId.HasValue)
 				{
@@ -399,7 +425,7 @@ namespace HordeServer.Services
 					}
 				}
 
-				NotificationService.NotifyDeviceService(Message, Device, Pool, Stream, Job, Step, Node);
+				NotificationService.NotifyDeviceService(Message, Device, Pool, Stream, Job, Step, Node, User);
 
 			}
 			catch (Exception Ex)
diff --git a/Engine/Source/Programs/Horde/Horde.Build/Notifications/INotificationService.cs b/Engine/Source/Programs/Horde/Horde.Build/Notifications/INotificationService.cs
--- a/Engine/Source/Programs/Horde/Horde.Build/Notifications/INotificationService.cs
+++ b/Engine/Source/Programs/Horde/Horde.Build/Notifications/INotificationService.cs
@@ -103,15 +103,16 @@ namespace HordeServer.Notifications
 		void NotifyIssueUpdated(IIssue Issue);
 
 		/// <summary>
-        /// Send a device service notification
-        /// </summary>
-        /// <param name="Message">The message to send</param>
-        /// <param name="Device">The device</param>
-        /// <param name="Pool">The pool</param>
-        /// <param name="Stream"></param>
-        /// <param name="Job"></param>
-        /// <param name="Step"></param>
-        /// <param name="Node"></param>
-        void NotifyDeviceService(string Message, IDevice? Device = null, IDevicePool? Pool = null, IStream? Stream = null, IJob? Job = null, IJobStep? Step = null, INode? Node = null);
+		/// Send a device service notification
+		/// </summary>
+		/// <param name="Message">The message to send</param>
+		/// <param name="Device">The device</param>
+		/// <param name="Pool">The pool</param>
+		/// <param name="Stream"></param>
+		/// <param name="Job"></param>
+		/// <param name="Step"></param>
+		/// <param name="Node"></param>
+		/// <param name="User"></param>
+		void NotifyDeviceService(string Message, IDevice? Device = null, IDevicePool? Pool = null, IStream? Stream = null, IJob? Job = null, IJobStep? Step = null, INode? Node = null, IUser? User = null);
     }
 }
diff --git a/Engine/Source/Programs/Horde/Horde.Build/Notifications/INotificationSink.cs b/Engine/Source/Programs/Horde/Horde.Build/Notifications/INotificationSink.cs
--- a/Engine/Source/Programs/Horde/Horde.Build/Notifications/INotificationSink.cs
+++ b/Engine/Source/Programs/Horde/Horde.Build/Notifications/INotificationSink.cs
@@ -85,7 +85,7 @@ namespace HordeServer.Notifications
 		/// <param name="Description"></param>
 		/// <returns></returns>
 		Task NotifyConfigUpdateFailureAsync(string ErrorMessage, string FileName, int? Change = null, IUser? Author = null, string? Description = null);
-		
+
 		/// <summary>
 		/// Notification for device service
 		/// </summary>
@@ -94,9 +94,11 @@ namespace HordeServer.Notifications
 		/// <param name="Pool"></param>
 		/// <param name="Stream"></param>
 		/// <param name="Job"></param>
-        /// <param name="Step"></param>
-        /// <param name="Node"></param>
+		/// <param name="Step"></param>
+		/// <param name="Node"></param>
+		/// <param name="User"></param>
 		/// <returns></returns>
-        Task NotifyDeviceServiceAsync(string Message, IDevice? Device = null, IDevicePool? Pool = null, IStream? Stream = null, IJob? Job = null, IJobStep? Step = null, INode? Node = null);
+		Task NotifyDeviceServiceAsync(string Message, IDevice? Device = null, IDevicePool? Pool = null, IStream? Stream = null, IJob? Job = null, IJobStep? Step = null, INode? Node = null, IUser? User = null);
     }
 }
+
diff --git a/Engine/Source/Programs/Horde/Horde.Build/Notifications/Impl/NotificationService.cs b/Engine/Source/Programs/Horde/Horde.Build/Notifications/Impl/NotificationService.cs
--- a/Engine/Source/Programs/Horde/Horde.Build/Notifications/Impl/NotificationService.cs
+++ b/Engine/Source/Programs/Horde/Horde.Build/Notifications/Impl/NotificationService.cs
@@ -259,7 +259,7 @@ namespace HordeServer.Notifications.Impl
 		}
 
 		/// <inheritdoc/>
-		public void NotifyDeviceService(string Message, IDevice? Device = null, IDevicePool? Pool = null, IStream? Stream = null, IJob? Job = null, IJobStep? Step = null, INode? Node = null)
+		public void NotifyDeviceService(string Message, IDevice? Device = null, IDevicePool? Pool = null, IStream? Stream = null, IJob? Job = null, IJobStep? Step = null, INode? Node = null, IUser? User = null)
 		{
 			EnqueueTasks(Sink => Sink.NotifyDeviceServiceAsync(Message, Device, Pool, Stream, Job, Step, Node));
 		}
diff --git a/Engine/Source/Programs/Horde/Horde.Build/Notifications/Impl/SlackNotificationSink.cs b/Engine/Source/Programs/Horde/Horde.Build/Notifications/Impl/SlackNotificationSink.cs
--- a/Engine/Source/Programs/Horde/Horde.Build/Notifications/Impl/SlackNotificationSink.cs
+++ b/Engine/Source/Programs/Horde/Horde.Build/Notifications/Impl/SlackNotificationSink.cs
@@ -1013,27 +1013,44 @@ namespace HordeServer.Notifications.Impl
 		#region Device notifications
 
 		/// <inheritdoc/>
-		public async Task NotifyDeviceServiceAsync(string Message, IDevice? Device = null, IDevicePool? Pool = null, IStream? Stream = null, IJob? Job = null, IJobStep? Step = null, INode? Node = null)
+		public async Task NotifyDeviceServiceAsync(string Message, IDevice? Device = null, IDevicePool? Pool = null, IStream? Stream = null, IJob? Job = null, IJobStep? Step = null, INode? Node = null, IUser? User = null)
 		{
-			Logger.LogDebug("Sending device service failure notification for {DeviceName} in pool {PoolName}", Device?.Name, Pool?.Name);
+			string Recipient = $"#{Settings.DeviceServiceNotificationChannel}";
+
+			if (User != null)
+			{
+				string? SlackRecipient = await GetSlackUserId(User);
+
+				if (SlackRecipient == null)
+				{
+					Logger.LogError("NotifyDeviceServiceAsync - Unable to send user slack notification, user {UserId} slack user id not found", User.Id);
+					return;
+				}
+
+				Recipient = SlackRecipient;
+			}
+
+			Logger.LogDebug("Sending device service notification to {Recipient}", Recipient);
+
 			if (Settings.DeviceServiceNotificationChannel != null)
 			{
-				await SendDeviceServiceMessage($"#{Settings.DeviceServiceNotificationChannel}", Message, Device, Pool, Stream, Job, Step, Node);
+				await SendDeviceServiceMessage(Recipient, Message, Device, Pool, Stream, Job, Step, Node, User);
 			}
 		}
 
 		/// <summary>
-		/// Creates a Slack message about a completed step job.
+		/// Creates a Slack message for a device service notification
 		/// </summary>
 		/// <param name="Recipient"></param>
-        /// <param name="Message"></param>
-        /// <param name="Device"></param>
-        /// <param name="Pool"></param>
+		/// <param name="Message"></param>
+		/// <param name="Device"></param>
+		/// <param name="Pool"></param>
 		/// <param name="Stream"></param>
 		/// <param name="Job">The job that contains the step that completed.</param>
 		/// <param name="Step">The job step that completed.</param>
 		/// <param name="Node">The node for the job step.</param>
-		private Task SendDeviceServiceMessage(string Recipient, string Message, IDevice? Device = null, IDevicePool? Pool = null, IStream? Stream = null, IJob? Job = null, IJobStep? Step = null, INode? Node = null)
+		/// <param name="User">The user to notify.</param>
+		private Task SendDeviceServiceMessage(string Recipient, string Message, IDevice? Device = null, IDevicePool? Pool = null, IStream? Stream = null, IJob? Job = null, IJobStep? Step = null, INode? Node = null, IUser? User = null)
 		{
 			BlockKitAttachment Attachment = new BlockKitAttachment();
 
@@ -1052,19 +1069,22 @@ namespace HordeServer.Notifications.Impl
 
 			Attachment.Blocks.Add(new HeaderBlock($"{Message}", false, false));
 
-			if (Stream != null && Job != null && Step != null && Node != null)
+			if (User == null)
 			{
-				Uri JobStepLink = new Uri($"{Settings.DashboardUrl}job/{Job.Id}?step={Step.Id}");
-				Uri JobStepLogLink = new Uri($"{Settings.DashboardUrl}log/{Step.LogId}");
-				
-				Attachment.FallbackText += $" - {Stream.Name} - {GetJobChangeText(Job)} - {Job.Name} - {Node.Name}";				
-				Attachment.Blocks.Add(new SectionBlock($"*<{JobStepLink}|{Stream.Name} - {GetJobChangeText(Job)} - {Job.Name} - {Node.Name}>*"));
-				Attachment.Blocks.Add(new SectionBlock($"<{JobStepLogLink}|View Job Step Log>"));
-			}
-			else
-			{				
-				Attachment.FallbackText += " - No job information (Gauntlet might need to be updated in stream)";				
-				Attachment.Blocks.Add(new SectionBlock("*No job information (Gauntlet might need to be updated in stream)*"));
+				if (Stream != null && Job != null && Step != null && Node != null)
+				{
+					Uri JobStepLink = new Uri($"{Settings.DashboardUrl}job/{Job.Id}?step={Step.Id}");
+					Uri JobStepLogLink = new Uri($"{Settings.DashboardUrl}log/{Step.LogId}");
+
+					Attachment.FallbackText += $" - {Stream.Name} - {GetJobChangeText(Job)} - {Job.Name} - {Node.Name}";
+					Attachment.Blocks.Add(new SectionBlock($"*<{JobStepLink}|{Stream.Name} - {GetJobChangeText(Job)} - {Job.Name} - {Node.Name}>*"));
+					Attachment.Blocks.Add(new SectionBlock($"<{JobStepLogLink}|View Job Step Log>"));
+				}
+				else
+				{
+					Attachment.FallbackText += " - No job information (Gauntlet might need to be updated in stream)";
+					Attachment.Blocks.Add(new SectionBlock("*No job information (Gauntlet might need to be updated in stream)*"));
+				}
 			}
 
 			return SendMessageAsync(Recipient, Attachments: new[] { Attachment });
