Horde: Fix broken test
#rnx
#skipci

diff --git a/Engine/Source/Programs/Horde/Horde.Build.Tests/NotificationServiceTests.cs b/Engine/Source/Programs/Horde/Horde.Build.Tests/NotificationServiceTests.cs
--- a/Engine/Source/Programs/Horde/Horde.Build.Tests/NotificationServiceTests.cs
+++ b/Engine/Source/Programs/Horde/Horde.Build.Tests/NotificationServiceTests.cs
@@ -38,7 +38,7 @@ namespace HordeServerTests
 		public Task NotifyLabelCompleteAsync(IUser User, IJob Job, IStream Stream, ILabel Label, int LabelIdx, LabelOutcome Outcome, List<(string, JobStepOutcome, Uri)> StepData) { throw new NotImplementedException(); }
 		public Task NotifyIssueUpdatedAsync(IIssue Issue) { throw new NotImplementedException(); }
 		public Task NotifyConfigUpdateFailureAsync(string ErrorMessage, string FileName, int? Change = null, IUser? Author = null, string? Description = null) { throw new NotImplementedException(); }
-		public Task NotifyDeviceServiceAsync(string Message, IDevice? Device = null, IDevicePool? Pool = null, IStream? Stream = null, IJob? Job = null, IJobStep? Step = null, INode? Node = null) { throw new NotImplementedException(); }
+		public Task NotifyDeviceServiceAsync(string Message, IDevice? Device = null, IDevicePool? Pool = null, IStream? Stream = null, IJob? Job = null, IJobStep? Step = null, INode? Node = null, IUser? User = null) { throw new NotImplementedException(); }
 	}
 	
 
