Horde: Fix missing user on shared device notification
#rnx
#skipci

diff --git a/Engine/Source/Programs/Horde/Horde.Build/Notifications/Impl/NotificationService.cs b/Engine/Source/Programs/Horde/Horde.Build/Notifications/Impl/NotificationService.cs
--- a/Engine/Source/Programs/Horde/Horde.Build/Notifications/Impl/NotificationService.cs
+++ b/Engine/Source/Programs/Horde/Horde.Build/Notifications/Impl/NotificationService.cs
@@ -261,7 +261,7 @@ namespace HordeServer.Notifications.Impl
 		/// <inheritdoc/>
 		public void NotifyDeviceService(string Message, IDevice? Device = null, IDevicePool? Pool = null, IStream? Stream = null, IJob? Job = null, IJobStep? Step = null, INode? Node = null, IUser? User = null)
 		{
-			EnqueueTasks(Sink => Sink.NotifyDeviceServiceAsync(Message, Device, Pool, Stream, Job, Step, Node));
+			EnqueueTasks(Sink => Sink.NotifyDeviceServiceAsync(Message, Device, Pool, Stream, Job, Step, Node, User));
 		}
 
 		/// <summary>
