Horde: Set shared device checkout time to 5 days
#rnx
#skipci

diff --git a/Engine/Source/Programs/Horde/Horde.Build/Devices/Impl/DeviceCollection.cs b/Engine/Source/Programs/Horde/Horde.Build/Devices/Impl/DeviceCollection.cs
--- a/Engine/Source/Programs/Horde/Horde.Build/Devices/Impl/DeviceCollection.cs
+++ b/Engine/Source/Programs/Horde/Horde.Build/Devices/Impl/DeviceCollection.cs
@@ -707,8 +707,7 @@ namespace HordeServer.Collections.Impl
 			List<DeviceDocument> CheckedOutDevices = await Devices.Find(Filter).ToListAsync();
 
 			DateTime UtcNow = DateTime.UtcNow;
-			List<DeviceDocument> ExpiredDevices = CheckedOutDevices.FindAll(x => (UtcNow - x.CheckOutTime!.Value).TotalMinutes > 5).ToList();
-
+			List<DeviceDocument> ExpiredDevices = CheckedOutDevices.FindAll(x => (UtcNow - x.CheckOutTime!.Value).TotalDays >= 5).ToList();
 
 			if (ExpiredDevices.Count > 0)
 			{
