Horde: Add device checkout time to dashboard view
#rnx
#skipci

diff --git a/Engine/Source/Programs/Horde/Horde.Build/Notifications/Impl/SlackNotificationSink.cs b/Engine/Source/Programs/Horde/Horde.Build/Notifications/Impl/SlackNotificationSink.cs
--- a/Engine/Source/Programs/Horde/Horde.Build/Notifications/Impl/SlackNotificationSink.cs
+++ b/Engine/Source/Programs/Horde/Horde.Build/Notifications/Impl/SlackNotificationSink.cs
@@ -1052,42 +1052,46 @@ namespace HordeServer.Notifications.Impl
 		/// <param name="User">The user to notify.</param>
 		private Task SendDeviceServiceMessage(string Recipient, string Message, IDevice? Device = null, IDevicePool? Pool = null, IStream? Stream = null, IJob? Job = null, IJobStep? Step = null, INode? Node = null, IUser? User = null)
 		{
-			BlockKitAttachment Attachment = new BlockKitAttachment();
 
 			// truncate message to avoid slack error on message length
 			if (Message.Length > 150)
 			{
 				Message = Message.Substring(0, 146) + "...";
 			}
-            
-            Attachment.FallbackText = $"{Message}";
 
-			if (Device != null && Pool != null)
+			if (User != null)
 			{
-                Attachment.FallbackText += $" - Device: {Device.Name} Pool: {Pool.Name}";
-            }
+				return SendMessageAsync(Recipient, Message);
+			}
+			
+			BlockKitAttachment Attachment = new BlockKitAttachment();
+							
+			Attachment.FallbackText = $"{Message}";
 
+			if (Device != null && Pool != null)
+			{
+				Attachment.FallbackText += $" - Device: {Device.Name} Pool: {Pool.Name}";
+			}
+				
 			Attachment.Blocks.Add(new HeaderBlock($"{Message}", false, false));
 
-			if (User == null)
+			if (Stream != null && Job != null && Step != null && Node != null)
 			{
-				if (Stream != null && Job != null && Step != null && Node != null)
-				{
-					Uri JobStepLink = new Uri($"{Settings.DashboardUrl}job/{Job.Id}?step={Step.Id}");
-					Uri JobStepLogLink = new Uri($"{Settings.DashboardUrl}log/{Step.LogId}");
+				Uri JobStepLink = new Uri($"{Settings.DashboardUrl}job/{Job.Id}?step={Step.Id}");
+				Uri JobStepLogLink = new Uri($"{Settings.DashboardUrl}log/{Step.LogId}");
 
-					Attachment.FallbackText += $" - {Stream.Name} - {GetJobChangeText(Job)} - {Job.Name} - {Node.Name}";
-					Attachment.Blocks.Add(new SectionBlock($"*<{JobStepLink}|{Stream.Name} - {GetJobChangeText(Job)} - {Job.Name} - {Node.Name}>*"));
-					Attachment.Blocks.Add(new SectionBlock($"<{JobStepLogLink}|View Job Step Log>"));
-				}
-				else
-				{
-					Attachment.FallbackText += " - No job information (Gauntlet might need to be updated in stream)";
-					Attachment.Blocks.Add(new SectionBlock("*No job information (Gauntlet might need to be updated in stream)*"));
-				}
+				Attachment.FallbackText += $" - {Stream.Name} - {GetJobChangeText(Job)} - {Job.Name} - {Node.Name}";
+				Attachment.Blocks.Add(new SectionBlock($"*<{JobStepLink}|{Stream.Name} - {GetJobChangeText(Job)} - {Job.Name} - {Node.Name}>*"));
+				Attachment.Blocks.Add(new SectionBlock($"<{JobStepLogLink}|View Job Step Log>"));
+			}
+			else
+			{
+				Attachment.FallbackText += " - No job information (Gauntlet might need to be updated in stream)";
+				Attachment.Blocks.Add(new SectionBlock("*No job information (Gauntlet might need to be updated in stream)*"));
 			}
 
 			return SendMessageAsync(Recipient, Attachments: new[] { Attachment });
+
 		}
 		
 		#endregion
