Horde: Add missing size strategy field to pool update endpoint

diff --git a/Engine/Source/Programs/Horde/Horde.Build.Tests/PoolsControllerTest.cs b/Engine/Source/Programs/Horde/Horde.Build.Tests/PoolsControllerTest.cs
--- a/Engine/Source/Programs/Horde/Horde.Build.Tests/PoolsControllerTest.cs
+++ b/Engine/Source/Programs/Horde/Horde.Build.Tests/PoolsControllerTest.cs
@@ -1,20 +1,46 @@
 ﻿// Copyright Epic Games, Inc. All Rights Reserved.
 
+using System.Collections.Generic;
 using System.Threading.Tasks;
+using Horde.Build.Fleet.Autoscale;
+using HordeServer.Api;
+using HordeServer.Models;
+using HordeServerTests;
+using Microsoft.AspNetCore.Mvc;
 using Microsoft.VisualStudio.TestTools.UnitTesting;
 
-namespace HordeServerTests
+namespace Horde.Build.Tests
 {
     [TestClass]
-    public class PoolsControllerTest : ControllerIntegrationTest
+    public class PoolsControllerTest : TestSetup
     {
         [TestMethod]
         public async Task GetPoolsTest()
         {
-            var Res = await client.GetAsync("/api/v1/pools");
-            Res.EnsureSuccessStatusCode();
-            var jsonData = await Res.Content.ReadAsStringAsync();
-            Assert.AreEqual("[]", jsonData);
+	        IPool Pool1 = await PoolService.CreatePoolAsync("Pool1", Properties: new() { {"foo", "bar"}, {"lorem", "ipsum"} });
+	        ActionResult<List<object>> RawResult = await PoolsController.GetPoolsAsync();
+	        Assert.AreEqual(1, RawResult.Value!.Count);
+	        GetPoolResponse Response = (RawResult.Value![0] as GetPoolResponse)!;
+	        Assert.AreEqual(Pool1.Id.ToString(), Response.Id);
+	        Assert.AreEqual(Pool1.Name, Response.Name);
+	        Assert.AreEqual(Pool1.SizeStrategy, Response.SizeStrategy);
+        }
+        
+        [TestMethod]
+        public async Task UpdatePoolTest()
+        {
+	        IPool Pool1 = await PoolService.CreatePoolAsync("Pool1", Properties: new() { {"foo", "bar"}, {"lorem", "ipsum"} });
+
+	        await PoolsController.UpdatePoolAsync(Pool1.Id.ToString(), new UpdatePoolRequest
+	        {
+		        Name = "Pool1Modified",
+		        SizeStrategy = PoolSizeStrategy.JobQueue
+	        });
+
+	        ActionResult<object> GetResult = await PoolsController.GetPoolAsync(Pool1.Id.ToString());
+	        GetPoolResponse Response = (GetResult.Value! as GetPoolResponse)!;
+	        Assert.AreEqual("Pool1Modified", Response.Name);
+	        Assert.AreEqual(PoolSizeStrategy.JobQueue, Response.SizeStrategy);
         }
     }
 }
\ No newline at end of file
diff --git a/Engine/Source/Programs/Horde/Horde.Build.Tests/TestSetup.cs b/Engine/Source/Programs/Horde/Horde.Build.Tests/TestSetup.cs
--- a/Engine/Source/Programs/Horde/Horde.Build.Tests/TestSetup.cs
+++ b/Engine/Source/Programs/Horde/Horde.Build.Tests/TestSetup.cs
@@ -98,6 +98,7 @@ namespace HordeServerTests
 
 		public JobsController JobsController => GetJobsController();
 		public AgentsController AgentsController => GetAgentsController();
+		public PoolsController PoolsController => GetPoolsController();
 
 		private static bool IsDatadogWriterPatched;
 
@@ -231,6 +232,13 @@ namespace HordeServerTests
 			return AgentCtrl;
 		}
 		
+		private PoolsController GetPoolsController()
+		{
+			PoolsController Controller = new PoolsController(AclService, PoolService);
+			Controller.ControllerContext = GetControllerContext();
+			return Controller;
+		}
+		
 		private ControllerContext GetControllerContext()
 		{
 			ControllerContext ControllerContext = new ControllerContext();
diff --git a/Engine/Source/Programs/Horde/Horde.Build/Api/Pools.cs b/Engine/Source/Programs/Horde/Horde.Build/Api/Pools.cs
--- a/Engine/Source/Programs/Horde/Horde.Build/Api/Pools.cs
+++ b/Engine/Source/Programs/Horde/Horde.Build/Api/Pools.cs
@@ -160,6 +160,11 @@ namespace HordeServer.Api
 		/// Whether to enable autoscaling for this pool
 		/// </summary>
 		public bool EnableAutoscaling { get; set; }
+		
+		/// <summary>
+		/// Pool sizing strategy to be used for this pool
+		/// </summary>
+		public PoolSizeStrategy? SizeStrategy { get; set; }
 
 		/// <summary>
 		/// The minimum nunmber of agents to retain in this pool
@@ -191,6 +196,7 @@ namespace HordeServer.Api
 			this.Name = Pool.Name;
 			this.Condition = Pool.Condition;
 			this.EnableAutoscaling = Pool.EnableAutoscaling;
+			this.SizeStrategy = Pool.SizeStrategy;
 			this.MinAgents = Pool.MinAgents;
 			this.NumReserveAgents = Pool.NumReserveAgents;
 			this.Workspaces = Pool.Workspaces.Select(x => new GetAgentWorkspaceResponse(x)).ToList();
diff --git a/Engine/Source/Programs/Horde/Horde.Build/Controllers/PoolsController.cs b/Engine/Source/Programs/Horde/Horde.Build/Controllers/PoolsController.cs
--- a/Engine/Source/Programs/Horde/Horde.Build/Controllers/PoolsController.cs
+++ b/Engine/Source/Programs/Horde/Horde.Build/Controllers/PoolsController.cs
@@ -139,7 +139,7 @@ namespace HordeServer.Controllers
 				return NotFound(PoolIdValue);
 			}
 
-			await PoolService.UpdatePoolAsync(Pool, Update.Name, Update.Condition, Update.EnableAutoscaling, Update.MinAgents, Update.NumReserveAgents, Update.Properties);
+			await PoolService.UpdatePoolAsync(Pool, Update.Name, Update.Condition, Update.EnableAutoscaling, Update.MinAgents, Update.NumReserveAgents, Update.Properties, Update.SizeStrategy);
 			return new OkResult();
 		}
 
diff --git a/Engine/Source/Programs/Horde/Horde.Build/Fleet/PoolService.cs b/Engine/Source/Programs/Horde/Horde.Build/Fleet/PoolService.cs
--- a/Engine/Source/Programs/Horde/Horde.Build/Fleet/PoolService.cs
+++ b/Engine/Source/Programs/Horde/Horde.Build/Fleet/PoolService.cs
@@ -95,12 +95,13 @@ namespace HordeServer.Services
 		/// <param name="NewMinAgents">Minimum number of agents in the pool</param>
 		/// <param name="NewNumReserveAgents">Minimum number of idle agents to maintain</param>
 		/// <param name="NewProperties">Properties on the pool to update. Any properties with a value of null will be removed.</param>
+		/// <param name="SizeStrategy">New pool sizing strategy for the pool</param>
 		/// <returns>Async task object</returns>
-		public async Task<IPool?> UpdatePoolAsync(IPool? Pool, string? NewName = null, Condition? NewCondition = null, bool? NewEnableAutoscaling = null, int? NewMinAgents = null, int? NewNumReserveAgents = null, Dictionary<string, string?>? NewProperties = null)
+		public async Task<IPool?> UpdatePoolAsync(IPool? Pool, string? NewName = null, Condition? NewCondition = null, bool? NewEnableAutoscaling = null, int? NewMinAgents = null, int? NewNumReserveAgents = null, Dictionary<string, string?>? NewProperties = null, PoolSizeStrategy? SizeStrategy = null)
 		{
 			for (; Pool != null; Pool = await Pools.GetAsync(Pool.Id))
 			{
-				IPool? NewPool = await Pools.TryUpdateAsync(Pool, NewName, NewCondition, NewEnableAutoscaling, NewMinAgents, NewNumReserveAgents, NewProperties: NewProperties);
+				IPool? NewPool = await Pools.TryUpdateAsync(Pool, NewName, NewCondition, NewEnableAutoscaling, NewMinAgents, NewNumReserveAgents, NewProperties: NewProperties, SizeStrategy: SizeStrategy);
 				if (NewPool != null)
 				{
 					return NewPool;
