Horde Storage - Improvements to legacy replicator to attempt it to self heal better when finding corrupt transaction logs (which can happen if we run out of diskspace as we are writing it)

diff --git a/Engine/Source/Programs/Horde/HordeStorage/Callisto/Implementation/FileTransactionLog.cs b/Engine/Source/Programs/Horde/HordeStorage/Callisto/Implementation/FileTransactionLog.cs
--- a/Engine/Source/Programs/Horde/HordeStorage/Callisto/Implementation/FileTransactionLog.cs
+++ b/Engine/Source/Programs/Horde/HordeStorage/Callisto/Implementation/FileTransactionLog.cs
@@ -243,7 +243,7 @@ namespace Callisto.Implementation
                 {
                     string headerMd5 = StringUtils.FormatAsHexString(opHeader.Md5hash);
                     string calculatedMd5 = StringUtils.FormatAsHexString(md5Hash);
-                    _logger.Warning("Reading header {@Header}, mismatching hash found {CalculatedHash} expected {ExpectedHash}", opHeader, calculatedMd5, headerMd5);
+                    _logger.Warning("Reading header {@Header}, mismatching hash found {CalculatedHash} expected {ExpectedHash} . Was read at offset {Offset}", opHeader, calculatedMd5, headerMd5, offset);
                     throw new EventHashMismatchException(opHeader.Offset, headerMd5, calculatedMd5);
                 }
 
diff --git a/Engine/Source/Programs/Horde/HordeStorage/Horde.Storage/Implementation/Replication/ReplicationSettings.cs b/Engine/Source/Programs/Horde/HordeStorage/Horde.Storage/Implementation/Replication/ReplicationSettings.cs
--- a/Engine/Source/Programs/Horde/HordeStorage/Horde.Storage/Implementation/Replication/ReplicationSettings.cs
+++ b/Engine/Source/Programs/Horde/HordeStorage/Horde.Storage/Implementation/Replication/ReplicationSettings.cs
@@ -76,6 +76,11 @@ namespace Horde.Storage
         /// The replication will stop its run after this many events have been processed, yielding its time until the next time a replication is triggered.
         /// </summary>
         public int MaxReplicationsPerRun { get; set; } = 1000;
+
+        /// <summary>
+        /// Max number of offsets to skip when finding broken events in the transaction logs, higher value means we scan further increasing the likelihood of finding a valid entry, but also means risk that we skip records you would want.
+        /// </summary>
+        public int? MaxOffsetsAttempted { get; set; } = 100000;
     }
 
     public enum ReplicatorVersion
diff --git a/Engine/Source/Programs/Horde/HordeStorage/Horde.Storage/Implementation/Replication/Replicator.cs b/Engine/Source/Programs/Horde/HordeStorage/Horde.Storage/Implementation/Replication/Replicator.cs
--- a/Engine/Source/Programs/Horde/HordeStorage/Horde.Storage/Implementation/Replication/Replicator.cs
+++ b/Engine/Source/Programs/Horde/HordeStorage/Horde.Storage/Implementation/Replication/Replicator.cs
@@ -40,7 +40,7 @@ namespace Horde.Storage.Implementation
         {
             CallistoReader remoteCallistoReader = new CallistoReader(_restClient, _replicatorSettings.NamespaceToReplicate);
 
-            return remoteCallistoReader.GetOps(stateReplicatorOffset, stateReplicatingGeneration, currentSite, enumerationState: enumerationState, cancellationToken: replicationToken);
+            return remoteCallistoReader.GetOps(stateReplicatorOffset, stateReplicatingGeneration, currentSite, enumerationState: enumerationState, cancellationToken: replicationToken, maxOffsetsAttempted: _replicatorSettings.MaxOffsetsAttempted);
         }
 
         protected override async Task ReplicateOp(IRestClient remoteClient, TransactionEvent op, CancellationToken replicationToken)
