Horde Storage - Attempting to replicate a blob that has been deleted is now a warning and not a error (as its a valid case).

diff --git a/Engine/Source/Programs/Horde/HordeStorage/Horde.Storage/Implementation/Replication/RefsReplicator.cs b/Engine/Source/Programs/Horde/HordeStorage/Horde.Storage/Implementation/Replication/RefsReplicator.cs
--- a/Engine/Source/Programs/Horde/HordeStorage/Horde.Storage/Implementation/Replication/RefsReplicator.cs
+++ b/Engine/Source/Programs/Horde/HordeStorage/Horde.Storage/Implementation/Replication/RefsReplicator.cs
@@ -454,6 +454,12 @@ namespace Horde.Storage.Implementation
                     HttpRequestMessage blobRequest = BuildHttpRequest(HttpMethod.Get, $"api/v1/blobs/{ns}/{blobToReplicate}");
                     HttpResponseMessage blobResponse = await _httpClient.SendAsync(blobRequest, HttpCompletionOption.ResponseHeadersRead, cancellationToken);
 
+                    if (blobResponse.StatusCode == HttpStatusCode.NotFound)
+                    {
+                        _logger.Warning("Failed to replicate {Blob} in {Namespace} due to it not existing.", blobToReplicate, ns);
+                        return;
+                    }
+
                     if (blobResponse.StatusCode != HttpStatusCode.OK)
                     {
                         _logger.Error("Bad http response when replicating {Blob} in {Namespace} . Status code: {StatusCode}", blobToReplicate, ns, blobResponse.StatusCode);
