Horde: Ensure labels are included in find job timings endpoint

diff --git a/Engine/Source/Programs/Horde/Horde.Build.Tests/JobsControllerDbTest.cs b/Engine/Source/Programs/Horde/Horde.Build.Tests/JobsControllerDbTest.cs
--- a/Engine/Source/Programs/Horde/Horde.Build.Tests/JobsControllerDbTest.cs
+++ b/Engine/Source/Programs/Horde/Horde.Build.Tests/JobsControllerDbTest.cs
@@ -69,7 +69,9 @@ namespace HordeServerTests
 	        object Obj = (await JobsController.FindJobTimingsAsync(Fixture.Stream!.Id.ToString(), Templates)).Value!;
 	        FindJobTimingsResponse Res = (Obj as FindJobTimingsResponse)!;
 	        Assert.AreEqual(1, Res.Timings.Count);
-	        Assert.AreEqual(Job.Name, Res.Timings[Job.Id.ToString()].Job.Name);
+	        GetJobTimingResponse TimingResponse = Res.Timings[Job.Id.ToString()];
+	        Assert.AreEqual(0, TimingResponse.JobResponse!.Labels!.Count);
+	        Assert.AreEqual(Job.Name, TimingResponse.Job!.Name);
         }
     }
 }
\ No newline at end of file
diff --git a/Engine/Source/Programs/Horde/Horde.Build/Api/Jobs.cs b/Engine/Source/Programs/Horde/Horde.Build/Api/Jobs.cs
--- a/Engine/Source/Programs/Horde/Horde.Build/Api/Jobs.cs
+++ b/Engine/Source/Programs/Horde/Horde.Build/Api/Jobs.cs
@@ -386,7 +386,12 @@ namespace HordeServer.Api
 		/// <summary>
 		/// The job object
 		/// </summary>
-		public IJob Job { get; set; }
+		public IJob? Job { get; set; }
+		
+		/// <summary>
+		/// A full job response
+		/// </summary>
+		public GetJobResponse? JobResponse { get; set; }
 		
 		/// <summary>
 		/// Timing info for each step
@@ -402,11 +407,13 @@ namespace HordeServer.Api
 		/// Constructor
 		/// </summary>
 		/// <param name="Job">The job object</param>
+		/// <param name="JobResponse">The job response</param>
 		/// <param name="Steps">Timing info for each steps</param>
 		/// <param name="Labels">Timing info for each label</param>
-		public GetJobTimingResponse(IJob Job, Dictionary<string, GetStepTimingInfoResponse> Steps, List<GetLabelTimingInfoResponse> Labels)
+		public GetJobTimingResponse(IJob? Job, GetJobResponse? JobResponse, Dictionary<string, GetStepTimingInfoResponse> Steps, List<GetLabelTimingInfoResponse> Labels)
 		{
 			this.Job = Job;
+			this.JobResponse = JobResponse;
 			this.Steps = Steps;
 			this.Labels = Labels;
 		}
diff --git a/Engine/Source/Programs/Horde/Horde.Build/Controllers/JobsController.cs b/Engine/Source/Programs/Horde/Horde.Build/Controllers/JobsController.cs
--- a/Engine/Source/Programs/Horde/Horde.Build/Controllers/JobsController.cs
+++ b/Engine/Source/Programs/Horde/Horde.Build/Controllers/JobsController.cs
@@ -564,10 +564,10 @@ namespace HordeServer.Controllers
 
 			IJobTiming JobTiming = await JobService.GetJobTimingAsync(Job);
 			IGraph Graph = await JobService.GetGraphAsync(Job);
-			return PropertyFilter.Apply(CreateJobTimingResponse(Job, Graph, JobTiming), Filter);
+			return PropertyFilter.Apply(await CreateJobTimingResponse(Job, Graph, JobTiming), Filter);
 		}
 
-		private static GetJobTimingResponse CreateJobTimingResponse(IJob Job, IGraph Graph, IJobTiming JobTiming)
+		private async Task<GetJobTimingResponse> CreateJobTimingResponse(IJob Job, IGraph Graph, IJobTiming JobTiming, bool IncludeJobResponse = false)
 		{
 			Dictionary<INode, TimingInfo> NodeToTimingInfo = Job.GetTimingInfo(Graph, JobTiming);
 
@@ -588,7 +588,13 @@ namespace HordeServer.Controllers
 				Labels.Add(new GetLabelTimingInfoResponse(Label, TimingInfo));
 			}
 
-			return new GetJobTimingResponse(Job, Steps, Labels);
+			GetJobResponse? JobResponse = null;
+			if (IncludeJobResponse)
+			{
+				JobResponse = await CreateJobResponseAsync(Job, Graph, false, true, false, false);
+			}
+
+			return new GetJobTimingResponse(Job, JobResponse, Steps, Labels);
 		}
 		
 		/// <summary>
@@ -627,7 +633,7 @@ namespace HordeServer.Controllers
 				{
 					IJobTiming JobTiming = await JobService.GetJobTimingAsync(Job);
 					IGraph Graph = await JobService.GetGraphAsync(Job);
-					return CreateJobTimingResponse(Job, Graph, JobTiming);
+					return await CreateJobTimingResponse(Job, Graph, JobTiming, true);
 				});
 			
 			return PropertyFilter.Apply(new FindJobTimingsResponse(JobTimings), Filter);
