Horde: Include batches and costs in find timings endpoint

diff --git a/Engine/Source/Programs/Horde/Horde.Build/Controllers/JobsController.cs b/Engine/Source/Programs/Horde/Horde.Build/Controllers/JobsController.cs
--- a/Engine/Source/Programs/Horde/Horde.Build/Controllers/JobsController.cs
+++ b/Engine/Source/Programs/Horde/Horde.Build/Controllers/JobsController.cs
@@ -591,7 +591,7 @@ namespace HordeServer.Controllers
 			GetJobResponse? JobResponse = null;
 			if (IncludeJobResponse)
 			{
-				JobResponse = await CreateJobResponseAsync(Job, Graph, false, true, false, false);
+				JobResponse = await CreateJobResponseAsync(Job, Graph, true, true, false, true);
 			}
 
 			return new GetJobTimingResponse(Job, JobResponse, Steps, Labels);
