Horde Storage - Added logging while running a snapshot replication to better see progress from logs

diff --git a/Engine/Source/Programs/Horde/HordeStorage/Horde.Storage/Implementation/Replication/RefsReplicator.cs b/Engine/Source/Programs/Horde/HordeStorage/Horde.Storage/Implementation/Replication/RefsReplicator.cs
--- a/Engine/Source/Programs/Horde/HordeStorage/Horde.Storage/Implementation/Replication/RefsReplicator.cs
+++ b/Engine/Source/Programs/Horde/HordeStorage/Horde.Storage/Implementation/Replication/RefsReplicator.cs
@@ -292,6 +292,9 @@ namespace Horde.Storage.Implementation
                     Interlocked.Increment(ref countOfObjectsCurrentlyReplicating);
                     LogReplicationHeartbeat(countOfObjectsCurrentlyReplicating);
 
+                    if (countOfObjectsReplicated % 100 == 0)
+                        _logger.Information("Snapshot replication still running, replicated {CountOfObjects} of {TotalCountOfObjects}", countOfObjectsReplicated, snapshot.LiveObjectsCount);
+
                     Info.CountOfRunningReplications = countOfObjectsCurrentlyReplicating;
                     Info.LastRun = DateTime.Now;
 
diff --git a/Engine/Source/Programs/Horde/HordeStorage/Horde.Storage/Implementation/TransactionLog/ReplicationLogSnapshot.cs b/Engine/Source/Programs/Horde/HordeStorage/Horde.Storage/Implementation/TransactionLog/ReplicationLogSnapshot.cs
--- a/Engine/Source/Programs/Horde/HordeStorage/Horde.Storage/Implementation/TransactionLog/ReplicationLogSnapshot.cs
+++ b/Engine/Source/Programs/Horde/HordeStorage/Horde.Storage/Implementation/TransactionLog/ReplicationLogSnapshot.cs
@@ -72,6 +72,11 @@ namespace Horde.Storage.Implementation.TransactionLog
             get { return _liveObjects.Values; }
         }
 
+        public int LiveObjectsCount
+        {
+            get { return _liveObjects.Values.Count; }
+        }
+
         public static async Task<ReplicationLogSnapshot> DeserializeSnapshot(Stream stream)
         {
             await using GZipStream gzipStream = new GZipStream(stream, CompressionMode.Decompress);
