Horde Storage - More tweaks to the snapshot replication logging

diff --git a/Engine/Source/Programs/Horde/HordeStorage/Horde.Storage/Implementation/Replication/RefsReplicator.cs b/Engine/Source/Programs/Horde/HordeStorage/Horde.Storage/Implementation/Replication/RefsReplicator.cs
--- a/Engine/Source/Programs/Horde/HordeStorage/Horde.Storage/Implementation/Replication/RefsReplicator.cs
+++ b/Engine/Source/Programs/Horde/HordeStorage/Horde.Storage/Implementation/Replication/RefsReplicator.cs
@@ -293,7 +293,7 @@ namespace Horde.Storage.Implementation
                     LogReplicationHeartbeat(countOfObjectsCurrentlyReplicating);
 
                     if (countOfObjectsReplicated % 100 == 0)
-                        _logger.Information("Snapshot replication still running, replicated {CountOfObjects} of {TotalCountOfObjects}", countOfObjectsReplicated, snapshot.LiveObjectsCount);
+                        _logger.Information("{Name} Snapshot replication still running, replicated {CountOfObjects} of {TotalCountOfObjects}. Replicating {Namespace}", _name, countOfObjectsReplicated, snapshot.LiveObjectsCount, ns);
 
                     Info.CountOfRunningReplications = countOfObjectsCurrentlyReplicating;
                     Info.LastRun = DateTime.Now;
