Horde Storage - Changed legacy replication to use http client for replication of blobs, fixes issue when replicating very large (1.5 GB blob)

diff --git a/Engine/Source/Programs/Horde/HordeStorage/Horde.Storage/Implementation/Replication/Replicator.cs b/Engine/Source/Programs/Horde/HordeStorage/Horde.Storage/Implementation/Replication/Replicator.cs
--- a/Engine/Source/Programs/Horde/HordeStorage/Horde.Storage/Implementation/Replication/Replicator.cs
+++ b/Engine/Source/Programs/Horde/HordeStorage/Horde.Storage/Implementation/Replication/Replicator.cs
@@ -5,11 +5,13 @@ using System.Collections.Generic;
 using System.IO;
 using System.Linq;
 using System.Net;
+using System.Net.Http;
 using System.Threading;
 using System.Threading.Tasks;
 using Datadog.Trace;
 using EpicGames.Horde.Storage;
 using Jupiter;
+using Jupiter.Common.Implementation;
 using Jupiter.Implementation;
 using Microsoft.Extensions.Options;
 using Newtonsoft.Json;
@@ -25,12 +27,12 @@ namespace Horde.Storage.Implementation
     {
         private readonly ITransactionLogWriter _transactionLogWriter;
 
-        public ReplicatorV1(ReplicatorSettings replicatorSettings, IOptionsMonitor<ReplicationSettings> replicationSettings, IOptionsMonitor<JupiterSettings> jupiterSettings, IBlobService blobService, ITransactionLogWriter transactionLogWriter, IServiceCredentials serviceCredentials) : base(replicatorSettings, replicationSettings, jupiterSettings, blobService, CreateRemoteClient(replicatorSettings, serviceCredentials))
+        public ReplicatorV1(ReplicatorSettings replicatorSettings, IOptionsMonitor<ReplicationSettings> replicationSettings, IOptionsMonitor<JupiterSettings> jupiterSettings, IBlobService blobService, ITransactionLogWriter transactionLogWriter, IServiceCredentials serviceCredentials, IHttpClientFactory httpClientFactory) : base(replicatorSettings, replicationSettings, jupiterSettings, blobService, CreateRemoteClient(replicatorSettings, serviceCredentials), serviceCredentials, httpClientFactory)
         {
             _transactionLogWriter = transactionLogWriter;
         }
 
-        internal ReplicatorV1(ReplicatorSettings replicatorSettings, IOptionsMonitor<ReplicationSettings> replicationSettings, IOptionsMonitor<JupiterSettings> jupiterSettings, IBlobService blobService, ITransactionLogWriter transactionLogWriter, IRestClient remoteClient) : base(replicatorSettings, replicationSettings, jupiterSettings, blobService, remoteClient)
+        internal ReplicatorV1(ReplicatorSettings replicatorSettings, IOptionsMonitor<ReplicationSettings> replicationSettings, IOptionsMonitor<JupiterSettings> jupiterSettings, IBlobService blobService, ITransactionLogWriter transactionLogWriter, IRestClient remoteClient, IServiceCredentials serviceCredentials, IHttpClientFactory httpClientFactory) : base(replicatorSettings, replicationSettings, jupiterSettings, blobService, remoteClient, serviceCredentials, httpClientFactory)
         {
             _transactionLogWriter = transactionLogWriter;
         }
@@ -100,6 +102,8 @@ namespace Horde.Storage.Implementation
         private readonly CancellationTokenSource _replicationTokenSource = new CancellationTokenSource();
         private readonly IReplicator.ReplicatorInfo _replicatorInfo;
         protected readonly IRestClient _restClient;
+        private readonly IServiceCredentials _serviceCredentials;
+        private readonly HttpClient _httpClient;
 
 
         public IReplicator.ReplicatorState State { get; private set; }
@@ -116,7 +120,7 @@ namespace Horde.Storage.Implementation
             return remoteClient;
         }
 
-        protected Replicator(ReplicatorSettings replicatorSettings, IOptionsMonitor<ReplicationSettings> replicationSettings, IOptionsMonitor<JupiterSettings> jupiterSettings, IBlobService blobService, IRestClient remoteClient)
+        protected Replicator(ReplicatorSettings replicatorSettings, IOptionsMonitor<ReplicationSettings> replicationSettings, IOptionsMonitor<JupiterSettings> jupiterSettings, IBlobService blobService, IRestClient remoteClient, IServiceCredentials serviceCredentials, IHttpClientFactory httpClientFactory)
         {
             _name = replicatorSettings.ReplicatorName;
             _blobService = blobService;
@@ -138,6 +142,11 @@ namespace Horde.Storage.Implementation
 
             _replicatorInfo = new IReplicator.ReplicatorInfo(_replicatorSettings.ReplicatorName, _replicatorSettings.NamespaceToReplicate, State);
             _restClient = remoteClient;
+
+            _serviceCredentials = serviceCredentials;
+
+            _httpClient = httpClientFactory.CreateClient();
+            _httpClient.BaseAddress = new Uri(replicatorSettings.ConnectionString);
         }
 
         ~Replicator()
@@ -341,16 +350,12 @@ namespace Horde.Storage.Implementation
                     // this because new transactions are written to callisto first (to establish the GC handle) before content is uploaded to io.
                     // as such its possible (though not very likely) when we replicate the very newest transaction that we find a transaction but no content
                     // and as the content upload can be large it can take a little time for it to exist in io.
-                    byte[]? rawContent = null;
                     BlobIdentifier? calculatedBlob = null;
                     const int MaxAttempts = 3;
                     for (int attempts = 0; attempts < MaxAttempts; attempts++)
                     {
-                        RestRequest remoteIoGet = new RestRequest("api/v1/s/{ns}/{blob}");
-                        remoteIoGet.AddUrlSegment("ns", ns);
-                        remoteIoGet.AddUrlSegment("blob", blob);
-
-                        IRestResponse response = await remoteClient.ExecuteGetAsync(remoteIoGet, replicationToken);
+                        HttpRequestMessage blobRequest = BuildHttpRequest(HttpMethod.Get, $"api/v1/s/{ns}/{blob}");
+                        HttpResponseMessage response = await _httpClient.SendAsync(blobRequest, HttpCompletionOption.ResponseHeadersRead, replicationToken);
 
                         if (response.StatusCode == HttpStatusCode.NotFound)
                         {
@@ -377,46 +382,29 @@ namespace Horde.Storage.Implementation
                             return;
                         }
 
-                        if (response.IsSuccessful)
+                        if (response.IsSuccessStatusCode)
                         {
-                            rawContent = response.RawBytes;
+                            using FilesystemBufferedPayload payload = await FilesystemBufferedPayload.Create(await response.Content.ReadAsStreamAsync(replicationToken));
 
-                            calculatedBlob = BlobIdentifier.FromBlob(rawContent);
+                            {
+                                await using Stream s = payload.GetStream();
+                                calculatedBlob = await BlobIdentifier.FromStream(s);
+                            }
                             if (!blob.Equals(calculatedBlob))
                             {
-                                _logger.Warning("Mismatching blob when replicating {Blob}. Determined Hash was {Hash} size was {Size} HttpStatusCode {StatusCode} HttpMessage {HttpMessage} ResponseUri {Url}", blob, calculatedBlob, rawContent.LongLength, response.StatusCode, response.ErrorMessage, response.ResponseUri);
+                                _logger.Warning("Mismatching blob when replicating {Blob}. Determined Hash was {Hash} size was {Size} HttpStatusCode {StatusCode}",
+                                    blob, calculatedBlob, payload.Length, response.StatusCode);
                                 continue; // attempt to replicate again
                             }
-    
-                            // this blob is good and should be added to the blob store
-                            break;
-                        }
 
-                        if (response.ErrorException is WebException we)
-                        {
-                            if (we.Status == WebExceptionStatus.Timeout)
-                            {
-                                const string template = "Operation timed out while replicating {Blob}, retrying. Retry attempt {Attempts}";
-                                _logger.Information(we, template, blob, attempts);
-                                continue;
-                            }
+                            await _blobService.PutObject(ns, payload, calculatedBlob);
+                            break;
                         }
 
-                        throw new Exception($"Replicator \"{_name}\" failed to replicate blob {blob} due unsuccessful status from remote blob store: {response.StatusCode} . Error message: \"{response.ErrorMessage}\"", response.ErrorException);
+                        throw new Exception($"Replicator \"{_name}\" failed to replicate blob {blob} due unsuccessful status from remote blob store: {response.StatusCode} . Error message: \"{response.ReasonPhrase}\"");
                     }
 
-                    if (rawContent == null || calculatedBlob == null)
-                    {
-                        _logger.Warning("Remote blob {Blob} not present in remote blob store after multiple attempts to replicate. Assuming this blob has been GCed and continuing replication.", blob);
-                        return;
-                    }
 
-                    if (!blob.Equals(calculatedBlob))
-                    {
-                        _logger.Warning("Mismatching blob when replicating {Blob}. Determined Hash was {Hash} size was {Size}. Multiple attempts failed, giving up.", blob, calculatedBlob, rawContent.LongLength);
-                        return; 
-                    }
-                    await _blobService.PutObject(ns, rawContent, calculatedBlob);
                 }, replicationToken);
             }
 
@@ -438,6 +426,14 @@ namespace Horde.Storage.Implementation
             serializer.Serialize(writer, newState);
         }
 
+        private HttpRequestMessage BuildHttpRequest(HttpMethod httpMethod, string uri)
+        {
+            string? token = _serviceCredentials.GetToken();
+            HttpRequestMessage request = new HttpRequestMessage(httpMethod, uri);
+            if (!string.IsNullOrEmpty(token))
+                request.Headers.Add("Authorization", "Bearer " + token);
+            return request;
+        }
         public void Dispose()
         {
             SaveState(_stateFile, State);
diff --git a/Engine/Source/Programs/Horde/HordeStorage/Horde.Storage/Tests/UnitTests/Horde.Storage.Tests.Unit.csproj b/Engine/Source/Programs/Horde/HordeStorage/Horde.Storage/Tests/UnitTests/Horde.Storage.Tests.Unit.csproj
--- a/Engine/Source/Programs/Horde/HordeStorage/Horde.Storage/Tests/UnitTests/Horde.Storage.Tests.Unit.csproj
+++ b/Engine/Source/Programs/Horde/HordeStorage/Horde.Storage/Tests/UnitTests/Horde.Storage.Tests.Unit.csproj
@@ -17,6 +17,7 @@
     <PackageReference Include="Microsoft.AspNetCore.TestHost" Version="6.0.0" />
     <PackageReference Include="Microsoft.NET.Test.Sdk" Version="17.0.0" />
     <PackageReference Include="Moq" Version="4.16.1" />
+    <PackageReference Include="Moq.Contrib.HttpClient" Version="1.3.0" />
     <PackageReference Include="MSTest.TestAdapter" Version="2.2.7" />
     <PackageReference Include="MSTest.TestFramework" Version="2.2.7" />
     <PackageReference Include="SerilogAnalyzer" Version="0.15.0" />
diff --git a/Engine/Source/Programs/Horde/HordeStorage/Horde.Storage/Tests/UnitTests/ReplicatorTests.cs b/Engine/Source/Programs/Horde/HordeStorage/Horde.Storage/Tests/UnitTests/ReplicatorTests.cs
--- a/Engine/Source/Programs/Horde/HordeStorage/Horde.Storage/Tests/UnitTests/ReplicatorTests.cs
+++ b/Engine/Source/Programs/Horde/HordeStorage/Horde.Storage/Tests/UnitTests/ReplicatorTests.cs
@@ -5,6 +5,7 @@ using System.Collections.Generic;
 using System.IO;
 using System.Linq;
 using System.Net;
+using System.Net.Http;
 using System.Text;
 using System.Threading;
 using Microsoft.VisualStudio.TestTools.UnitTesting;
@@ -17,6 +18,8 @@ using Microsoft.Extensions.Options;
 using Moq;
 using RestSharp;
 using EpicGames.Horde.Storage;
+using Horde.Storage.Controllers;
+using Moq.Contrib.HttpClient;
 
 namespace Horde.Storage.UnitTests
 {
@@ -46,7 +49,7 @@ namespace Horde.Storage.UnitTests
         {
             ReplicatorSettings replicatorSettings = new ReplicatorSettings
             {
-                NamespaceToReplicate = NamespaceV1, ReplicatorName = ReplicatorNameV1, Version = ReplicatorVersion.V1
+                NamespaceToReplicate = NamespaceV1, ReplicatorName = ReplicatorNameV1, Version = ReplicatorVersion.V1, ConnectionString = "http://localhost"
             };
 
             byte[] contents = Encoding.ASCII.GetBytes("test string");
@@ -94,7 +97,12 @@ namespace Horde.Storage.UnitTests
             IOptionsMonitor<ReplicationSettings> replicationSettingsMonitor = Mock.Of<IOptionsMonitor<ReplicationSettings>>(_ => _.CurrentValue == replicationSettings);
             IOptionsMonitor<JupiterSettings> jupiterSettingsMonitor = Mock.Of<IOptionsMonitor<JupiterSettings>>(_ => _.CurrentValue == new JupiterSettings());
 
-            using IReplicator replicator = new ReplicatorV1(replicatorSettings, replicationSettingsMonitor, jupiterSettingsMonitor, blobStoreMock.Object, transactionLogWriter.Object, remoteClientMock.Object);
+            Mock<HttpMessageHandler> handler = new Mock<HttpMessageHandler>();
+            handler.SetupRequest($"http://localhost/api/v1/s/{NamespaceV1}/{blobToReplication}").ReturnsResponse(contents, "application/octet-stream");
+
+            IHttpClientFactory httpClientFactory = handler.CreateClientFactory();
+
+            using IReplicator replicator = new ReplicatorV1(replicatorSettings, replicationSettingsMonitor, jupiterSettingsMonitor, blobStoreMock.Object, transactionLogWriter.Object, remoteClientMock.Object, serviceCredentials, httpClientFactory);
 
             Assert.IsNull(replicator.State.ReplicatorOffset,"Expected state to have been reset during test initialize");
 
@@ -111,8 +119,7 @@ namespace Horde.Storage.UnitTests
             transactionLogWriter.Verify(writer => writer.Add(ns, It.Is<AddTransactionEvent>(e => e.Blobs[0].Equals(blobToReplication))), Times.Once);
 
             // as this was a add operation we should have transferred the blob from the remote blob store to the local
-            remoteClientMock.Verify(client => client.ExecuteGetAsync(It.IsAny<IRestRequest>(), It.IsAny<CancellationToken>()), Times.Once);
-            blobStoreMock.Verify(blobStore => blobStore.PutObject(ns, It.IsAny<byte[]>(), blobToReplication), Times.Once);
+            blobStoreMock.Verify(blobStore => blobStore.PutObject(ns, It.IsAny<IBufferedPayload>(), blobToReplication), Times.Once);
 
             Assert.AreEqual(_currentLogGeneration, replicator.State.ReplicatingGeneration);
 
@@ -125,11 +132,11 @@ namespace Horde.Storage.UnitTests
         {
             ReplicatorSettings replicatorSettings = new ReplicatorSettings
             {
-                NamespaceToReplicate = NamespaceV1, ReplicatorName = ReplicatorNameV1, Version = ReplicatorVersion.V1
+                NamespaceToReplicate = NamespaceV1, ReplicatorName = ReplicatorNameV1, Version = ReplicatorVersion.V1, ConnectionString = "http://localhost"
             };
 
             Mock<IRestClient> remoteClientMock = new Mock<IRestClient> { DefaultValue = DefaultValue.Empty };
-
+            IServiceCredentials serviceCredentials = Mock.Of<IServiceCredentials>();
             CallistoReader.CallistoGetResponse[] responses = new[]
             {
                 new CallistoReader.CallistoGetResponse(
@@ -151,7 +158,10 @@ namespace Horde.Storage.UnitTests
             IOptionsMonitor<ReplicationSettings> replicationSettingsMonitor = Mock.Of<IOptionsMonitor<ReplicationSettings>>(_ => _.CurrentValue == replicationSettings);
             IOptionsMonitor<JupiterSettings> jupiterSettingsMonitor = Mock.Of<IOptionsMonitor<JupiterSettings>>(_ => _.CurrentValue == new JupiterSettings());
 
-            using IReplicator replicator = new ReplicatorV1(replicatorSettings, replicationSettingsMonitor, jupiterSettingsMonitor, blobStoreMock.Object, transactionLogWriter.Object, remoteClientMock.Object);
+            Mock<HttpMessageHandler> handler = new Mock<HttpMessageHandler>();
+            IHttpClientFactory httpClientFactory = handler.CreateClientFactory();
+
+            using IReplicator replicator = new ReplicatorV1(replicatorSettings, replicationSettingsMonitor, jupiterSettingsMonitor, blobStoreMock.Object, transactionLogWriter.Object, remoteClientMock.Object, serviceCredentials, httpClientFactory);
 
             Assert.IsNull(replicator.State.ReplicatorOffset,"Expected state to have been reset during test initialize");
 
@@ -175,7 +185,7 @@ namespace Horde.Storage.UnitTests
         {
             ReplicatorSettings replicatorSettings = new ReplicatorSettings
             {
-                NamespaceToReplicate = NamespaceV1, ReplicatorName = ReplicatorNameV1, Version = ReplicatorVersion.V1
+                NamespaceToReplicate = NamespaceV1, ReplicatorName = ReplicatorNameV1, Version = ReplicatorVersion.V1, ConnectionString = "http://localhost"
             };
 
             byte[] contents = Encoding.ASCII.GetBytes("test string");
@@ -226,7 +236,11 @@ namespace Horde.Storage.UnitTests
             IOptionsMonitor<ReplicationSettings> replicationSettingsMonitor = Mock.Of<IOptionsMonitor<ReplicationSettings>>(_ => _.CurrentValue == replicationSettings);
             IOptionsMonitor<JupiterSettings> jupiterSettingsMonitor = Mock.Of<IOptionsMonitor<JupiterSettings>>(_ => _.CurrentValue == new JupiterSettings());
 
-            using IReplicator replicator = new ReplicatorV1(replicatorSettings, replicationSettingsMonitor, jupiterSettingsMonitor, blobStoreMock.Object, transactionLogWriter.Object, remoteClientMock.Object);
+            Mock<HttpMessageHandler> handler = new Mock<HttpMessageHandler>();
+            handler.SetupRequest($"http://localhost/api/v1/s/{NamespaceV1}/{blobToReplication}").ReturnsResponse(contents, "application/octet-stream");
+            IHttpClientFactory httpClientFactory = handler.CreateClientFactory();
+
+            using IReplicator replicator = new ReplicatorV1(replicatorSettings, replicationSettingsMonitor, jupiterSettingsMonitor, blobStoreMock.Object, transactionLogWriter.Object, remoteClientMock.Object, serviceCredentials, httpClientFactory);
 
             Assert.IsNull(replicator.State.ReplicatorOffset,"Expected state to have been reset during test initialize");
 
@@ -243,8 +257,7 @@ namespace Horde.Storage.UnitTests
             transactionLogWriter.Verify(writer => writer.Add(ns, It.Is<AddTransactionEvent>(e => e.Blobs[0].Equals(blobToReplication))), Times.Once);
 
             // as this was a add operation we should have transferred the blob from the remote blob store to the local
-            remoteClientMock.Verify(client => client.ExecuteGetAsync(It.IsAny<IRestRequest>(), It.IsAny<CancellationToken>()), Times.Exactly(2));
-            blobStoreMock.Verify(blobStore => blobStore.PutObject(ns, It.IsAny<byte[]>(), blobToReplication), Times.Once);
+            blobStoreMock.Verify(blobStore => blobStore.PutObject(ns, It.IsAny<IBufferedPayload>(), blobToReplication), Times.Once);
 
             Assert.AreEqual(_currentLogGeneration, replicator.State.ReplicatingGeneration);
 
@@ -257,7 +270,7 @@ namespace Horde.Storage.UnitTests
         {
             ReplicatorSettings replicatorSettings = new ReplicatorSettings
             {
-                NamespaceToReplicate = NamespaceV1, ReplicatorName = ReplicatorNameV1, Version = ReplicatorVersion.V1
+                NamespaceToReplicate = NamespaceV1, ReplicatorName = ReplicatorNameV1, Version = ReplicatorVersion.V1, ConnectionString = "http://localhost"
             };
 
             NamespaceId ns = replicatorSettings.NamespaceToReplicate;
@@ -298,7 +311,9 @@ namespace Horde.Storage.UnitTests
             IOptionsMonitor<ReplicationSettings> replicationSettingsMonitor = Mock.Of<IOptionsMonitor<ReplicationSettings>>(_ => _.CurrentValue == replicationSettings);
             IOptionsMonitor<JupiterSettings> jupiterSettingsMonitor = Mock.Of<IOptionsMonitor<JupiterSettings>>(_ => _.CurrentValue == new JupiterSettings());
 
-            using IReplicator replicator = new ReplicatorV1(replicatorSettings, replicationSettingsMonitor, jupiterSettingsMonitor, blobStoreMock.Object, transactionLogWriter.Object, remoteClientMock.Object);
+            Mock<HttpMessageHandler> handler = new Mock<HttpMessageHandler>();
+            IHttpClientFactory httpClientFactory = handler.CreateClientFactory();
+            using IReplicator replicator = new ReplicatorV1(replicatorSettings, replicationSettingsMonitor, jupiterSettingsMonitor, blobStoreMock.Object, transactionLogWriter.Object, remoteClientMock.Object, serviceCredentials, httpClientFactory);
 
             Assert.IsNull(replicator.State.ReplicatorOffset, "Expected state to have been reset during test initialize");
 
