Horde Storage - Added endpoint to query any instance for which peers exist
Can be used by clients to determine if there is a better region or route to be used for connections

diff --git a/Engine/Source/Programs/Horde/HordeStorage/Horde.Storage/Controllers/StatusController.cs b/Engine/Source/Programs/Horde/HordeStorage/Horde.Storage/Controllers/StatusController.cs
--- a/Engine/Source/Programs/Horde/HordeStorage/Horde.Storage/Controllers/StatusController.cs
+++ b/Engine/Source/Programs/Horde/HordeStorage/Horde.Storage/Controllers/StatusController.cs
@@ -1,21 +1,14 @@
 ﻿// Copyright Epic Games, Inc. All Rights Reserved.
 
-using System;
 using System.Collections.Generic;
-using System.ComponentModel.DataAnnotations;
-using System.IO;
 using System.Linq;
 using System.Reflection;
-using System.Threading;
-using System.Threading.Tasks;
-using Datadog.Trace.DuckTyping;
 using Horde.Storage.Implementation;
 using Jupiter;
-using Jupiter.Implementation;
 using Microsoft.AspNetCore.Authorization;
 using Microsoft.AspNetCore.Mvc;
 using Microsoft.Extensions.Options;
-using StatsdClient;
+using Newtonsoft.Json;
 
 namespace Horde.Storage.Controllers
 {
@@ -25,18 +18,20 @@ namespace Horde.Storage.Controllers
     {
         private readonly VersionFile _versionFile;
         private readonly IOptionsMonitor<JupiterSettings> _jupiterSettings;
+        private readonly IOptionsMonitor<ClusterSettings> _clusterSettings;
 
-        public StatusController(VersionFile versionFile, IOptionsMonitor<JupiterSettings> jupiterSettings)
+        public StatusController(VersionFile versionFile, IOptionsMonitor<JupiterSettings> jupiterSettings, IOptionsMonitor<ClusterSettings> clusterSettings)
         {
             _versionFile = versionFile;
             _jupiterSettings = jupiterSettings;
+            _clusterSettings = clusterSettings;
         }
 
         /// <summary>
-        /// Fetch information about Jupiter
+        /// Fetch information about Horde Storage
         /// </summary>
         /// <remarks>
-        /// General information about the Jupiter service, which version it is running and more.
+        /// General information about the service, which version it is running and more.
         /// </remarks>
         /// <returns></returns>
         [HttpGet("")]
@@ -62,10 +57,70 @@ namespace Horde.Storage.Controllers
             return new string[]
             {
                 "transactionlog",
-                "ddc",
-                "tree"
+                "ddc"
             };
         }
+
+        /// <summary>
+        /// Fetch information about other deployments
+        /// </summary>
+        /// <remarks>
+        /// General information about the Jupiter service, which version it is running and more.
+        /// </remarks>
+        /// <returns></returns>
+        [HttpGet("peers")]
+        [Authorize("Any")]
+        [ProducesResponseType(type: typeof(PeersResponse), 200)]
+        public IActionResult Peers([FromQuery] bool includeInternalEndpoints = false)
+        {
+            return Ok(new PeersResponse(_jupiterSettings, _clusterSettings, includeInternalEndpoints));
+        }
+    }
+
+    public class PeersResponse
+    {
+        [JsonConstructor]
+        public PeersResponse(string currentSite, List<KnownPeer> peers)
+        {
+            CurrentSite = currentSite;
+            Peers = peers;
+        }
+
+        public PeersResponse(IOptionsMonitor<JupiterSettings> jupiterSettings, IOptionsMonitor<ClusterSettings> clusterSettings, bool includeInternalEndpoints)
+        {
+            CurrentSite = jupiterSettings.CurrentValue.CurrentSite;
+            Peers = clusterSettings.CurrentValue.Peers.Select(settings => new KnownPeer(settings, includeInternalEndpoints)).ToList();
+        }
+
+        public string CurrentSite { get; set; }
+
+        public List<KnownPeer> Peers { get; set; } = new List<KnownPeer>();
+    }
+
+    public class KnownPeer
+    {
+        [JsonConstructor]
+        public KnownPeer(string site, string fullName, List<string> endpoints)
+        {
+            Site = site;
+            FullName = fullName;
+            Endpoints = endpoints;
+        }
+
+        public KnownPeer(PeerSettings peerSettings, bool includeInternalEndpoints)
+        {
+            Site = peerSettings.Name;
+            FullName = peerSettings.FullName;
+            IEnumerable<PeerEndpoints> endpoints = peerSettings.Endpoints;
+            if (!includeInternalEndpoints)
+                endpoints = endpoints.Where(s => !s.IsInternal);
+            Endpoints = endpoints.Select(e => e.Url).ToList();
+        }
+
+        public string Site { get; set; }
+        public string FullName { get; set; }
+
+        public List<string> Endpoints { get; set; }
     }
 
     public class StatusResponse
diff --git a/Engine/Source/Programs/Horde/HordeStorage/Horde.Storage/Horde.Storage.Startup.cs b/Engine/Source/Programs/Horde/HordeStorage/Horde.Storage/Horde.Storage.Startup.cs
--- a/Engine/Source/Programs/Horde/HordeStorage/Horde.Storage/Horde.Storage.Startup.cs
+++ b/Engine/Source/Programs/Horde/HordeStorage/Horde.Storage/Horde.Storage.Startup.cs
@@ -171,6 +171,7 @@ namespace Horde.Storage
             services.AddOptions<ConsistencyCheckSettings>().Bind(Configuration.GetSection("ConsistencyCheck")).ValidateDataAnnotations();
 
             services.AddOptions<UpstreamRelaySettings>().Configure(o => Configuration.GetSection("Upstream").Bind(o)).ValidateDataAnnotations();
+            services.AddOptions<ClusterSettings>().Configure(o => Configuration.GetSection("Cluster").Bind(o)).ValidateDataAnnotations();
 
             services.AddOptions<GCSettings>().Configure(o => Configuration.GetSection("GC").Bind(o)).ValidateDataAnnotations();
             services.AddOptions<ReplicationSettings>().Configure(o => Configuration.GetSection("Replication").Bind(o)).ValidateDataAnnotations();
diff --git a/Engine/Source/Programs/Horde/HordeStorage/Horde.Storage/Settings.cs b/Engine/Source/Programs/Horde/HordeStorage/Horde.Storage/Settings.cs
--- a/Engine/Source/Programs/Horde/HordeStorage/Horde.Storage/Settings.cs
+++ b/Engine/Source/Programs/Horde/HordeStorage/Horde.Storage/Settings.cs
@@ -267,4 +267,24 @@ namespace Horde.Storage
         [Required] public string ConnectionString { get; set; } = null!;
     }
 
+    public class ClusterSettings
+    {
+        public List<PeerSettings> Peers { get; set; } = new List<PeerSettings>();
+    }
+
+    public class PeerSettings
+    {
+        [Required] public string Name { get; set; } = null!;
+
+        [Required] public string FullName { get; set; } = null!;
+
+        public List<PeerEndpoints> Endpoints { get; set; } = new List<PeerEndpoints>();
+    }
+
+    public class PeerEndpoints
+    { 
+        [Required] public string Url { get; set; } = null!;
+
+        public bool IsInternal { get; set; } = false;
+    }
 }
