Horde: Add 24 hour heads up notification for when shared device checkout will expire
#rnx
#skipcis

diff --git a/Engine/Source/Programs/Horde/Horde.Build/Devices/IDeviceCollection.cs b/Engine/Source/Programs/Horde/Horde.Build/Devices/IDeviceCollection.cs
--- a/Engine/Source/Programs/Horde/Horde.Build/Devices/IDeviceCollection.cs
+++ b/Engine/Source/Programs/Horde/Horde.Build/Devices/IDeviceCollection.cs
@@ -222,6 +222,11 @@ namespace HordeServer.Collections
 		/// </summary>
 		public Task<List<(UserId, IDevice)>?> ExpireCheckedOutAsync();
 
+		/// <summary>
+		/// Gets a list of users to notify whose device checkout is about to expire
+		/// </summary>
+		public Task<List<(UserId, IDevice)>?> ExpireNotificatonsAsync();
+
 	}
 
 }
\ No newline at end of file
diff --git a/Engine/Source/Programs/Horde/Horde.Build/Devices/Impl/DeviceCollection.cs b/Engine/Source/Programs/Horde/Horde.Build/Devices/Impl/DeviceCollection.cs
--- a/Engine/Source/Programs/Horde/Horde.Build/Devices/Impl/DeviceCollection.cs
+++ b/Engine/Source/Programs/Horde/Horde.Build/Devices/Impl/DeviceCollection.cs
@@ -197,6 +197,9 @@ namespace HordeServer.Collections.Impl
 			[BsonIgnoreIfNull]
 			public DateTime? CheckOutTime { get; set; }
 
+			[BsonIgnoreIfNull]
+			public bool? CheckoutExpiringNotificationSent { get; set; }
+
 			[BsonIgnoreIfNull]
 			public DateTime? ProblemTimeUtc { get; set; }
 
@@ -249,6 +252,8 @@ namespace HordeServer.Collections.Impl
 
 		readonly IMongoCollection<DeviceReservationDocument> Reservations;
 
+		readonly int CheckoutDays = 7;
+
 		/// <summary>
 		/// Constructor
 		/// </summary>
@@ -421,11 +426,16 @@ namespace HordeServer.Collections.Impl
 
             string? UserId = CheckedOutByUserId?.ToString();
 
-            Updates.Add(UpdateBuilder.Set(x => x.CheckedOutByUser, string.IsNullOrEmpty(UserId) ? null : UserId));
+            Updates.Add(UpdateBuilder.Set(x => x.CheckedOutByUser, string.IsNullOrEmpty(UserId) ? null : UserId));			
 
 			if (CheckedOutByUserId != null)
 			{
 				Updates.Add(UpdateBuilder.Set(x => x.CheckOutTime, DateTime.UtcNow));
+				Updates.Add(UpdateBuilder.Set(x => x.CheckoutExpiringNotificationSent, false));
+			}
+			else
+			{
+				Updates.Add(UpdateBuilder.Set(x => x.CheckoutExpiringNotificationSent, true));
 			}
 
 			await Devices.FindOneAndUpdateAsync<DeviceDocument>(x => x.Id == DeviceId, UpdateBuilder.Combine(Updates));
@@ -707,7 +717,7 @@ namespace HordeServer.Collections.Impl
 			List<DeviceDocument> CheckedOutDevices = await Devices.Find(Filter).ToListAsync();
 
 			DateTime UtcNow = DateTime.UtcNow;
-			List<DeviceDocument> ExpiredDevices = CheckedOutDevices.FindAll(x => (UtcNow - x.CheckOutTime!.Value).TotalDays >= 5).ToList();
+			List<DeviceDocument> ExpiredDevices = CheckedOutDevices.FindAll(x => (UtcNow - x.CheckOutTime!.Value).TotalDays >= CheckoutDays).ToList();
 
 			if (ExpiredDevices.Count > 0)
 			{
@@ -719,6 +729,7 @@ namespace HordeServer.Collections.Impl
 
 				DeviceUpdates.Add(DeviceBuilder.Set(x => x.CheckedOutByUser, null));
 				DeviceUpdates.Add(DeviceBuilder.Set(x => x.CheckOutTime, null));
+				DeviceUpdates.Add(DeviceBuilder.Set(x => x.CheckoutExpiringNotificationSent, true));
 
 				UpdateResult Result = await Devices.UpdateManyAsync(Builders<DeviceDocument>.Filter.In(x => x.Id, ExpiredDevices.Select(y => y.Id)), DeviceBuilder.Combine(DeviceUpdates));
 
@@ -733,6 +744,41 @@ namespace HordeServer.Collections.Impl
 
 		}
 
+		/// <summary>
+		/// Gets a list of users to notify  that their device is about to expire in the next 24 hours
+		/// </summary>
+		public async Task<List<(UserId, IDevice)>?> ExpireNotificatonsAsync()
+		{
+			FilterDefinition<DeviceDocument> Filter = Builders<DeviceDocument>.Filter.Empty;
+			Filter &= Builders<DeviceDocument>.Filter.Where(x => x.CheckedOutByUser != null && x.CheckoutExpiringNotificationSent != true);
+
+			List<DeviceDocument> CheckedOutDevices = await Devices.Find(Filter).ToListAsync();
+
+			DateTime UtcNow = DateTime.UtcNow;
+			List<DeviceDocument> ExpiredDevices = CheckedOutDevices.FindAll(x => (UtcNow - x.CheckOutTime!.Value).TotalDays >= (CheckoutDays - 1)).ToList();
+
+			if (ExpiredDevices.Count > 0)
+			{
+				FilterDefinition<DeviceDocument> UpdateFilter = Builders<DeviceDocument>.Filter.In(x => x.Id, ExpiredDevices.Select(y => y.Id));
+
+				UpdateDefinitionBuilder<DeviceDocument> DeviceBuilder = Builders<DeviceDocument>.Update;
+				List<UpdateDefinition<DeviceDocument>> DeviceUpdates = new List<UpdateDefinition<DeviceDocument>>();
+
+				DeviceUpdates.Add(DeviceBuilder.Set(x => x.CheckoutExpiringNotificationSent, true));
+
+				UpdateResult Result = await Devices.UpdateManyAsync(Builders<DeviceDocument>.Filter.In(x => x.Id, ExpiredDevices.Select(y => y.Id)), DeviceBuilder.Combine(DeviceUpdates));
+
+				if (Result.ModifiedCount > 0)
+				{
+					return ExpiredDevices.Select(x => (UserId.Parse(x.CheckedOutByUser!), (IDevice)x)).ToList();
+				}
+
+			}
+
+			return null;
+
+		}
+
 
 		/// <inheritdoc/>
 		public async Task<List<IDeviceReservation>> FindAllReservationsAsync()
diff --git a/Engine/Source/Programs/Horde/Horde.Build/Devices/Impl/DeviceService.cs b/Engine/Source/Programs/Horde/Horde.Build/Devices/Impl/DeviceService.cs
--- a/Engine/Source/Programs/Horde/Horde.Build/Devices/Impl/DeviceService.cs
+++ b/Engine/Source/Programs/Horde/Horde.Build/Devices/Impl/DeviceService.cs
@@ -175,12 +175,22 @@ namespace HordeServer.Services
 			if (!StoppingToken.IsCancellationRequested)
 			{
 				await Devices.ExpireReservationsAsync();
+
+				List<(UserId, IDevice)>? ExpireNotifications = await Devices.ExpireNotificatonsAsync();
+				if (ExpireNotifications != null && ExpireNotifications.Count > 0)
+				{
+					foreach ((UserId, IDevice) ExpiredDevice in ExpireNotifications)
+					{
+						await NotifyDeviceServiceAsync($"Device {ExpiredDevice.Item2.PlatformId.ToString().ToUpperInvariant()} / {ExpiredDevice.Item2.Name} checkout will expire in 24 hours.  Please visit https://horde.devtools.epicgames.com/devices to renew the checkout if needed.", null, null, null, ExpiredDevice.Item1);
+					}
+				}
+
 				List<(UserId, IDevice)>? ExpireCheckouts = await Devices.ExpireCheckedOutAsync();
 				if (ExpireCheckouts != null && ExpireCheckouts.Count > 0)
 				{
 					foreach ((UserId, IDevice) ExpiredDevice in ExpireCheckouts)
 					{
-						await NotifyDeviceServiceAsync($"Device {ExpiredDevice.Item2.Name}:{ExpiredDevice.Item2.PlatformId} checkout has automatically expired and it has been returned to the shared pool.", null, null, null, ExpiredDevice.Item1);
+						await NotifyDeviceServiceAsync($"Device {ExpiredDevice.Item2.PlatformId.ToString().ToUpperInvariant()} / {ExpiredDevice.Item2.Name} checkout has expired.  The device has been returned to the shared pool and should no longer be accessed.  Please visit https://horde.devtools.epicgames.com/devices to checkout devices as needed.", null, null, null, ExpiredDevice.Item1);
 					}
 				}
 			}
