Horde: Make version modal query current server and agent versions
#skipci
#rnx

diff --git a/Engine/Source/Programs/Horde/HordeDashboard/src/backend/Backend.ts b/Engine/Source/Programs/Horde/HordeDashboard/src/backend/Backend.ts
--- a/Engine/Source/Programs/Horde/HordeDashboard/src/backend/Backend.ts
+++ b/Engine/Source/Programs/Horde/HordeDashboard/src/backend/Backend.ts
@@ -1462,13 +1462,6 @@ export class Backend {
                 return;
             }
 
-            try {
-                this.serverInfo = await this.getServerInfo();
-                this.agentSoftwareInfo = await this.getAgentSoftwareChannel();
-            } catch (reason) {
-                console.error(`Error getting server/agent info from server, defaults used: ${reason}`);
-            }
-
             await dashboard.update();
 
             if (dashboard.localCache) {
@@ -1493,10 +1486,6 @@ export class Backend {
     updateID?: any;
     logout: boolean = false;
 
-    agentSoftwareInfo: GetAgentSoftwareChannelResponse = { modifiedTime: "" };
-    serverInfo: GetServerInfoResponse = { serverVersion: "0", osDescription: "Unknown", singleInstance: false };
-
-
     private backend: Fetch;
 
 }
