Horde: Make general xge build failure a systemic error
#rnx
#skipci

diff --git a/Engine/Source/Programs/Horde/Horde.Agent/Parser/Matchers/XoreaxEventMatcher.cs b/Engine/Source/Programs/Horde/Horde.Agent/Parser/Matchers/XoreaxEventMatcher.cs
--- a/Engine/Source/Programs/Horde/Horde.Agent/Parser/Matchers/XoreaxEventMatcher.cs
+++ b/Engine/Source/Programs/Horde/Horde.Agent/Parser/Matchers/XoreaxEventMatcher.cs
@@ -27,7 +27,7 @@ namespace HordeAgent.Parser.Matchers
 			if (Cursor.IsMatch(@"BUILD FAILED: (.*)xgConsole\.exe(.*)"))
 			{
 				LogEventBuilder Builder = new LogEventBuilder(Cursor);
-				return Builder.ToMatch(LogEventPriority.High, LogLevel.Information, KnownLogEvents.Systemic_Xge_BuildFailed);
+				return Builder.ToMatch(LogEventPriority.High, LogLevel.Error, KnownLogEvents.Systemic_Xge_BuildFailed);
 			}
 
 			if (Cursor.IsMatch(@"^\s*--------------------Build System Warning[- ]"))
