Horde: Speculative job view fix
#rnx
#skipci

diff --git a/Engine/Source/Programs/Horde/HordeDashboard/src/backend/JobHandler.ts b/Engine/Source/Programs/Horde/HordeDashboard/src/backend/JobHandler.ts
--- a/Engine/Source/Programs/Horde/HordeDashboard/src/backend/JobHandler.ts
+++ b/Engine/Source/Programs/Horde/HordeDashboard/src/backend/JobHandler.ts
@@ -1,7 +1,7 @@
 // Copyright Epic Games, Inc. All Rights Reserved.
 
 import { action, observable } from 'mobx';
-import moment from 'moment-timezone';
+import moment from 'moment';
 import backend from '../backend';
 import { JobData, JobState, JobStepOutcome, JobStreamQuery, StreamData } from '../backend/Api';
 import graphCache, { GraphQuery } from '../backend/GraphCache';
@@ -14,7 +14,7 @@ const jobsRefreshTime = 10000;
 
 export class JobHandler {
 
-    constructor(includeBatches: boolean = false, jobLimit?:number) {
+    constructor(includeBatches: boolean = false, jobLimit?: number) {
         this.includeBatches = includeBatches;
         this.jobLimit = jobLimit;
     }
@@ -75,7 +75,7 @@ export class JobHandler {
 
         let wasUpdated = false;
 
-        try {            
+        try {
 
             // discover new/updated jobs
             let filter = "id,streamId,name,change,preflightChange,templateId,templateHash,graphHash,startedByUserInfo,abortedByUserInfo,createTime,state,arguments,updateTime,labels,defaultLabel";
@@ -98,6 +98,8 @@ export class JobHandler {
 
             if (bumpCount) {
 
+                this.modifiedAfter = undefined;
+
                 let maxCreate: Date | undefined;
 
                 // find max create time
@@ -118,12 +120,10 @@ export class JobHandler {
 
                 query.modifiedAfter = this.modifiedAfter;
 
-            }            
+            }
 
             const cancelId = this.cancelId++;
 
-            const queryTime = moment.utc().toISOString();
-
             const mjobs = await backend.getStreamJobs(this.stream!.id, query, false);
 
             // check for canceled after modified test
@@ -220,10 +220,6 @@ export class JobHandler {
                 wasUpdated = true;
             }
 
-            if (!this.bumpCount) {
-                this.modifiedAfter = queryTime;
-            }
-
         } catch (reason) {
             console.log(reason);
         } finally {
@@ -234,13 +230,29 @@ export class JobHandler {
 
                 if (this.count === this.jobs.length && this.count < this.jobLimit) {
                     this.bumpCount = true;
-                }    
+                }
+            }
+
+            if (this.jobs.length) {
+                let job = this.jobs[0];
+                this.jobs.forEach(j => {
+                    const date1 = new Date(job.updateTime);
+                    const date2 = new Date(j.updateTime);
+                    if (date2.getTime() > date1.getTime()) {
+                        job = j;
+                    }
+                });
+
+                this.modifiedAfter = moment(job.updateTime).add(1, 'milliseconds').toDate().toISOString();
+            } else {
+                this.modifiedAfter = undefined;
             }
 
+
             if (wasUpdated) {
                 this.setUpdated();
             }
-            
+
             this.updating = false;
         }
 
@@ -323,7 +335,7 @@ export class JobHandler {
     includePreflights = true;
     private preflightStartedByUserId?: string;
 
-    private jobLimit?:number;
+    private jobLimit?: number;
 
     initial = true;
 
