Horde Storage - Added option to disable legacy blob cleanup and disabled this in prod, as I suspect this is the cause of the large memory usage on our worker nodes.

diff --git a/Engine/Source/Programs/Horde/HordeStorage/Horde.Storage/Implementation/GC/BlobCleanupService.cs b/Engine/Source/Programs/Horde/HordeStorage/Horde.Storage/Implementation/GC/BlobCleanupService.cs
--- a/Engine/Source/Programs/Horde/HordeStorage/Horde.Storage/Implementation/GC/BlobCleanupService.cs
+++ b/Engine/Source/Programs/Horde/HordeStorage/Horde.Storage/Implementation/GC/BlobCleanupService.cs
@@ -35,7 +35,7 @@ namespace Horde.Storage.Implementation
         {
             _settings = settings;
 
-            if (settings.CurrentValue.CleanOldBlobs)
+            if (settings.CurrentValue.CleanOldBlobsLegacy)
             {
                 OrphanBlobCleanup orphanBlobCleanup = provider.GetService<OrphanBlobCleanup>()!;
                 RegisterCleanup(orphanBlobCleanup);
diff --git a/Engine/Source/Programs/Horde/HordeStorage/Horde.Storage/Settings.cs b/Engine/Source/Programs/Horde/HordeStorage/Horde.Storage/Settings.cs
--- a/Engine/Source/Programs/Horde/HordeStorage/Horde.Storage/Settings.cs
+++ b/Engine/Source/Programs/Horde/HordeStorage/Horde.Storage/Settings.cs
@@ -253,6 +253,7 @@ namespace Horde.Storage
         
         public bool CleanOldRefRecords { get; set; } = false;
         public bool CleanOldBlobs { get; set; } = true;
+        public bool CleanOldBlobsLegacy { get; set; } = false;
 
         public TimeSpan LastAccessCutoff { get; set; } = TimeSpan.FromDays(14);
 
@@ -260,6 +261,7 @@ namespace Horde.Storage
         public TimeSpan RefCleanupPollFrequency { get; set; } = TimeSpan.FromMinutes(60);
         public List<string> CleanNamespacesV1 { get; set; } = new List<string>();
         public List<string> CleanNamespaces { get; set; }  = new List<string>();
+
     }
 
     public class UpstreamRelaySettings
diff --git a/Engine/Source/Programs/Horde/HordeStorage/Horde.Storage/Tests/Functional/GC/GCBlobTests.cs b/Engine/Source/Programs/Horde/HordeStorage/Horde.Storage/Tests/Functional/GC/GCBlobTests.cs
--- a/Engine/Source/Programs/Horde/HordeStorage/Horde.Storage/Tests/Functional/GC/GCBlobTests.cs
+++ b/Engine/Source/Programs/Horde/HordeStorage/Horde.Storage/Tests/Functional/GC/GCBlobTests.cs
@@ -78,7 +78,8 @@ namespace Horde.Storage.FunctionalTests.GC
                     {
                         GCSettings gcSettings = new GCSettings()
                         {
-                            CleanOldBlobs = true
+                            CleanOldBlobs = true,
+                            CleanOldBlobsLegacy = true,
                         };
                         IOptionsMonitor<GCSettings> gcSettingsMon = Mock.Of<IOptionsMonitor<GCSettings>>(_ => _.CurrentValue == gcSettings);
 
@@ -168,6 +169,7 @@ namespace Horde.Storage.FunctionalTests.GC
             GCSettings gcSettings = new GCSettings()
             {
                 CleanOldBlobs = true,
+                CleanOldBlobsLegacy = true,
                 CleanNamespacesV1 = new List<string> { TestNamespace.ToString() }
             };
             IOptionsMonitor<GCSettings> gcSettingsMon = Mock.Of<IOptionsMonitor<GCSettings>>(_ => _.CurrentValue == gcSettings);
@@ -237,6 +239,7 @@ namespace Horde.Storage.FunctionalTests.GC
                     new KeyValuePair<string, string>("Horde_Storage:StorageImplementations:0", "MemoryBlobStore"),
                     new KeyValuePair<string, string>("GC:CleanNamespaces:0", TestNamespace.ToString()),
                     new KeyValuePair<string, string>("GC:CleanOldBlobs", true.ToString()),
+                    new KeyValuePair<string, string>("GC:CleanOldBlobsLegacy", true.ToString()),
                     new KeyValuePair<string, string>("Horde_Storage:BlobIndexImplementation", GetImplementation()),
                     
                 })
