Horde:
- trends refactor using batching
- add in pages, zoom, cost

diff --git a/Engine/Source/Programs/Horde/HordeDashboard/src/backend/Api.ts b/Engine/Source/Programs/Horde/HordeDashboard/src/backend/Api.ts
--- a/Engine/Source/Programs/Horde/HordeDashboard/src/backend/Api.ts
+++ b/Engine/Source/Programs/Horde/HordeDashboard/src/backend/Api.ts
@@ -296,6 +296,13 @@ export type JobQuery = {
 	count?: number;
 };
 
+export type JobTimingsQuery = {
+	streamId?: string;
+	template?: string[];
+	filter?: string;
+	count?: number;
+};
+
 
 export type IssueQuery = {
 	jobId?: string;
@@ -2328,12 +2335,18 @@ export type GetGraphResponse = {
 /**The timing info for  */
 export type GetJobTimingResponse = {
 
+	job: JobData;
+	jobResponse: JobData;
 	/**Timing info for each step */
 	steps: { [key: string]: GetStepTimingInfoResponse };
 
 	/**Timing information for each label */
 	labels: GetLabelTimingInfoResponse[];
 }
+/** batch timing info */
+export type FindJobTimingsResponse = {
+	timings: { [jobId: string]: GetJobTimingResponse };
+}
 
 /**Information about the timing info for a label */
 export type GetLabelTimingInfoResponse = GetTimingInfoResponse &
diff --git a/Engine/Source/Programs/Horde/HordeDashboard/src/backend/Backend.ts b/Engine/Source/Programs/Horde/HordeDashboard/src/backend/Backend.ts
--- a/Engine/Source/Programs/Horde/HordeDashboard/src/backend/Backend.ts
+++ b/Engine/Source/Programs/Horde/HordeDashboard/src/backend/Backend.ts
@@ -3,7 +3,7 @@
 import { configure } from "mobx";
 import { isNumber } from 'util';
 import templateCache from '../backend/TemplateCache';
-import { AgentData, ArtifactData, AuditLogEntry, AuditLogQuery, BatchUpdatePoolRequest, ChangeSummaryData, CreateDeviceRequest, CreateDeviceResponse, CreateJobRequest, CreateJobResponse, CreatePoolRequest, CreateSoftwareResponse, CreateSubscriptionRequest, CreateSubscriptionResponse, DashboardPreference, EventData, FindIssueResponse, GetAgentSoftwareChannelResponse, GetArtifactZipRequest, GetDevicePlatformResponse, GetDevicePoolResponse, GetDeviceReservationResponse, GetDeviceResponse, GetGraphResponse, GetIssueStreamResponse, GetJiraIssueResponse, GetJobsTabResponse, GetJobStepRefResponse, GetJobStepTraceResponse, GetJobTimingResponse, GetLogEventResponse, GetNotificationResponse, GetPerforceServerStatusResponse, GetServerInfoResponse, GetServerSettingsResponse, GetSoftwareResponse, GetSubscriptionResponse, GetUserResponse, GetUtilizationTelemetryResponse, GlobalConfig, IssueData, IssueQuery, IssueQueryV2, JobData, JobQuery, JobsTabColumnType, JobStreamQuery, LeaseData, LogData, LogLineData, PoolData, ProjectData, ScheduleData, ScheduleQuery, SearchLogFileResponse, ServerUpdateResponse, SessionData, StreamData, TabType, TemplateData, TestData, UpdateAgentRequest, UpdateDeviceRequest, UpdateGlobalConfigRequest, UpdateIssueRequest, UpdateJobRequest, UpdateLeaseRequest, UpdateNotificationsRequest, UpdatePoolRequest, UpdateServerSettingsRequest, UpdateStepRequest, UpdateStepResponse, UpdateUserRequest, UsersQuery } from './Api';
+import { AgentData, ArtifactData, AuditLogEntry, AuditLogQuery, BatchUpdatePoolRequest, ChangeSummaryData, CreateDeviceRequest, CreateDeviceResponse, CreateJobRequest, CreateJobResponse, CreatePoolRequest, CreateSoftwareResponse, CreateSubscriptionRequest, CreateSubscriptionResponse, DashboardPreference, EventData, FindIssueResponse, FindJobTimingsResponse, GetAgentSoftwareChannelResponse, GetArtifactZipRequest, GetDevicePlatformResponse, GetDevicePoolResponse, GetDeviceReservationResponse, GetDeviceResponse, GetGraphResponse, GetIssueStreamResponse, GetJiraIssueResponse, GetJobsTabResponse, GetJobStepRefResponse, GetJobStepTraceResponse, GetJobTimingResponse, GetLogEventResponse, GetNotificationResponse, GetPerforceServerStatusResponse, GetServerInfoResponse, GetServerSettingsResponse, GetSoftwareResponse, GetSubscriptionResponse, GetUserResponse, GetUtilizationTelemetryResponse, GlobalConfig, IssueData, IssueQuery, IssueQueryV2, JobData, JobQuery, JobsTabColumnType, JobStreamQuery, JobTimingsQuery, LeaseData, LogData, LogLineData, PoolData, ProjectData, ScheduleData, ScheduleQuery, SearchLogFileResponse, ServerUpdateResponse, SessionData, StreamData, TabType, TemplateData, TestData, UpdateAgentRequest, UpdateDeviceRequest, UpdateGlobalConfigRequest, UpdateIssueRequest, UpdateJobRequest, UpdateLeaseRequest, UpdateNotificationsRequest, UpdatePoolRequest, UpdateServerSettingsRequest, UpdateStepRequest, UpdateStepResponse, UpdateUserRequest, UsersQuery } from './Api';
 import dashboard from './Dashboard';
 import { ChallengeStatus, Fetch } from './Fetch';
 import graphCache, { GraphQuery } from './GraphCache';
@@ -434,6 +434,34 @@ export class Backend {
 
     }
 
+    getBatchJobTiming(query: JobTimingsQuery): Promise<FindJobTimingsResponse> {
+
+        return new Promise<FindJobTimingsResponse>((resolve, reject) => {
+            this.backend.get(`/api/v1/jobs/timing`, { 
+                params: query
+            }).then((response) => {
+                let timingsWrapper = response.data as FindJobTimingsResponse;
+                let graphCalls: { jobResponse: JobData, call: any }[] = [];
+                Object.values(timingsWrapper.timings).forEach(timing => {
+                    if(!timing.jobResponse.graphHash) {
+                        return reject(`Job ${timing.job.id} has undefined graph hash`);
+                    }
+                    graphCalls.push( { jobResponse: timing.jobResponse, call: graphCache.get({ graphHash: timing.jobResponse.graphHash, jobId: timing.jobResponse.id }) });
+                });
+                let allPromises = graphCalls.map(item => item.call);
+                Promise.all(allPromises).then(responses => {
+                    for(let idx = 0; idx < responses.length; idx++) {
+                        graphCalls[idx].jobResponse.graphRef = responses[idx] as GetGraphResponse;
+                    }
+                    resolve(timingsWrapper);
+                });
+            }).catch(reason => {
+                reject(reason);
+            });
+        });
+
+    }
+
     getLogEvents(logId: string): Promise<EventData[]> {
 
         return new Promise<EventData[]>((resolve, reject) => {
