Horde Storage - Fixed build error in gc blob tests

diff --git a/Engine/Source/Programs/Horde/HordeStorage/Horde.Storage/Tests/Functional/GC/GCBlobTests.cs b/Engine/Source/Programs/Horde/HordeStorage/Horde.Storage/Tests/Functional/GC/GCBlobTests.cs
--- a/Engine/Source/Programs/Horde/HordeStorage/Horde.Storage/Tests/Functional/GC/GCBlobTests.cs
+++ b/Engine/Source/Programs/Horde/HordeStorage/Horde.Storage/Tests/Functional/GC/GCBlobTests.cs
@@ -78,7 +78,7 @@ namespace Horde.Storage.FunctionalTests.GC
                     {
                         GCSettings gcSettings = new GCSettings()
                         {
-                            CleanOldBlobs = true,
+                            CleanOldBlobs = true
                             CleanOldBlobsLegacy = true,
                         };
                         IOptionsMonitor<GCSettings> gcSettingsMon = Mock.Of<IOptionsMonitor<GCSettings>>(_ => _.CurrentValue == gcSettings);
