Horde Storage - Tests for the peers endpoints

diff --git a/Engine/Source/Programs/Horde/HordeStorage/Horde.Storage/Tests/Functional/StatusTests.cs b/Engine/Source/Programs/Horde/HordeStorage/Horde.Storage/Tests/Functional/StatusTests.cs
new file mode 100644
--- /dev/null
+++ b/Engine/Source/Programs/Horde/HordeStorage/Horde.Storage/Tests/Functional/StatusTests.cs
@@ -0,0 +1,105 @@
+// Copyright Epic Games, Inc. All Rights Reserved.
+
+using System.Collections.Generic;
+using System.Linq;
+using System.Net.Http;
+using System.Threading.Tasks;
+using Horde.Storage.Controllers;
+using Microsoft.AspNetCore.TestHost;
+using Microsoft.AspNetCore.Hosting;
+using Microsoft.Extensions.Configuration;
+using Microsoft.Extensions.DependencyInjection;
+using Microsoft.Extensions.DependencyInjection.Extensions;
+using Microsoft.Extensions.Options;
+using Microsoft.VisualStudio.TestTools.UnitTesting;
+using Serilog;
+using Serilog.Core;
+
+namespace Horde.Storage.FunctionalTests.Status
+{
+    [TestClass]
+    public class StatusTests
+    {
+        protected TestServer? _server;
+        protected HttpClient? _httpClient;
+
+        [TestInitialize]
+        public async Task Setup()
+        {
+            IConfigurationRoot configuration = new ConfigurationBuilder()
+                // we are not reading the base appSettings here as we want exact control over what runs in the tests
+                .AddJsonFile("appsettings.Testing.json", true)
+                .AddEnvironmentVariables()
+                .Build();
+
+            Logger logger = new LoggerConfiguration()
+                .ReadFrom.Configuration(configuration)
+                .CreateLogger();
+
+            TestServer server = new TestServer(new WebHostBuilder()
+                .UseConfiguration(configuration)
+                .UseEnvironment("Testing")
+                .UseSerilog(logger)
+                .ConfigureTestServices(collection =>
+                {
+                    collection.Configure<ClusterSettings>(settings =>
+                    {
+                        settings.Peers = new PeerSettings[]
+                        {
+                            new PeerSettings
+                            {
+                                Name = "use",
+                                FullName = "us-east-1",
+                                Endpoints = new PeerEndpoints[]
+                                {
+                                    new PeerEndpoints { Url = "http://use-internal.jupiter.com", IsInternal = true },
+                                    new PeerEndpoints { Url = "http://use.jupiter.com", IsInternal = false },
+                                }.ToList()
+                            }
+                        }.ToList();
+                    });
+                })
+                .UseStartup<HordeStorageStartup>()
+            );
+            _httpClient = server.CreateClient();
+            _server = server;
+
+            await Task.CompletedTask;
+        }
+
+        [TestMethod]
+        public async Task GetPeerConnection()
+        {
+            HttpResponseMessage result = await _httpClient!.GetAsync($"api/v1/status/peers");
+            result.EnsureSuccessStatusCode();
+            PeersResponse peersResponse = await result.Content.ReadAsAsync<PeersResponse>();
+            Assert.AreEqual("test", peersResponse.CurrentSite);
+
+            Assert.AreEqual(1, peersResponse.Peers.Count);
+
+            Assert.AreEqual("use", peersResponse.Peers[0].Site);
+            Assert.AreEqual("us-east-1", peersResponse.Peers[0].FullName);
+            Assert.AreEqual(1, peersResponse.Peers[0].Endpoints.Count);
+            Assert.AreEqual("http://use.jupiter.com", peersResponse.Peers[0].Endpoints[0]);
+        }
+
+        
+        [TestMethod]
+        public async Task GetPeerConnectionInternal()
+        {
+            HttpResponseMessage result = await _httpClient!.GetAsync($"api/v1/status/peers?includeInternalEndpoints=true");
+            result.EnsureSuccessStatusCode();
+            PeersResponse peersResponse = await result.Content.ReadAsAsync<PeersResponse>();
+            Assert.AreEqual("test", peersResponse.CurrentSite);
+
+            Assert.AreEqual(1, peersResponse.Peers.Count);
+
+            Assert.AreEqual("use", peersResponse.Peers[0].Site);
+            Assert.AreEqual("us-east-1", peersResponse.Peers[0].FullName);
+            Assert.AreEqual(2, peersResponse.Peers[0].Endpoints.Count);
+
+            Assert.AreEqual("http://use-internal.jupiter.com", peersResponse.Peers[0].Endpoints[0]);
+            Assert.AreEqual("http://use.jupiter.com", peersResponse.Peers[0].Endpoints[1]);
+        }
+    }
+}
