Horde Storage - Moved some apis to only be exposed on a secondary internal only port.
This includes any api that is not needed by users including everything that required a admin token.
The internal port exposes all the normal apis but also the ones that are flagged as internal.
This also removes the custom kestrel dispatcher paths (and thus the "high-speed" port as it is not used and doesnt really give any signficant speedup - using a http only port is what gives most of the benefits)

diff --git a/Engine/Source/Programs/Horde/HordeStorage/Horde.Storage/Controllers/AdminController.cs b/Engine/Source/Programs/Horde/HordeStorage/Horde.Storage/Controllers/AdminController.cs
--- a/Engine/Source/Programs/Horde/HordeStorage/Horde.Storage/Controllers/AdminController.cs
+++ b/Engine/Source/Programs/Horde/HordeStorage/Horde.Storage/Controllers/AdminController.cs
@@ -19,6 +19,7 @@ namespace Horde.Storage.Controllers
 {
     [ApiController]
     [Route("api/v1/admin")]
+    [InternalApiFilter]
     public class AdminController : Controller
     {
         private readonly LastAccessService _lastAccessService;
diff --git a/Engine/Source/Programs/Horde/HordeStorage/Horde.Storage/Controllers/DDCRefController.cs b/Engine/Source/Programs/Horde/HordeStorage/Horde.Storage/Controllers/DDCRefController.cs
--- a/Engine/Source/Programs/Horde/HordeStorage/Horde.Storage/Controllers/DDCRefController.cs
+++ b/Engine/Source/Programs/Horde/HordeStorage/Horde.Storage/Controllers/DDCRefController.cs
@@ -79,15 +79,12 @@ namespace Horde.Storage.Controllers
         {
             NamespaceId[] namespaces = await _refsStore.GetNamespaces().ToArrayAsync();
 
-            if (!ShouldDoAuth())
+            // filter namespaces down to only the namespaces the user has access to
+            namespaces = namespaces.Where(ns =>
             {
-                // filter namespaces down to only the namespaces the user has access to
-                namespaces = namespaces.Where(ns =>
-                {
-                    Task<AuthorizationResult> authorizationResult = _authorizationService.AuthorizeAsync(User, ns, NamespaceAccessRequirement.Name);
-                    return authorizationResult.Result.Succeeded;
-                }).ToArray();
-            }
+                Task<AuthorizationResult> authorizationResult = _authorizationService.AuthorizeAsync(User, ns, NamespaceAccessRequirement.Name);
+                return authorizationResult.Result.Succeeded;
+            }).ToArray();
 
             return Ok(new GetNamespacesResponse(namespaces));
         }
@@ -114,16 +111,13 @@ namespace Horde.Storage.Controllers
             [FromQuery] string[] fields,
             [FromRoute] string? format = null)
         {
-            if (!ShouldDoAuth())
             {
-                using (IScope _ = Tracer.Instance.StartActive("authorize"))
-                {
-                    AuthorizationResult authorizationResult = await _authorizationService.AuthorizeAsync(User, ns, NamespaceAccessRequirement.Name);
+                using IScope _ = Tracer.Instance.StartActive("authorize");
+                AuthorizationResult authorizationResult = await _authorizationService.AuthorizeAsync(User, ns, NamespaceAccessRequirement.Name);
 
-                    if (!authorizationResult.Succeeded)
-                    {
-                        return Forbid();
-                    }
+                if (!authorizationResult.Succeeded)
+                {
+                    return Forbid();
                 }
             }
 
@@ -194,201 +188,6 @@ namespace Horde.Storage.Controllers
             }
         }
 
-        private static bool IsServedOnHighPerfHttpPort(HostString host, IEnumerable<int> highPerfPorts)
-        {
-            foreach (int port in highPerfPorts)
-            {
-                if (host.Port.HasValue && host.Port.Value == port)
-                    return true;
-            }
-            
-            return false;
-        }
-
-        [ApiExplorerSettings(IgnoreApi = true)]
-        [NonAction]
-        public async Task FastGet(HttpContext httpContext, Func<Task> next)
-        {
-            bool isHighPerf = IsServedOnHighPerfHttpPort(httpContext.Request.Host, _jupiterSettings.CurrentValue.DisableAuthOnPorts);
-
-            if (isHighPerf && httpContext.Request.Path.StartsWithSegments("/api/v1/c/ddc", out var remaining))
-            {
-                string[] parts = remaining.ToString().Split("/");
-                if (parts.Length == 4)
-                {
-                    NamespaceId ns = new NamespaceId(parts[1]);
-                    BucketId bucket = new BucketId(parts[2]);
-                    string keyAndFormat = parts[3];
-                    if (httpContext.Request.Method == "GET")
-                    {
-                        KeyId key;
-                        // assume raw key if nothing is specified
-                        string format = "raw";
-                        string[] keyParts = keyAndFormat.Split(".");
-                        if (keyParts.Length == 2)
-                        {
-                            key = new KeyId(keyParts[0]);
-                            format = keyParts[1];
-                        }
-                        else
-                        {
-                            key = new KeyId(keyAndFormat);
-                        }
-
-                        if (format != "raw")
-                        {
-                            // if not requesting the raw format we can not be used so continue searching middlewares
-                            await next.Invoke();
-                            return;
-                        }
-
-                        httpContext.Response.Headers["X-Jupiter-ImplUsed"] = "ReqMid";
-                        try
-                        {
-                            (RefResponse refRes, BlobContents? tempBlob) = await _ddcRefService.Get(ns, bucket, key, fields: new[] {"blob"});
-                            BlobContents blob = tempBlob!;
-                            
-                            httpContext.Response.StatusCode = StatusCodes.Status200OK;
-                            httpContext.Response.Headers["Content-Type"] = MediaTypeNames.Application.Octet;
-                            httpContext.Response.Headers["Content-Length"] = Convert.ToString(blob.Length);
-                            httpContext.Response.Headers[CommonHeaders.HashHeaderName] = refRes.ContentHash.ToString();
-
-                            using (IScope _ = Tracer.Instance.StartActive("body.write"))
-                            {
-                                await blob.Stream.CopyToAsync(httpContext.Response.Body);
-                            }
-                        }
-                        catch (NamespaceNotFoundException e)
-                        {
-                            httpContext.Response.StatusCode = StatusCodes.Status400BadRequest;
-                            httpContext.Response.Headers["Content-Type"] = "application/json";
-                            await httpContext.Response.WriteAsync("{\"title\": \"Namespace " + e.Namespace + " did not exist\", \"status\": 400}");
-                        }
-                        catch (RefRecordNotFoundException e)
-                        {
-                            httpContext.Response.StatusCode = StatusCodes.Status400BadRequest;
-                            httpContext.Response.Headers["Content-Type"] = "application/json";
-                            await httpContext.Response.WriteAsync("{\"title\": \"Object " + e.Bucket + " " + e.Key + " did not exist\", \"status\": 400}");
-                        }
-                        catch (BlobNotFoundException e)
-                        {
-                            httpContext.Response.StatusCode = StatusCodes.Status400BadRequest;
-                            httpContext.Response.Headers["Content-Type"] = "application/json";
-                            await httpContext.Response.WriteAsync("{\"title\": \"Object " + e.Blob + " in " + e.Ns + " not found\", \"status\": 400}");
-                        }
-
-                        return;
-                    }
-                    else if (httpContext.Request.Method == "HEAD")
-                    {
-                        KeyId key = new KeyId(keyAndFormat);
-                        httpContext.Response.Headers["X-Jupiter-ImplUsed"] = "ReqMid";
-                        try
-                        {
-                            RefRecord record = await _ddcRefService.Exists(ns, bucket, key);
-                            
-                            httpContext.Response.StatusCode = StatusCodes.Status204NoContent;
-                            httpContext.Response.Headers[CommonHeaders.HashHeaderName] = record.ContentHash.ToString();
-                        }
-                        catch (NamespaceNotFoundException e)
-                        {
-                            httpContext.Response.StatusCode = StatusCodes.Status400BadRequest;
-                            httpContext.Response.Headers["Content-Type"] = "application/json";
-                            await httpContext.Response.WriteAsync($"{{\"title\": \"Namespace {e.Namespace} did not exist\", \"status\": 400}}");
-                        }
-                        catch (RefRecordNotFoundException e)
-                        {
-                            httpContext.Response.StatusCode = StatusCodes.Status400BadRequest;
-                            httpContext.Response.Headers["Content-Type"] = "application/json";
-                            await httpContext.Response.WriteAsync($"{{\"title\": \"Object {e.Bucket} {e.Key} in namespace {e.Namespace} did not exist\", \"status\": 400}}");
-                        }
-                        catch (MissingBlobsException e)
-                        {
-                            httpContext.Response.StatusCode = StatusCodes.Status400BadRequest;
-                            httpContext.Response.Headers["Content-Type"] = "application/json";
-                            await httpContext.Response.WriteAsync($"{{\"title\": \"Blobs {e.Blobs} from object {e.Bucket} {e.Key} in namespace {e.Namespace} did not exist\", \"status\": 400}}");
-                        }
-
-                        return;
-                    }
-
-                }
-            }
-            else if (isHighPerf && httpContext.Request.Path.StartsWithSegments("/api/v1/c/ddc-rpc/batchGet", out _))
-            {
-                string requestBody;
-
-                using (StreamReader sr = new StreamReader(httpContext.Request.Body))
-                {
-                    using IScope _ = Tracer.Instance.StartActive("body.read");
-                    requestBody = await sr.ReadToEndAsync();
-                }
-
-                BatchGetOp batch;
-                using (IScope _ = Tracer.Instance.StartActive("json.deserialize"))
-                {
-                    batch = JsonConvert.DeserializeObject<BatchGetOp>(requestBody)!;
-                }
-                
-                BatchWriter writer = new BatchWriter();
-                httpContext.Response.StatusCode = StatusCodes.Status200OK;
-                await writer.WriteToStream(httpContext.Response.Body, batch.Namespace!.Value, batch.Operations.Select(op => new Tuple<BatchWriter.OpVerb, BucketId, KeyId>(op.Verb, op.Bucket!.Value, op.Key!.Value)).ToList(),
-                    async (verb, ns, bucket, key, fields) =>
-                    {
-                        BatchWriter.OpState opState;
-                        RefResponse? refResponse = null;
-                        BlobContents? blob = null;
-                        try
-                        {
-                            switch (verb)
-                            {
-                                case BatchWriter.OpVerb.GET:
-                                    (refResponse, blob) = await _ddcRefService.Get(ns, bucket, key, fields);
-                                    opState = BatchWriter.OpState.OK;
-                                    break;
-                                case BatchWriter.OpVerb.HEAD:
-                                    await _ddcRefService.Exists(ns, bucket, key);
-                                    opState = BatchWriter.OpState.Exists;
-                                    refResponse = null;
-                                    blob = null;
-                                    break;
-                                default:
-                                    throw new ArgumentOutOfRangeException(nameof(verb), verb, null);
-                            }
-                        }
-                        catch (RefRecordNotFoundException)
-                        {
-                            opState = BatchWriter.OpState.NotFound;
-                        }
-                        catch (BlobNotFoundException)
-                        {
-                            opState = BatchWriter.OpState.NotFound;
-                        }
-                        catch (MissingBlobsException)
-                        {
-                            opState = BatchWriter.OpState.NotFound;
-                        }
-                        catch (NamespaceNotFoundException)
-                        {
-                            opState = BatchWriter.OpState.NotFound;
-                        }
-                        catch (Exception e)
-                        {
-                            _logger.Error(e, "Unknown exception when executing batch get");
-
-                            // we want to make sure that we always continue to write the results even when we get errors
-                            opState = BatchWriter.OpState.Failed;
-                        }
-
-                        return new Tuple<ContentHash?, BlobContents?, BatchWriter.OpState>(refResponse?.ContentHash, blob, opState);
-                    });
-
-                return;
-            }
-
-            await next.Invoke();
-        }
-
         /// <summary>
         /// Checks if a refs key exists
         /// </summary>
@@ -405,7 +204,6 @@ namespace Horde.Storage.Controllers
             [FromRoute] [Required] BucketId bucket,
             [FromRoute] [Required] KeyId key)
         {
-            if (!ShouldDoAuth())
             {
                 AuthorizationResult authorizationResult = await _authorizationService.AuthorizeAsync(User, ns, NamespaceAccessRequirement.Name);
 
@@ -456,7 +254,6 @@ namespace Horde.Storage.Controllers
             [FromRoute] [Required] KeyId key,
             [FromBody] RefRequest refRequest)
         {
-            if (!ShouldDoAuth())
             {
                 AuthorizationResult authorizationResult = await _authorizationService.AuthorizeAsync(User, ns, NamespaceAccessRequirement.Name);
 
@@ -499,7 +296,6 @@ namespace Horde.Storage.Controllers
             [FromRoute] [Required] BucketId bucket,
             [FromRoute] [Required] KeyId key)
         {
-            if (!ShouldDoAuth())
             {
                 AuthorizationResult authorizationResult = await _authorizationService.AuthorizeAsync(User, ns, NamespaceAccessRequirement.Name);
 
@@ -833,7 +629,6 @@ namespace Horde.Storage.Controllers
                 throw new ArgumentNullException();
             }
 
-            if (!ShouldDoAuth())
             {
                 AuthorizationResult authorizationResult = await _authorizationService.AuthorizeAsync(User, batch.Namespace, NamespaceAccessRequirement.Name);
 
@@ -898,17 +693,6 @@ namespace Horde.Storage.Controllers
             // we have already set the result by writing to Response
             return new EmptyResult();
         }
-
-        private bool ShouldDoAuth()
-        {
-            foreach (int port in _jupiterSettings.CurrentValue.DisableAuthOnPorts)
-            {
-                if (port == HttpContext.Connection.LocalPort)
-                    return true;
-            }
-
-            return false;
-        }
     }
 
     public class HashMismatchException : Exception
diff --git a/Engine/Source/Programs/Horde/HordeStorage/Horde.Storage/Controllers/DebugController.cs b/Engine/Source/Programs/Horde/HordeStorage/Horde.Storage/Controllers/DebugController.cs
--- a/Engine/Source/Programs/Horde/HordeStorage/Horde.Storage/Controllers/DebugController.cs
+++ b/Engine/Source/Programs/Horde/HordeStorage/Horde.Storage/Controllers/DebugController.cs
@@ -1,10 +1,7 @@
 ﻿// Copyright Epic Games, Inc. All Rights Reserved.
 
 using System;
-using System.Net.Mime;
-using System.Threading.Tasks;
 using Microsoft.AspNetCore.Authorization;
-using Microsoft.AspNetCore.Http;
 using Microsoft.AspNetCore.Mvc;
 
 namespace Horde.Storage.Controllers
@@ -13,11 +10,6 @@ namespace Horde.Storage.Controllers
     [Route("api/v1/c/_debug")]
     public class DebugController : Controller
     {
-        private static readonly byte[] OneKilobyteData = new byte[1024];
-        private static readonly byte[] OneMegabyteData = new byte[1024 * 1024];
-        private static readonly byte[] SixteenMegabyteData = new byte[16 * 1024 * 1024];
-        private static bool s_testBuffersFilled = false;
-        
         /// <summary>
         /// Return bytes of specified length with auth, used for testing only
         /// </summary>
@@ -38,48 +30,6 @@ namespace Horde.Storage.Controllers
         {
             return GenerateByteResponse(length);
         }
-        
-        [ApiExplorerSettings(IgnoreApi = true)]
-        [NonAction]
-        public async Task FastGetBytes(HttpContext httpContext, Func<Task> next)
-        {
-            async Task Send(byte[] data)
-            {
-                httpContext.Response.StatusCode = StatusCodes.Status200OK;
-                httpContext.Response.Headers["Content-Type"] = MediaTypeNames.Application.Octet;
-                httpContext.Response.Headers["Content-Length"] = Convert.ToString(data.Length);
-                await httpContext.Response.Body.WriteAsync(data);
-            }
-
-            if (!s_testBuffersFilled)
-            {
-                Array.Fill(OneKilobyteData, (byte)'k');
-                Array.Fill(OneMegabyteData, (byte)'m');
-                Array.Fill(SixteenMegabyteData, (byte)'s');
-                s_testBuffersFilled = true;
-            }
-            
-            if (httpContext.Request.Path.StartsWithSegments("/api/v1/c/_debug/fastBytes/1byte"))
-            {
-                await Send(new byte[] {0x4A});
-            }
-            else if (httpContext.Request.Path.StartsWithSegments("/api/v1/c/_debug/fastBytes/1kilobyte"))
-            {
-                await Send(OneKilobyteData);
-            }
-            else if (httpContext.Request.Path.StartsWithSegments("/api/v1/c/_debug/fastBytes/1megabyte"))
-            {
-                await Send(OneMegabyteData);
-            }
-            else if (httpContext.Request.Path.StartsWithSegments("/api/v1/c/_debug/fastBytes/16megabyte"))
-            {
-                await Send(SixteenMegabyteData);
-            }
-            else
-            {
-                await next.Invoke();    
-            }
-        }
 
         private FileContentResult GenerateByteResponse(int length)
         {
diff --git a/Engine/Source/Programs/Horde/HordeStorage/Horde.Storage/Controllers/ReferencesController.cs b/Engine/Source/Programs/Horde/HordeStorage/Horde.Storage/Controllers/ReferencesController.cs
--- a/Engine/Source/Programs/Horde/HordeStorage/Horde.Storage/Controllers/ReferencesController.cs
+++ b/Engine/Source/Programs/Horde/HordeStorage/Horde.Storage/Controllers/ReferencesController.cs
@@ -69,15 +69,12 @@ namespace Horde.Storage.Controllers
         {
             NamespaceId[] namespaces = await _objectService.GetNamespaces().ToArrayAsync();
 
-            if (ShouldDoAuth())
+            // filter namespaces down to only the namespaces the user has access to
+            namespaces = namespaces.Where(ns =>
             {
-                // filter namespaces down to only the namespaces the user has access to
-                namespaces = namespaces.Where(ns =>
-                {
-                    Task<AuthorizationResult> authorizationResult = _authorizationService.AuthorizeAsync(User, ns, NamespaceAccessRequirement.Name);
-                    return authorizationResult.Result.Succeeded;
-                }).ToArray();
-            }
+                Task<AuthorizationResult> authorizationResult = _authorizationService.AuthorizeAsync(User, ns, NamespaceAccessRequirement.Name);
+                return authorizationResult.Result.Succeeded;
+            }).ToArray();
 
             return Ok(new GetNamespacesResponse(namespaces));
         }
@@ -99,16 +96,13 @@ namespace Horde.Storage.Controllers
             [FromQuery] string[] fields,
             [FromRoute] string? format = null)
         {
-            if (ShouldDoAuth())
             {
-                using (IScope _ = Tracer.Instance.StartActive("authorize"))
-                {
-                    AuthorizationResult authorizationResult = await _authorizationService.AuthorizeAsync(User, ns, NamespaceAccessRequirement.Name);
+                using IScope _ = Tracer.Instance.StartActive("authorize");
+                AuthorizationResult authorizationResult = await _authorizationService.AuthorizeAsync(User, ns, NamespaceAccessRequirement.Name);
 
-                    if (!authorizationResult.Succeeded)
-                    {
-                        return Forbid();
-                    }
+                if (!authorizationResult.Succeeded)
+                {
+                    return Forbid();
                 }
             }
 
@@ -226,16 +220,13 @@ namespace Horde.Storage.Controllers
             [FromRoute] [Required] IoHashKey key,
             [FromQuery] string[] fields)
         {
-            if (ShouldDoAuth())
             {
-                using (IScope _ = Tracer.Instance.StartActive("authorize"))
-                {
-                    AuthorizationResult authorizationResult = await _authorizationService.AuthorizeAsync(User, ns, NamespaceAccessRequirement.Name);
+                using IScope _ = Tracer.Instance.StartActive("authorize");
+                AuthorizationResult authorizationResult = await _authorizationService.AuthorizeAsync(User, ns, NamespaceAccessRequirement.Name);
 
-                    if (!authorizationResult.Succeeded)
-                    {
-                        return Forbid();
-                    }
+                if (!authorizationResult.Succeeded)
+                {
+                    return Forbid();
                 }
             }
 
@@ -276,16 +267,13 @@ namespace Horde.Storage.Controllers
             [FromRoute] [Required] BucketId bucket,
             [FromRoute] [Required] IoHashKey key)
         {
-            if (ShouldDoAuth())
             {
-                using (IScope _ = Tracer.Instance.StartActive("authorize"))
-                {
-                    AuthorizationResult authorizationResult = await _authorizationService.AuthorizeAsync(User, ns, NamespaceAccessRequirement.Name);
+                using IScope _ = Tracer.Instance.StartActive("authorize");
+                AuthorizationResult authorizationResult = await _authorizationService.AuthorizeAsync(User, ns, NamespaceAccessRequirement.Name);
 
-                    if (!authorizationResult.Succeeded)
-                    {
-                        return Forbid();
-                    }
+                if (!authorizationResult.Succeeded)
+                {
+                    return Forbid();
                 }
             }
 
@@ -352,16 +340,13 @@ namespace Horde.Storage.Controllers
             [FromRoute] [Required] BucketId bucket,
             [FromRoute] [Required] IoHashKey key)
         {
-            if (ShouldDoAuth())
             {
-                using (IScope _ = Tracer.Instance.StartActive("authorize"))
-                {
-                    AuthorizationResult authorizationResult = await _authorizationService.AuthorizeAsync(User, ns, NamespaceAccessRequirement.Name);
+                using IScope _ = Tracer.Instance.StartActive("authorize");
+                AuthorizationResult authorizationResult = await _authorizationService.AuthorizeAsync(User, ns, NamespaceAccessRequirement.Name);
 
-                    if (!authorizationResult.Succeeded)
-                    {
-                        return Forbid();
-                    }
+                if (!authorizationResult.Succeeded)
+                {
+                    return Forbid();
                 }
             }
 
@@ -455,16 +440,13 @@ namespace Horde.Storage.Controllers
             [FromRoute] [Required] IoHashKey key,
             [FromRoute] [Required] BlobIdentifier hash)
         {
-            if (ShouldDoAuth())
             {
-                using (IScope _ = Tracer.Instance.StartActive("authorize"))
-                {
-                    AuthorizationResult authorizationResult = await _authorizationService.AuthorizeAsync(User, ns, NamespaceAccessRequirement.Name);
+                using IScope _ = Tracer.Instance.StartActive("authorize");
+                AuthorizationResult authorizationResult = await _authorizationService.AuthorizeAsync(User, ns, NamespaceAccessRequirement.Name);
 
-                    if (!authorizationResult.Succeeded)
-                    {
-                        return Forbid();
-                    }
+                if (!authorizationResult.Succeeded)
+                {
+                    return Forbid();
                 }
             }
 
@@ -573,17 +555,6 @@ namespace Horde.Storage.Controllers
                 return NotFound(new ProblemDetails {Title = $"Namespace {e.Namespace} did not exist"});
             }
         }
-
-        private bool ShouldDoAuth()
-        {
-            foreach (int port in _jupiterSettings.CurrentValue.DisableAuthOnPorts)
-            {
-                if (port == HttpContext.Connection.LocalPort)
-                    return false;
-            }
-
-            return true;
-        }
     }
 
     public class RefMetadataResponse
diff --git a/Engine/Source/Programs/Horde/HordeStorage/Horde.Storage/Controllers/ReplicationController.cs b/Engine/Source/Programs/Horde/HordeStorage/Horde.Storage/Controllers/ReplicationController.cs
--- a/Engine/Source/Programs/Horde/HordeStorage/Horde.Storage/Controllers/ReplicationController.cs
+++ b/Engine/Source/Programs/Horde/HordeStorage/Horde.Storage/Controllers/ReplicationController.cs
@@ -18,6 +18,7 @@ namespace Horde.Storage.Controllers
     [FormatFilter]
     [Route("api/v1/g")]
     [Authorize("Admin")]
+    [InternalApiFilter]
     public class ReplicationController : ControllerBase
     {
         private readonly ReplicationService _replicationService;
diff --git a/Engine/Source/Programs/Horde/HordeStorage/Horde.Storage/Controllers/ReplicationLogController.cs b/Engine/Source/Programs/Horde/HordeStorage/Horde.Storage/Controllers/ReplicationLogController.cs
--- a/Engine/Source/Programs/Horde/HordeStorage/Horde.Storage/Controllers/ReplicationLogController.cs
+++ b/Engine/Source/Programs/Horde/HordeStorage/Horde.Storage/Controllers/ReplicationLogController.cs
@@ -21,6 +21,7 @@ namespace Horde.Storage.Controllers
 {
     [ApiController]
     [Route("api/v1/replication-log")]
+    [InternalApiFilter]
     public class ReplicationLogController : ControllerBase
     {
         private readonly IServiceProvider _provider;
diff --git a/Engine/Source/Programs/Horde/HordeStorage/Horde.Storage/Horde.Storage.Startup.cs b/Engine/Source/Programs/Horde/HordeStorage/Horde.Storage/Horde.Storage.Startup.cs
--- a/Engine/Source/Programs/Horde/HordeStorage/Horde.Storage/Horde.Storage.Startup.cs
+++ b/Engine/Source/Programs/Horde/HordeStorage/Horde.Storage/Horde.Storage.Startup.cs
@@ -131,18 +131,6 @@ namespace Horde.Storage
             });
         }
 
-        protected override void OnUseEndpoints(IWebHostEnvironment env, IEndpointRouteBuilder endpoints)
-        {
-            HordeStorageSettings settings = endpoints.ServiceProvider.GetService<IOptionsMonitor<HordeStorageSettings>>()!.CurrentValue!;
-
-            if (settings.UseNewDDCEndpoints)
-            {
-                DDCEndpoints ddcEndpoints = endpoints.ServiceProvider.GetService<DDCEndpoints>()!;
-
-                endpoints.Map(ddcEndpoints.GetRawRoute, ddcEndpoints.GetRaw).RequireAuthorization("Cache.read");
-            }
-        }
-
         protected override void OnAddService(IServiceCollection services)
         {
             // For requests served on the high-perf HTTP port there is not authn/authz.
@@ -408,21 +396,7 @@ namespace Horde.Storage
                     throw new ArgumentOutOfRangeException();
             }
         }
-
-        protected override void OnConfigureAppEarly(IApplicationBuilder app, IWebHostEnvironment env)
-        {
-            HordeStorageSettings settings = app.ApplicationServices.GetService<IOptionsMonitor<HordeStorageSettings>>()!.CurrentValue!;
-            // Register two methods as middlewares for benchmarking purposes
-            // By circumventing the full ASP.NET stack we can measure the performance difference between these calls.
-            if (!settings.UseNewDDCEndpoints)
-            {
-                DDCRefController ddcRefController = ActivatorUtilities.CreateInstance<DDCRefController>(app.ApplicationServices);
-                DebugController debugController = ActivatorUtilities.CreateInstance<DebugController>(app.ApplicationServices);
-                app.Use(ddcRefController.FastGet);
-                app.Use(debugController.FastGetBytes);
-            }
-        }
-
+        
         private ILeaderElection CreateLeaderElection(IServiceProvider provider)
         {
             HordeStorageSettings settings = provider.GetService<IOptionsMonitor<HordeStorageSettings>>()!.CurrentValue!;
diff --git a/Engine/Source/Programs/Horde/HordeStorage/Horde.Storage/Implementation/Kestrel/AsciiString.cs b/Engine/Source/Programs/Horde/HordeStorage/Horde.Storage/Implementation/Kestrel/AsciiString.cs
deleted file mode 100644
diff --git a/Engine/Source/Programs/Horde/HordeStorage/Horde.Storage/Implementation/Kestrel/KestrelDispatcher.HttpConnection.cs b/Engine/Source/Programs/Horde/HordeStorage/Horde.Storage/Implementation/Kestrel/KestrelDispatcher.HttpConnection.cs
deleted file mode 100644
diff --git a/Engine/Source/Programs/Horde/HordeStorage/Horde.Storage/Implementation/Kestrel/KestrelDispatcher.cs b/Engine/Source/Programs/Horde/HordeStorage/Horde.Storage/Implementation/Kestrel/KestrelDispatcher.cs
deleted file mode 100644
diff --git a/Engine/Source/Programs/Horde/HordeStorage/Horde.Storage/Program.cs b/Engine/Source/Programs/Horde/HordeStorage/Horde.Storage/Program.cs
--- a/Engine/Source/Programs/Horde/HordeStorage/Horde.Storage/Program.cs
+++ b/Engine/Source/Programs/Horde/HordeStorage/Horde.Storage/Program.cs
@@ -1,9 +1,5 @@
 // Copyright Epic Games, Inc. All Rights Reserved.
 
-using System;
-using System.IO.Pipelines;
-using System.Net.Sockets;
-using Horde.Storage.Implementation.Kestrel;
 using Jupiter;
 using Microsoft.Extensions.Hosting;
 
@@ -14,17 +10,12 @@ namespace Horde.Storage
     {
         public static int Main(string[] args)
         {
-            return BaseProgram<HordeStorageStartup>.BaseMain(args, HttpConnectionFactory);
-        }
-
-        private static BaseHttpConnection HttpConnectionFactory(IServiceProvider sp, PipeReader reader, PipeWriter writer, Socket socket)
-        {
-            return new KestrelDispatcher(sp, reader, writer, socket);
+            return BaseProgram<HordeStorageStartup>.BaseMain(args);
         }
 
         public static IHostBuilder CreateHostBuilder(string[] Args)
         {
-            return BaseProgram<HordeStorageStartup>.CreateHostBuilder(Args, HttpConnectionFactory);
+            return BaseProgram<HordeStorageStartup>.CreateHostBuilder(Args);
         }
     }
 }
diff --git a/Engine/Source/Programs/Horde/HordeStorage/Horde.Storage/Settings.cs b/Engine/Source/Programs/Horde/HordeStorage/Horde.Storage/Settings.cs
--- a/Engine/Source/Programs/Horde/HordeStorage/Horde.Storage/Settings.cs
+++ b/Engine/Source/Programs/Horde/HordeStorage/Horde.Storage/Settings.cs
@@ -136,7 +136,6 @@ namespace Horde.Storage
 
         public int LastAccessRollupFrequencySeconds { get; set; } = 900; // 15 minutes
 
-        public bool UseNewDDCEndpoints { get; set; } = false;
         public bool BlobStoreParallel { get; set; } = false;
     }
 
diff --git a/Engine/Source/Programs/Horde/HordeStorage/Horde.Storage/Tests/Functional/Objects/ReplicationServerTests.cs b/Engine/Source/Programs/Horde/HordeStorage/Horde.Storage/Tests/Functional/Objects/ReplicationServerTests.cs
new file mode 100644
--- /dev/null
+++ b/Engine/Source/Programs/Horde/HordeStorage/Horde.Storage/Tests/Functional/Objects/ReplicationServerTests.cs
@@ -0,0 +1,98 @@
+// Copyright Epic Games, Inc. All Rights Reserved.
+
+using System;
+using System.Net;
+using System.Net.Http;
+using System.Net.Http.Headers;
+using System.Threading.Tasks;
+using Horde.Storage.Implementation;
+using Jupiter;
+using Jupiter.Implementation;
+using Microsoft.AspNetCore.Hosting;
+using Microsoft.Extensions.Configuration;
+using Microsoft.VisualStudio.TestTools.UnitTesting;
+using Serilog;
+using Logger = Serilog.Core.Logger;
+using EpicGames.Horde.Storage;
+
+namespace Horde.Storage.FunctionalTests.References
+{
+
+    [TestClass]
+    public class ReplicationServerTests
+    {
+        private static IWebHost? _server;
+
+        protected readonly NamespaceId TestNamespace = new NamespaceId("test-namespace");
+        protected readonly BucketId TestBucket = new BucketId("default");
+
+        [TestInitialize]
+        public async Task Setup()
+
+        {
+            IConfigurationRoot configuration = new ConfigurationBuilder()
+                // we are not reading the base appSettings here as we want exact control over what runs in the tests
+                .AddJsonFile("appsettings.Testing.json", true)
+                .AddEnvironmentVariables()
+                .Build();
+
+            Logger logger = new LoggerConfiguration()
+                .ReadFrom.Configuration(configuration)
+                .CreateLogger();
+
+            IWebHostBuilder webHostBuilder = new WebHostBuilder()
+                .UseConfiguration(configuration)
+                .UseEnvironment("Testing")
+                .UseSerilog(logger)
+                .UseStartup<HordeStorageStartup>()
+                .UseKestrel(options =>
+                {
+                    options.Listen(IPAddress.Any, 80);
+                    options.Listen(IPAddress.Any, 8080);
+                })
+            ;
+            _server = webHostBuilder.Build();
+            await _server.StartAsync();
+        }
+
+        [TestMethod]
+        [Ignore("Needs to be run manually as it starts up proper http servers and that can cause port conflicts")]
+        public async Task ReplicationPublicEndpoint()
+        {
+            // insert a random object
+            CompactBinaryWriter writer = new CompactBinaryWriter();
+            writer.BeginObject();
+            writer.AddString("thisIsAField", "stringField");
+            writer.EndObject();
+
+            byte[] objectData = writer.Save();
+            BlobIdentifier objectHash = BlobIdentifier.FromBlob(objectData);
+
+            HttpContent requestContent = new ByteArrayContent(objectData);
+            requestContent.Headers.ContentType = new MediaTypeHeaderValue(CustomMediaTypeNames.UnrealCompactBinary);
+            requestContent.Headers.Add(CommonHeaders.HashHeaderName, objectHash.ToString());
+
+            {
+                using HttpClient httpClient = new HttpClient();
+                httpClient.BaseAddress = new Uri("http://localhost:8080");
+                HttpResponseMessage result = await httpClient!.PutAsync(requestUri: $"api/v1/refs/{TestNamespace}/{TestBucket}/{IoHashKey.FromName("newReferenceObject")}.uecb", requestContent);
+                result.EnsureSuccessStatusCode();
+            }
+            // access the replication-log using a public endpoint and a internal endpoint, only the internal endpoint should return results
+            {
+                using HttpClient httpClient = new HttpClient();
+                httpClient.BaseAddress = new Uri("http://localhost:8080");
+                HttpResponseMessage result = await httpClient!.GetAsync(requestUri: $"api/v1/replication-log/incremental/{TestNamespace}");
+                result.EnsureSuccessStatusCode();
+            }
+
+            {
+                using HttpClient httpClient = new HttpClient();
+                httpClient.BaseAddress = new Uri("http://localhost:80");
+
+                HttpResponseMessage result = await httpClient!.GetAsync(requestUri: $"api/v1/replication-log/incremental/{TestNamespace}");
+                Assert.AreEqual(HttpStatusCode.NotFound, result.StatusCode);
+            }
+        }
+    }
+}
diff --git a/Engine/Source/Programs/Horde/HordeStorage/Horde.Storage/Tests/UnitTests/KestrelDispatcherTest.cs b/Engine/Source/Programs/Horde/HordeStorage/Horde.Storage/Tests/UnitTests/KestrelDispatcherTest.cs
deleted file mode 100644
diff --git a/Engine/Source/Programs/Horde/HordeStorage/Horde.Storage/appsettings.Development.json b/Engine/Source/Programs/Horde/HordeStorage/Horde.Storage/appsettings.Development.json
--- a/Engine/Source/Programs/Horde/HordeStorage/Horde.Storage/appsettings.Development.json
+++ b/Engine/Source/Programs/Horde/HordeStorage/Horde.Storage/appsettings.Development.json
@@ -66,7 +66,8 @@
   },
 
   "Jupiter": {
-    "CurrentSite": "local"
+    "CurrentSite": "local",
+    "InternalApiPort": 57003
   },
 
   // disable auth and allow any namespace
@@ -92,6 +93,9 @@
       "Http2": {
         "Url": "http://0.0.0.0:57002",
         "Protocols": "Http2"
+      },
+      "Http-Internal": {
+        "Url": "http://0.0.0.0:57003"
       }
     }
   }
diff --git a/Engine/Source/Programs/Horde/HordeStorage/Horde.Storage/appsettings.Production.json b/Engine/Source/Programs/Horde/HordeStorage/Horde.Storage/appsettings.Production.json
--- a/Engine/Source/Programs/Horde/HordeStorage/Horde.Storage/appsettings.Production.json
+++ b/Engine/Source/Programs/Horde/HordeStorage/Horde.Storage/appsettings.Production.json
@@ -31,6 +31,10 @@
     "ReplicationPollFrequencySeconds": 15
   },
 
+  "Jupiter": {
+    "InternalApiPort": 8080
+  },
+
   "AllowedHosts": "*",
 
   "Kestrel": {
@@ -38,7 +42,7 @@
       "Http": {
         "Url": "http://0.0.0.0:80"
       },
-      "Http2": {
+      "InternalHttp2": {
         "Url": "http://0.0.0.0:8080",
         "Protocols": "Http2"
       }
diff --git a/Engine/Source/Programs/Horde/HordeStorage/Jupiter.Common/BaseProgram.cs b/Engine/Source/Programs/Horde/HordeStorage/Jupiter.Common/BaseProgram.cs
--- a/Engine/Source/Programs/Horde/HordeStorage/Jupiter.Common/BaseProgram.cs
+++ b/Engine/Source/Programs/Horde/HordeStorage/Jupiter.Common/BaseProgram.cs
@@ -15,34 +15,9 @@ using Serilog;
 
 namespace Jupiter
 {
-    public abstract class BaseHttpConnection
-    {
-        protected readonly IServiceProvider ServiceProvider;
-        protected readonly PipeReader Reader;
-        protected readonly PipeWriter Writer;
-        protected readonly Socket Socket;
-
-        public BaseHttpConnection(IServiceProvider serviceProvider, PipeReader reader, PipeWriter writer, Socket socket)
-        {
-            ServiceProvider = serviceProvider;
-            Reader = reader;
-            Writer = writer;
-            Socket = socket;
-        }
-
-        /// <summary>
-        /// A function that can process a connection.
-        /// </summary>
-        /// <param name="connection">A <see cref="ConnectionContext" /> representing the connection.</param>
-        /// <returns>A <see cref="Task"/> that represents the connection lifetime. When the task completes, the connection will be closed.</returns>
-        public abstract Task ExecuteAsync(ConnectionContext connection);
-    }
-    
     // ReSharper disable once UnusedMember.Global
     public static class BaseProgram<T> where T : BaseStartup
     {
-        public delegate BaseHttpConnection HttpConnectionFactory(IServiceProvider sp, PipeReader reader, PipeWriter writer, Socket socket);
-
         private static IConfiguration Configuration { get; } = GetConfigurationBuilder();
 
         private static IConfiguration GetConfigurationBuilder()
@@ -72,7 +47,7 @@ namespace Jupiter
 
         }
 
-        public static int BaseMain(string[] args, HttpConnectionFactory? httpConnectionFactory = null)
+        public static int BaseMain(string[] args)
         {
             Log.Logger = new LoggerConfiguration()
                 .ReadFrom.Configuration(Configuration)
@@ -81,7 +56,7 @@ namespace Jupiter
             try
             {
                 Log.Information("Creating ASPNET Host");
-                CreateHostBuilder(args, httpConnectionFactory).Build().Run();
+                CreateHostBuilder(args).Build().Run();
                 return 0;
             }
             catch (Exception ex)
@@ -95,7 +70,7 @@ namespace Jupiter
             }
         }
 
-        public static IHostBuilder CreateHostBuilder(string[] args, HttpConnectionFactory? httpConnectionFactory)
+        public static IHostBuilder CreateHostBuilder(string[] args)
         {
             return Host.CreateDefaultBuilder(args)
                 .ConfigureWebHostDefaults(webBuilder =>
diff --git a/Engine/Source/Programs/Horde/HordeStorage/Jupiter.Common/BaseStartup.cs b/Engine/Source/Programs/Horde/HordeStorage/Jupiter.Common/BaseStartup.cs
--- a/Engine/Source/Programs/Horde/HordeStorage/Jupiter.Common/BaseStartup.cs
+++ b/Engine/Source/Programs/Horde/HordeStorage/Jupiter.Common/BaseStartup.cs
@@ -4,7 +4,6 @@ using System;
 using System.Collections.Generic;
 using System.ComponentModel.DataAnnotations;
 using System.IO;
-using System.Linq;
 using System.Net.Mime;
 using System.Reflection;
 using System.Text;
@@ -12,7 +11,6 @@ using System.Threading.Tasks;
 using Amazon;
 using Datadog.Trace;
 using Jupiter.Common;
-using Jupiter.Implementation;
 using Microsoft.AspNetCore.Authentication;
 using Microsoft.AspNetCore.Authorization;
 using Microsoft.AspNetCore.Builder;
@@ -21,7 +19,6 @@ using Microsoft.AspNetCore.Hosting;
 using Microsoft.AspNetCore.Http;
 using Microsoft.AspNetCore.HttpOverrides;
 using Microsoft.AspNetCore.Mvc;
-using Microsoft.AspNetCore.Routing;
 using Microsoft.Extensions.Configuration;
 using Microsoft.Extensions.DependencyInjection;
 using Microsoft.Extensions.Diagnostics.HealthChecks;
@@ -243,10 +240,7 @@ namespace Jupiter
         public void Configure(IApplicationBuilder app, IWebHostEnvironment env)
         {
             JupiterSettings jupiterSettings = app.ApplicationServices.GetService<IOptionsMonitor<JupiterSettings>>()!.CurrentValue;
-            
-            bool IsHttpPort(HttpContext httpContext) => !IsHighPerfHttpPort(httpContext);
-            bool IsHighPerfHttpPort(HttpContext httpContext) => jupiterSettings.DisableAuthOnPorts.Any(port => port == httpContext.Connection.LocalPort);
-            
+
             if (jupiterSettings.ShowPII)
             {
                 _logger.Error("Personally Identifiable information being shown. This should not be generally enabled in prod.");
@@ -255,37 +249,24 @@ namespace Jupiter
                 Microsoft.IdentityModel.Logging.IdentityModelEventSource.ShowPII = true;
             }
 
-            if (jupiterSettings.DisableAuthOnPorts.Any())
-            {
-                _logger.Information("Auth is disabled on port(s) {Ports}", jupiterSettings.DisableAuthOnPorts);
-            }
-
-            app.MapWhen(IsHttpPort, httpApp =>
-            {
-                ConfigureMiddlewares(false, jupiterSettings, httpApp, env);
-            });
-            
-            app.MapWhen(IsHighPerfHttpPort, highPerfApp =>
-            {
-                ConfigureMiddlewares(true, jupiterSettings, highPerfApp, env);
-            });
+            ConfigureMiddlewares(jupiterSettings, app, env);
         }
 
-        private void ConfigureMiddlewares(bool isHighPerf, JupiterSettings jupiterSettings, IApplicationBuilder app, IWebHostEnvironment env)
+        private void ConfigureMiddlewares(JupiterSettings jupiterSettings, IApplicationBuilder app, IWebHostEnvironment env)
         {
             // enable use of forwarding headers as we expect a reverse proxy to be running in front of us
             //app.UseMiddleware<DatadogTraceMiddleware>("ForwardedHeaders");
             app.UseForwardedHeaders();
 
-            if (!isHighPerf && jupiterSettings.UseRequestLogging)
+            if (jupiterSettings.UseRequestLogging)
             {
                 //app.UseMiddleware<DatadogTraceMiddleware>("RequestLogging");
                 app.UseSerilogRequestLogging();
             }
 
             OnConfigureAppEarly(app, env);
-                
-            if (!isHighPerf && env.IsDevelopment() && UseDeveloperExceptionPage)
+
+            if (env.IsDevelopment() && UseDeveloperExceptionPage)
             {
                 app.UseDeveloperExceptionPage();
             }
@@ -297,13 +278,10 @@ namespace Jupiter
             //app.UseMiddleware<DatadogTraceMiddleware>("Routing");
             app.UseRouting();
 
-            if (!isHighPerf)
-            {
-                //app.UseMiddleware<DatadogTraceMiddleware>("Authentication");
-                app.UseAuthentication();
-                //app.UseMiddleware<DatadogTraceMiddleware>("Authorization");
-                app.UseAuthorization();
-            }
+            //app.UseMiddleware<DatadogTraceMiddleware>("Authentication");
+            app.UseAuthentication();
+            //app.UseMiddleware<DatadogTraceMiddleware>("Authorization");
+            app.UseAuthorization();
             
             //app.UseMiddleware<DatadogTraceMiddleware>("Endpoints");
             app.UseEndpoints(endpoints =>
@@ -335,12 +313,11 @@ namespace Jupiter
                     context.Response.Headers.ContentType = "text/plain";
                     await context.Response.Body.WriteAsync(Encoding.ASCII.GetBytes("Healthy"));
                 });
-                OnUseEndpoints(env, endpoints);
-                
+
                 endpoints.MapControllers();
             });
 
-            if (!isHighPerf && jupiterSettings.HostSwaggerDocumentation)
+            if (jupiterSettings.HostSwaggerDocumentation)
             {
                 //app.UseMiddleware<DatadogTraceMiddleware>("Swagger");
                 app.UseSwagger();
@@ -350,11 +327,6 @@ namespace Jupiter
             OnConfigureApp(app, env);
         }
 
-        protected virtual void OnUseEndpoints(IWebHostEnvironment env, IEndpointRouteBuilder endpoints)
-        {
-            
-        }
-
         public virtual bool UseDeveloperExceptionPage { get; } = true;
 
         protected virtual void OnConfigureAppEarly(IApplicationBuilder app, IWebHostEnvironment env)
@@ -458,6 +430,8 @@ namespace Jupiter
 
     public class JupiterSettings
     {
+        public const int DefaultInternalPort = 8080;
+
         /// <summary>
         /// If the request is smaller then MemoryBufferSize we buffer it in memory rather then as a file
         /// </summary>
@@ -468,7 +442,7 @@ namespace Jupiter
         public bool DisableHealthChecks { get; set; } = false;
         public bool HostSwaggerDocumentation { get; set; } = true;
 
-        public List<int> DisableAuthOnPorts { get; set; } = new();
+        public int InternalApiPort { get; set; } = DefaultInternalPort;
 
         // Enable to echo every request to the log file, usually this is more efficiently done on the load balancer
         public bool UseRequestLogging { get; set; } = false;
@@ -479,8 +453,6 @@ namespace Jupiter
         [Required]
         [Key]
         public string CurrentSite { get; set; } = "";
-
-        //public Dictionary<string, int> DisableAuthOnPorts { get; set; } = new ();
     }
 
     public class NamespaceSettings
@@ -509,7 +481,7 @@ namespace Jupiter
         {
             using IScope _ = Tracer.Instance.StartActive(_scopeName);
 
-            //Move to next delegate/middleware in the pipleline
+            //Move to next delegate/middleware in the pipeline
             await _next.Invoke(httpContext);
         }
     }
diff --git a/Engine/Source/Programs/Horde/HordeStorage/Jupiter.Common/InternalApiFilter.cs b/Engine/Source/Programs/Horde/HordeStorage/Jupiter.Common/InternalApiFilter.cs
new file mode 100644
--- /dev/null
+++ b/Engine/Source/Programs/Horde/HordeStorage/Jupiter.Common/InternalApiFilter.cs
@@ -0,0 +1,28 @@
+// Copyright Epic Games, Inc. All Rights Reserved.
+
+using System;
+using Jupiter;
+using Microsoft.AspNetCore.Mvc;
+using Microsoft.AspNetCore.Mvc.Filters;
+using Microsoft.Extensions.DependencyInjection;
+using Microsoft.Extensions.Options;
+
+public class InternalApiFilter : Attribute, IResourceFilter
+{
+    public void OnResourceExecuting(ResourceExecutingContext context)
+    {
+        IOptionsMonitor<JupiterSettings>? settings = context.HttpContext.RequestServices.GetService<IOptionsMonitor<JupiterSettings>>();
+
+        bool isInternalPort = context.HttpContext.Connection.LocalPort == settings!.CurrentValue.InternalApiPort || 
+            /* unit tests do not run on ports, we consider them always on the internal port */ (context.HttpContext.Connection.LocalPort == 0 && context.HttpContext.Connection.LocalIpAddress == null);
+        if (!isInternalPort)
+        {
+            // this endpoint should only be exposed on the internal port, so we return a 404 as this is not on the internal port
+            context.Result = new NotFoundResult();
+        }
+    }
+
+    public void OnResourceExecuted(ResourceExecutedContext context)
+    {
+    }
+}
