Horde Storage - Make sure to dispose of the streams that are buffered to the filesystem when we are finished with them.
Also added a warning for if we are unable to resolve the references of a object, as there is nothing we can do to recover so we should just continue running the replication.

diff --git a/Engine/Source/Programs/Horde/HordeStorage/Horde.Storage/Implementation/Blob/IBlobService.cs b/Engine/Source/Programs/Horde/HordeStorage/Horde.Storage/Implementation/Blob/IBlobService.cs
--- a/Engine/Source/Programs/Horde/HordeStorage/Horde.Storage/Implementation/Blob/IBlobService.cs
+++ b/Engine/Source/Programs/Horde/HordeStorage/Horde.Storage/Implementation/Blob/IBlobService.cs
@@ -354,7 +354,8 @@ public class BlobService : IBlobService
         {
             using IScope _ = Tracer.Instance.StartActive("HierarchicalStore.Populate");
             // buffer the blob contents as it could be to large for us to buffer into memory before forwarding to the blob stores
-            using FilesystemBufferedPayload payload = await FilesystemBufferedPayload.Create(blobContents.Stream);
+            await using Stream s = blobContents.Stream;
+            using FilesystemBufferedPayload payload = await FilesystemBufferedPayload.Create(s);
 
             // Don't populate the last store, as that is where we got the hit
             // Populate each store traversed that did not have the content found lower in the hierarchy
diff --git a/Engine/Source/Programs/Horde/HordeStorage/Horde.Storage/Implementation/Replication/RefsReplicator.cs b/Engine/Source/Programs/Horde/HordeStorage/Horde.Storage/Implementation/Replication/RefsReplicator.cs
--- a/Engine/Source/Programs/Horde/HordeStorage/Horde.Storage/Implementation/Replication/RefsReplicator.cs
+++ b/Engine/Source/Programs/Horde/HordeStorage/Horde.Storage/Implementation/Replication/RefsReplicator.cs
@@ -267,11 +267,15 @@ namespace Horde.Storage.Implementation
 
             // process snapshot
             // fetch the snapshot from the remote blob store
-            HttpRequestMessage request = BuildHttpRequest(HttpMethod.Get, $"api/v1/blobs/{blobNamespace}/{snapshotBlob}");
-            HttpResponseMessage response = await _httpClient.SendAsync(request, HttpCompletionOption.ResponseHeadersRead, cancellationToken);
-            response.EnsureSuccessStatusCode();
-            Stream blobStream = await response.Content.ReadAsStreamAsync(cancellationToken);
-            ReplicationLogSnapshot snapshot = await ReplicationLogSnapshot.DeserializeSnapshot(blobStream);
+            ReplicationLogSnapshot snapshot;
+            {
+                HttpRequestMessage request = BuildHttpRequest(HttpMethod.Get, $"api/v1/blobs/{blobNamespace}/{snapshotBlob}");
+                HttpResponseMessage response = await _httpClient.SendAsync(request, HttpCompletionOption.ResponseHeadersRead, cancellationToken);
+                response.EnsureSuccessStatusCode();
+                await using Stream blobStream = await response.Content.ReadAsStreamAsync(cancellationToken);
+                using FilesystemBufferedPayload payload = await FilesystemBufferedPayload.Create(blobStream);
+                snapshot = await ReplicationLogSnapshot.DeserializeSnapshot(payload.GetStream());
+            }
 
             if (!snapshot.LastEvent.HasValue)
                 throw new Exception("No last event found");
@@ -437,6 +441,12 @@ namespace Horde.Storage.Implementation
             HttpRequestMessage referencesRequest = BuildHttpRequest(HttpMethod.Get, $"api/v1/objects/{ns}/{blob}/references");
             HttpResponseMessage referencesResponse = await _httpClient.SendAsync(referencesRequest, cancellationToken);
             string body = await referencesResponse.Content.ReadAsStringAsync(cancellationToken);
+
+            if (referencesResponse.StatusCode == HttpStatusCode.BadRequest)
+            {
+                _logger.Warning("Failed to resolve references for object {Blob} in {Namespace}. Skipping replication", blob, ns);
+                return;
+            }
             referencesResponse.EnsureSuccessStatusCode();
 
             ResolvedReferencesResult? refs = JsonConvert.DeserializeObject<ResolvedReferencesResult>(body);
@@ -468,7 +478,9 @@ namespace Horde.Storage.Implementation
                         _logger.Error("Bad http response when replicating {Blob} in {Namespace} . Status code: {StatusCode}", blobToReplicate, ns, blobResponse.StatusCode);
                         return;
                     }
-                    using FilesystemBufferedPayload payload = await FilesystemBufferedPayload.Create(await blobResponse.Content.ReadAsStreamAsync(cancellationToken));
+
+                    await using Stream s = await blobResponse.Content.ReadAsStreamAsync(cancellationToken);
+                    using FilesystemBufferedPayload payload = await FilesystemBufferedPayload.Create(s);
                     await _blobService.PutObject(ns, payload, blobToReplicate);
                 });
             }
diff --git a/Engine/Source/Programs/Horde/HordeStorage/Horde.Storage/Implementation/Replication/Replicator.cs b/Engine/Source/Programs/Horde/HordeStorage/Horde.Storage/Implementation/Replication/Replicator.cs
--- a/Engine/Source/Programs/Horde/HordeStorage/Horde.Storage/Implementation/Replication/Replicator.cs
+++ b/Engine/Source/Programs/Horde/HordeStorage/Horde.Storage/Implementation/Replication/Replicator.cs
@@ -384,7 +384,8 @@ namespace Horde.Storage.Implementation
 
                         if (response.IsSuccessStatusCode)
                         {
-                            using FilesystemBufferedPayload payload = await FilesystemBufferedPayload.Create(await response.Content.ReadAsStreamAsync(replicationToken));
+                            await using Stream responseStream = await response.Content.ReadAsStreamAsync(replicationToken);
+                            using FilesystemBufferedPayload payload = await FilesystemBufferedPayload.Create(responseStream);
 
                             {
                                 await using Stream s = payload.GetStream();
