Horde: Remove accidentally submitted debug code.

diff --git a/Engine/Source/Programs/Horde/Horde.Agent/Commands/Archive/ExtractArchiveCommand.cs b/Engine/Source/Programs/Horde/Horde.Agent/Commands/Archive/ExtractArchiveCommand.cs
--- a/Engine/Source/Programs/Horde/Horde.Agent/Commands/Archive/ExtractArchiveCommand.cs
+++ b/Engine/Source/Programs/Horde/Horde.Agent/Commands/Archive/ExtractArchiveCommand.cs
@@ -52,28 +52,6 @@ namespace HordeAgent.Commands.Archive
 		/// <inheritdoc/>
 		public override async Task<int> ExecuteAsync(ILogger Logger)
 		{
-			int TargetSize = 64 * 1024;
-			int MaxValue = (int)((1L << 32) / TargetSize);
-
-			List<int> Sizes = new List<int>();
-			byte[] Data = File.ReadAllBytes(@"E:\P4 UE5\Engine\Source\ThirdParty\CEF3\cef_binary_3.2623.1395.g3034273_linux64\Release\libcef.so");
-			BuzHash Hash = new BuzHash();
-
-			int LastStart = 0;
-			for (int Idx = 0; Idx < Data.Length; Idx++)
-			{
-				Hash.add(Data[Idx]);
-				if (Hash.get() < MaxValue)
-				{
-					Sizes.Add((Idx + 1) - LastStart);
-					Hash.reset();
-					LastStart = Idx + 1;
-				}
-			}
-
-			double Average = Sizes.Average();
-
-
 			TreePackOptions Options = new TreePackOptions();
 			Options.MaxBlobSize = 64 * 1024;
 			Options.MaxInlineBlobSize = 20 * 1024;
