Horde: Add GPU vendor and driver info to primary device capabilities
#rnx
#skpci

diff --git a/Engine/Source/Programs/Horde/Horde.Agent/Services/WorkerService.cs b/Engine/Source/Programs/Horde/Horde.Agent/Services/WorkerService.cs
--- a/Engine/Source/Programs/Horde/Horde.Agent/Services/WorkerService.cs
+++ b/Engine/Source/Programs/Horde/Horde.Agent/Services/WorkerService.cs
@@ -1607,12 +1607,9 @@ namespace HordeAgent.Services
 							WmiProperties Properties = new WmiProperties(Row);
 							if (Properties.TryGetValue("Name", out string? Name) && Properties.TryGetValue("DriverVersion", out string? DriverVersion))
 							{
-								DeviceCapabilities Device = new DeviceCapabilities();
-								Device.Handle = $"GPU-{++Index}";
-								Device.Properties.Add($"Name={Name}");
-								Device.Properties.Add($"Type=GPU");
-								Device.Properties.Add($"DriverVersion={DriverVersion}");
-								OtherDevices.Add(Device);
+								string Prefix = $"GPU-{++Index}";
+								PrimaryDevice.Properties.Add($"{Prefix}-Name={Name}");
+								PrimaryDevice.Properties.Add($"{Prefix}-DriverVersion={DriverVersion}");								
 							}
 						}
 					}
