Horde: Fix agent log parser tests
#skipci
#rnx

diff --git a/Engine/Source/Programs/Horde/Horde.Agent.Tests/LogParserTests.cs b/Engine/Source/Programs/Horde/Horde.Agent.Tests/LogParserTests.cs
--- a/Engine/Source/Programs/Horde/Horde.Agent.Tests/LogParserTests.cs
+++ b/Engine/Source/Programs/Horde/Horde.Agent.Tests/LogParserTests.cs
@@ -972,7 +972,7 @@ namespace HordeAgentTests
 			CheckEventGroup(Events.Slice(0, 3), 0, 3, LogLevel.Information, KnownLogEvents.Systemic_Xge_Standalone);
 			CheckEventGroup(Events.Slice(3, 3), 3, 3, LogLevel.Information, KnownLogEvents.Systemic_Xge);
 			CheckEventGroup(Events.Slice(6, 1), 7, 1, LogLevel.Information, KnownLogEvents.Systemic_Xge_ServiceNotRunning);
-			CheckEventGroup(Events.Slice(7, 1), 8, 1, LogLevel.Information, KnownLogEvents.Systemic_Xge_BuildFailed);
+			CheckEventGroup(Events.Slice(7, 1), 8, 1, LogLevel.Error, KnownLogEvents.Systemic_Xge_BuildFailed);
 		}
 
 		string GetSubProperty(LogEvent Event, string SpanName, string Name)
