Horde Storage - Made sure we only snapshot a single namespace at a time, as we require quite a lot of temporary memory to build a snapshot (for certain large namespaces)
This should help reduce the risk of running out of memory on our worker nodes.

diff --git a/Engine/Source/Programs/Horde/HordeStorage/Horde.Storage/Implementation/GC/OrphanBlobCleanupRefs.cs b/Engine/Source/Programs/Horde/HordeStorage/Horde.Storage/Implementation/GC/OrphanBlobCleanupRefs.cs
--- a/Engine/Source/Programs/Horde/HordeStorage/Horde.Storage/Implementation/GC/OrphanBlobCleanupRefs.cs
+++ b/Engine/Source/Programs/Horde/HordeStorage/Horde.Storage/Implementation/GC/OrphanBlobCleanupRefs.cs
@@ -90,27 +90,29 @@ namespace Horde.Storage.Implementation
                     {
                         IBlobIndex.BlobInfo? blobIndex = await _blobIndex.GetBlobInfo(blobNamespace, blob);
 
-                        if (blobIndex != null)
+                        if (blobIndex == null)
                         {
-                            foreach ((BucketId bucket, IoHashKey key) in blobIndex.References)
+                            continue;
+                        }
+
+                        foreach ((BucketId bucket, IoHashKey key) in blobIndex.References)
+                        {
+                            if (cancellationToken.IsCancellationRequested)
+                                break;
+
+                            try
+                            {
+                                (ObjectRecord, BlobContents?) _ = await _objectService.Get(blobNamespace, bucket, key, new string[] {"name"});
+                                found = true;
+                            }
+                            catch (ObjectNotFoundException)
+                            {
+                                // this is not a valid reference so we should delete
+                            }
+                            catch (MissingBlobsException)
                             {
-                                if (cancellationToken.IsCancellationRequested)
-                                    break;
-
-                                try
-                                {
-                                    (ObjectRecord, BlobContents?) _ = await _objectService.Get(blobNamespace, bucket, key, new string[] {"name"});
-                                    found = true;
-                                }
-                                catch (ObjectNotFoundException)
-                                {
-                                    // this is not a valid reference so we should delete
-                                }
-                                catch (MissingBlobsException)
-                                {
-                                    // we do not care if there are missing blobs, as long as the record is valid we keep this blob around
-                                    found = true;
-                                }
+                                // we do not care if there are missing blobs, as long as the record is valid we keep this blob around
+                                found = true;
                             }
                         }
                     }
diff --git a/Engine/Source/Programs/Horde/HordeStorage/Horde.Storage/Implementation/TransactionLog/ReplicationSnapshotService.cs b/Engine/Source/Programs/Horde/HordeStorage/Horde.Storage/Implementation/TransactionLog/ReplicationSnapshotService.cs
--- a/Engine/Source/Programs/Horde/HordeStorage/Horde.Storage/Implementation/TransactionLog/ReplicationSnapshotService.cs
+++ b/Engine/Source/Programs/Horde/HordeStorage/Horde.Storage/Implementation/TransactionLog/ReplicationSnapshotService.cs
@@ -76,6 +76,7 @@ namespace Horde.Storage.Implementation
                 ReplicationLogSnapshotBuilder builder = ActivatorUtilities.CreateInstance<ReplicationLogSnapshotBuilder>(_provider);
                 try
                 {
+                    _logger.Information("Building snapshot for {Namespace}", ns);
                     BlobIdentifier snapshotBlob = await builder.BuildSnapshot(ns, _settings.CurrentValue.SnapshotStorageNamespace, _cancellationTokenSource.Token);
                     _logger.Information("Snapshot built for {Namespace} with id {Id}", ns, snapshotBlob);
 
@@ -109,6 +110,6 @@ namespace Horde.Storage.Implementation
         public bool Enabled { get; set; } = false;
 
         public NamespaceId SnapshotStorageNamespace { get; set; } = INamespacePolicyResolver.JupiterInternalNamespace;
-        public int MaxCountOfNamespacesToSnapshotInParallel { get; set; } = 2;
+        public int MaxCountOfNamespacesToSnapshotInParallel { get; set; } = 1;
     }
 }
