Horde Storage - Added option to skip snapshots during replication, this can be used to get the incremental replication going quicker (which is the newest set of data). It will result in a partial replication (as old data is not restored) but that can be recovered after the fact by setting up another replicator of the same namespace (but with a different name). Allows for more flexiblity in terms of which content is required.

Setting this by default as the snapshot replicaton currently requires quite a lot of memory causing issues for our replication.

diff --git a/Engine/Source/Programs/Horde/HordeStorage/Horde.Storage/Implementation/Replication/RefsReplicator.cs b/Engine/Source/Programs/Horde/HordeStorage/Horde.Storage/Implementation/Replication/RefsReplicator.cs
--- a/Engine/Source/Programs/Horde/HordeStorage/Horde.Storage/Implementation/Replication/RefsReplicator.cs
+++ b/Engine/Source/Programs/Horde/HordeStorage/Horde.Storage/Implementation/Replication/RefsReplicator.cs
@@ -154,7 +154,7 @@ namespace Horde.Storage.Implementation
                 Guid? lastEvent = null;
 
                 bool haveRunBefore = _refsState.LastBucket != null && _refsState.LastEvent != null;
-                if (!haveRunBefore)
+                if (!haveRunBefore && !_replicatorSettings.SkipSnapshot)
                 {
                     // have not run before, replicate a snapshot
                     _logger.Information("{Name} Have not run replication before, attempting to use snapshot. State: {@State}", _name, _refsState);
@@ -210,6 +210,16 @@ namespace Horde.Storage.Implementation
 
                     if (useSnapshotException != null)
                     {
+                        // if we are told to use a snapshot but are configured to not do so we will simply reset the state tracking and thus start from the oldest incremental event
+                        // this will result in a partial replication
+                        if (_replicatorSettings.SkipSnapshot)
+                        {
+                            lastBucket = null;
+                            lastEvent = null;
+                            retry = true;
+                            continue;
+                        }
+
                         // if we have already attempted to recover using a snapshot and we still fail we just give up and throw the exception onwards.
                         if (haveAttemptedSnapshot)
                             throw new AggregateException(useSnapshotException);
diff --git a/Engine/Source/Programs/Horde/HordeStorage/Horde.Storage/Implementation/Replication/ReplicationSettings.cs b/Engine/Source/Programs/Horde/HordeStorage/Horde.Storage/Implementation/Replication/ReplicationSettings.cs
--- a/Engine/Source/Programs/Horde/HordeStorage/Horde.Storage/Implementation/Replication/ReplicationSettings.cs
+++ b/Engine/Source/Programs/Horde/HordeStorage/Horde.Storage/Implementation/Replication/ReplicationSettings.cs
@@ -81,6 +81,11 @@ namespace Horde.Storage
         /// Max number of offsets to skip when finding broken events in the transaction logs, higher value means we scan further increasing the likelihood of finding a valid entry, but also means risk that we skip records you would want.
         /// </summary>
         public int? MaxOffsetsAttempted { get; set; } = 100000;
+
+        /// <summary>
+        /// Do not read in old history using a snapshot, instead just start replicating the new incremental state. Means we will only have partial state but we will also get going quicker on relatively useful blobs.
+        /// </summary>
+        public bool SkipSnapshot { get; set; } = true;
     }
 
     public enum ReplicatorVersion
