Horde Storage - Defaulted skip snapshot to false (e.g. the old behavior) - this fixes issues in the tests.
Then added ability to specify this in the helm chart instead, and defaulting that to true (as we want to get the replication running without the snapshots for now while I rework how that code behaves)

diff --git a/Engine/Source/Programs/Horde/HordeStorage/Horde.Storage/Implementation/Replication/ReplicationSettings.cs b/Engine/Source/Programs/Horde/HordeStorage/Horde.Storage/Implementation/Replication/ReplicationSettings.cs
--- a/Engine/Source/Programs/Horde/HordeStorage/Horde.Storage/Implementation/Replication/ReplicationSettings.cs
+++ b/Engine/Source/Programs/Horde/HordeStorage/Horde.Storage/Implementation/Replication/ReplicationSettings.cs
@@ -85,7 +85,7 @@ namespace Horde.Storage
         /// <summary>
         /// Do not read in old history using a snapshot, instead just start replicating the new incremental state. Means we will only have partial state but we will also get going quicker on relatively useful blobs.
         /// </summary>
-        public bool SkipSnapshot { get; set; } = true;
+        public bool SkipSnapshot { get; set; } = false;
     }
 
     public enum ReplicatorVersion
