Horde Storage - Let kestrel decide protocol on the internal port (will be http1 with http2 upgrade as that is the default instead of forcing it to http2)

diff --git a/Engine/Source/Programs/Horde/HordeStorage/Horde.Storage/appsettings.Production.json b/Engine/Source/Programs/Horde/HordeStorage/Horde.Storage/appsettings.Production.json
--- a/Engine/Source/Programs/Horde/HordeStorage/Horde.Storage/appsettings.Production.json
+++ b/Engine/Source/Programs/Horde/HordeStorage/Horde.Storage/appsettings.Production.json
@@ -42,9 +42,8 @@
       "Http": {
         "Url": "http://0.0.0.0:80"
       },
-      "InternalHttp2": {
-        "Url": "http://0.0.0.0:8080",
-        "Protocols": "Http2"
+      "InternalHttp": {
+        "Url": "http://0.0.0.0:8080"
       }
     }
   }
