Horde: Suppot partial job name searches
#skipci
#rnx

diff --git a/Engine/Source/Programs/Horde/Horde.Build/Jobs/Impl/JobCollection.cs b/Engine/Source/Programs/Horde/Horde.Build/Jobs/Impl/JobCollection.cs
--- a/Engine/Source/Programs/Horde/Horde.Build/Jobs/Impl/JobCollection.cs
+++ b/Engine/Source/Programs/Horde/Horde.Build/Jobs/Impl/JobCollection.cs
@@ -531,7 +531,15 @@ namespace HordeServer.Collections.Impl
 			}
 			if (Name != null)
 			{
-				Filter &= FilterBuilder.Eq(x => x.Name, Name);
+				if (Name.StartsWith("$", StringComparison.InvariantCulture))
+				{
+					BsonRegularExpression Regex = new BsonRegularExpression(Name.Substring(1), "i");
+					Filter &= FilterBuilder.Regex(x => x.Name, Regex);
+				}
+				else
+				{
+					Filter &= FilterBuilder.Eq(x => x.Name, Name);
+				}				
 			}
 			if (Templates != null)
 			{
