Horde: Use batch state "ready" as signal in job queue scaling strategy

The previous use of "waiting" state was incorrect as that state represents a batch waiting for another batch. Not waiting for an agent to become available.

diff --git a/Engine/Source/Programs/Horde/Horde.Build.Tests/Fleet/JobQueueStrategyTest.cs b/Engine/Source/Programs/Horde/Horde.Build.Tests/Fleet/JobQueueStrategyTest.cs
--- a/Engine/Source/Programs/Horde/Horde.Build.Tests/Fleet/JobQueueStrategyTest.cs
+++ b/Engine/Source/Programs/Horde/Horde.Build.Tests/Fleet/JobQueueStrategyTest.cs
@@ -24,6 +24,7 @@ namespace Horde.Build.Tests.Fleet
 		public async Task GetPoolQueueSizes()
 		{
 			(JobQueueStrategy Strategy, PoolSizeData PoolSizeData) = await SetUpJobsAsync(1, 5);
+			await Clock.AdvanceAsync(Strategy.ReadyTimeThreshold + TimeSpan.FromSeconds(5));
 			Dictionary<PoolId, int> PoolQueueSizes = await Strategy.GetPoolQueueSizesAsync(Clock.UtcNow - TimeSpan.FromHours(2));
 			Assert.AreEqual(1, PoolQueueSizes.Count);
 			Assert.AreEqual(5, PoolQueueSizes[PoolSizeData.Pool.Id]);
@@ -32,7 +33,7 @@ namespace Horde.Build.Tests.Fleet
 		[TestMethod]
 		public async Task EmptyJobQueue()
 		{
-			await AssertAgentCount(0, -1);
+			await AssertAgentCount(0, -1, false);
 		}
 		
 		[TestMethod]
@@ -41,10 +42,16 @@ namespace Horde.Build.Tests.Fleet
 			await AssertAgentCount(2, 0);
 		}
 		
+		[TestMethod]
+		public async Task BatchesNotWaitingLongEnough()
+		{
+			await AssertAgentCount(3, -1, false);
+		}
+		
 		[TestMethod]
 		public async Task NumQueuedJobs3()
 		{
-			await AssertAgentCount(3, 1);
+			await AssertAgentCount(3, 1, true);
 		}
 		
 		[TestMethod]
@@ -59,9 +66,15 @@ namespace Horde.Build.Tests.Fleet
 			await AssertAgentCount(25, 6);
 		}
 
-		public async Task AssertAgentCount(int NumBatchesWaiting, int ExpectedAgentDelta)
+		public async Task AssertAgentCount(int NumBatchesReady, int ExpectedAgentDelta, bool WaitedBeyondThreshold = true)
 		{
-			(JobQueueStrategy Strategy, PoolSizeData PoolSizeData) = await SetUpJobsAsync(1, NumBatchesWaiting);
+			(JobQueueStrategy Strategy, PoolSizeData PoolSizeData) = await SetUpJobsAsync(1, NumBatchesReady);
+			TimeSpan TimeToWait = WaitedBeyondThreshold
+				? Strategy.ReadyTimeThreshold + TimeSpan.FromSeconds(5)
+				: TimeSpan.FromSeconds(15);
+			
+			await Clock.AdvanceAsync(TimeToWait);
+
 			List<PoolSizeData> Result = await Strategy.CalcDesiredPoolSizesAsync(new() { PoolSizeData });
 			Assert.AreEqual(1, Result.Count);
 			Assert.AreEqual(Result[0].Agents.Count + ExpectedAgentDelta, Result[0].DesiredAgentCount);
@@ -71,8 +84,8 @@ namespace Horde.Build.Tests.Fleet
 		/// Set up a fixture for job queue tests, ensuring a certain number of job batches are in running or waiting state
 		/// </summary>
 		/// <param name="NumBatchesRunning">Num of job batches that should be in state running</param>
-		/// <param name="NumBatchesWaiting">Num of job batches that should be in state waiting</param>
-		private async Task<(JobQueueStrategy, PoolSizeData)> SetUpJobsAsync(int NumBatchesRunning, int NumBatchesWaiting)
+		/// <param name="NumBatchesReady">Num of job batches that should be in state waiting</param>
+		private async Task<(JobQueueStrategy, PoolSizeData)> SetUpJobsAsync(int NumBatchesRunning, int NumBatchesReady)
 		{
 			IPool Pool1 = await PoolService.CreatePoolAsync("bogusPool1", null, true, 0, 0);
 			List<IAgent> Agents = new();
@@ -109,10 +122,10 @@ namespace Horde.Build.Tests.Fleet
 				await JobCollection.TryUpdateBatchAsync(Job, Graph, Job.Batches[0].Id, null, JobStepBatchState.Running, null);
 			}
 			
-			for (int i = 0; i < NumBatchesWaiting; i++)
+			for (int i = 0; i < NumBatchesReady; i++)
 			{
 				IJob Job = await AddPlaceholderJob(Graph, Stream.Id, NodeForAgentType1);
-				await JobCollection.TryUpdateBatchAsync(Job, Graph, Job.Batches[0].Id, null, JobStepBatchState.Waiting, null);
+				await JobCollection.TryUpdateBatchAsync(Job, Graph, Job.Batches[0].Id, null, JobStepBatchState.Ready, null);
 			}
 			
 			return (new (JobCollection, GraphCollection, StreamService, Clock), PoolSize);
diff --git a/Engine/Source/Programs/Horde/Horde.Build/Fleet/Autoscale/JobQueueStrategy.cs b/Engine/Source/Programs/Horde/Horde.Build/Fleet/Autoscale/JobQueueStrategy.cs
--- a/Engine/Source/Programs/Horde/Horde.Build/Fleet/Autoscale/JobQueueStrategy.cs
+++ b/Engine/Source/Programs/Horde/Horde.Build/Fleet/Autoscale/JobQueueStrategy.cs
@@ -35,6 +35,14 @@ namespace Horde.Build.Fleet.Autoscale
 		/// </summary>
 		private readonly TimeSpan SamplePeriod = TimeSpan.FromHours(2.0);
 		
+		/// <summary>
+		/// Time spent in ready state before considered truly waiting for an agent
+		///
+		/// A job batch can be in ready state before getting picked up and executed.
+		/// This threshold will help ensure only batches that have been waiting longer than this value will be considered.
+		/// </summary>
+		internal readonly TimeSpan ReadyTimeThreshold = TimeSpan.FromSeconds(45.0);
+		
 		/// <summary>
 		/// Minimum number of jobs in queue for scale-out logic to activate
 		///
@@ -89,7 +97,11 @@ namespace Horde.Build.Fleet.Autoscale
 			List<(IJob Job, IJobStepBatch Batch, PoolId PoolId)> JobBatches = new();
 			foreach (IJobStepBatch Batch in Job.Batches)
 			{
-				if (Batch.State != JobStepBatchState.Waiting) continue;
+				if (Batch.State != JobStepBatchState.Ready) continue;
+				
+				TimeSpan? WaitTime = Clock.UtcNow - Batch.ReadyTimeUtc;
+				if (WaitTime == null) continue;
+				if (WaitTime.Value < ReadyTimeThreshold) continue;
 
 				if (!Streams.TryGetValue(Job.StreamId, out IStream? Stream))
 				{
diff --git a/Engine/Source/Programs/Horde/Horde.Build/Jobs/Impl/JobCollection.cs b/Engine/Source/Programs/Horde/Horde.Build/Jobs/Impl/JobCollection.cs
--- a/Engine/Source/Programs/Horde/Horde.Build/Jobs/Impl/JobCollection.cs
+++ b/Engine/Source/Programs/Horde/Horde.Build/Jobs/Impl/JobCollection.cs
@@ -365,6 +365,11 @@ namespace HordeServer.Collections.Impl
 		/// The user collection
 		/// </summary>
 		IUserCollection UserCollection;
+		
+		/// <summary>
+		/// Clock representing wall-clock time
+		/// </summary>
+		IClock Clock;
 
 		/// <summary>
 		/// Logger for output
@@ -381,10 +386,12 @@ namespace HordeServer.Collections.Impl
 		/// </summary>
 		/// <param name="DatabaseService">The database service singleton</param>
 		/// <param name="UserCollection"></param>
+		/// <param name="Clock"></param>
 		/// <param name="Logger">The logger instance</param>
-		public JobCollection(DatabaseService DatabaseService, IUserCollection UserCollection, ILogger<JobCollection> Logger)
+		public JobCollection(DatabaseService DatabaseService, IUserCollection UserCollection, IClock Clock, ILogger<JobCollection> Logger)
 		{
 			this.UserCollection = UserCollection;
+			this.Clock = Clock;
 			this.Logger = Logger;
 
 			Jobs = DatabaseService.GetCollection<JobDocument>("Jobs");
@@ -772,12 +779,12 @@ namespace HordeServer.Collections.Impl
 
 				if (Batch.StartTime == null && NewState >= JobStepBatchState.Starting)
 				{
-					Batch.StartTime = DateTimeOffset.Now;
+					Batch.StartTime = Clock.UtcNow;
 					Updates.Add(UpdateBuilder.Set(x => x.Batches[BatchIdx].StartTime, Batch.StartTime));
 				}
 				if (NewState == JobStepBatchState.Complete)
 				{
-					Batch.FinishTime = DateTimeOffset.Now;
+					Batch.FinishTime = Clock.UtcNow;
 					Updates.Add(UpdateBuilder.Set(x => x.Batches[BatchIdx].FinishTime, Batch.FinishTime));
 				}
 			}
@@ -858,12 +865,12 @@ namespace HordeServer.Collections.Impl
 
 								if (Step.State == JobStepState.Running)
 								{
-									Step.StartTime = DateTimeOffset.Now;
+									Step.StartTime = Clock.UtcNow;
 									Updates.Add(UpdateBuilder.Set(x => x.Batches[BatchIdx].Steps[StepIdx].StartTime, Step.StartTime));
 								}
 								else if (Step.State == JobStepState.Completed || Step.State == JobStepState.Aborted)
 								{
-									Step.FinishTime = DateTimeOffset.Now;
+									Step.FinishTime = Clock.UtcNow;
 									Updates.Add(UpdateBuilder.Set(x => x.Batches[BatchIdx].Steps[StepIdx].FinishTime, Step.FinishTime));
 								}
 
@@ -1268,7 +1275,7 @@ namespace HordeServer.Collections.Impl
 					Step.Outcome = JobStepOutcome.Failure;
 					Updates.Add(UpdateBuilder.Set(x => x.Batches[BatchIdx].Steps[StepIdxCopy].Outcome, Step.Outcome));
 
-					Step.FinishTime = DateTimeOffset.Now;
+					Step.FinishTime = Clock.UtcNow;
 					Updates.Add(UpdateBuilder.Set(x => x.Batches[BatchIdx].Steps[StepIdxCopy].FinishTime, Step.FinishTime));
 				}
 				else if (Step.State == JobStepState.Ready)
