Horde: Refactor commit metadata replication to allow running without mutating any server state. Add a standalone command for updating ranges of commits.

diff --git a/Engine/Source/Programs/Horde/Horde.Build/Commands/Perforce/ReplicateCommand.cs b/Engine/Source/Programs/Horde/Horde.Build/Commands/Perforce/ReplicateCommand.cs
new file mode 100644
--- /dev/null
+++ b/Engine/Source/Programs/Horde/Horde.Build/Commands/Perforce/ReplicateCommand.cs
@@ -0,0 +1,70 @@
+// Copyright Epic Games, Inc. All Rights Reserved.
+
+using EpicGames.Core;
+using HordeServer.Collections;
+using HordeServer.Commits;
+using HordeServer.Commits.Impl;
+using HordeServer.Controllers;
+using HordeServer.Models;
+using HordeServer.Utilities;
+using Microsoft.Extensions.Configuration;
+using Microsoft.Extensions.DependencyInjection;
+using Microsoft.Extensions.Logging;
+using System;
+using System.Collections.Generic;
+using System.Linq;
+using System.Threading.Tasks;
+
+namespace HordeServer.Commands
+{
+	using StreamId = StringId<IStream>;
+
+	[Command("perforce", "replicate", "Replicates commits for a particular change of changes from Perforce")]
+	class ReplicateCommand : Command
+	{
+		[CommandLine("-StreamId=", Required = true)]
+		public List<string> StreamIds { get; set; } = new List<string>();
+
+		[CommandLine(Required = true)]
+		public int Change { get; set; }
+
+		[CommandLine]
+		public int Count { get; set; } = 1;
+
+		IConfiguration Configuration;
+
+		public ReplicateCommand(IConfiguration Configuration)
+		{
+			this.Configuration = Configuration;
+		}
+
+		public override async Task<int> ExecuteAsync(ILogger Logger)
+		{
+			IServiceProvider ServiceProvider = Startup.CreateServiceProvider(Configuration);
+			CommitService CommitService = ServiceProvider.GetRequiredService<CommitService>();
+
+			IStreamCollection StreamCollection = ServiceProvider.GetRequiredService<IStreamCollection>();
+
+			Dictionary<IStream, int> StreamToFirstChange = new Dictionary<IStream, int>();
+			foreach(string StreamId in StreamIds)
+			{
+				IStream? Stream = await StreamCollection.GetAsync(new StreamId(StreamId));
+				if (Stream == null)
+				{
+					throw new FatalErrorException($"Stream '{Stream}' not found");
+				}
+				StreamToFirstChange[Stream] = Change;
+			}
+
+			string ClusterName = StreamToFirstChange.First().Key.ClusterName;
+			await foreach (NewCommit Commit in CommitService.FindCommitsForClusterAsync(ClusterName, StreamToFirstChange).Take(Count))
+			{
+				string BriefSummary = Commit.Description.Replace('\n', ' ').Substring(0, 50);
+				Logger.LogInformation("Commit {Change} by {AuthorId}: {Summary}", Commit.Change, Commit.AuthorId, BriefSummary);
+				Logger.LogInformation(" - Base path: {BasePath}", Commit.BasePath);
+			}
+
+			return 0;
+		}
+	}
+}
diff --git a/Engine/Source/Programs/Horde/Horde.Build/Perforce/Impl/CommitService.cs b/Engine/Source/Programs/Horde/Horde.Build/Perforce/Impl/CommitService.cs
--- a/Engine/Source/Programs/Horde/Horde.Build/Perforce/Impl/CommitService.cs
+++ b/Engine/Source/Programs/Horde/Horde.Build/Perforce/Impl/CommitService.cs
@@ -148,9 +148,45 @@ namespace HordeServer.Commits.Impl
 		/// <param name="Streams"></param>
 		/// <returns></returns>
 		async Task UpdateCommitsForClusterAsync(string ClusterName, IEnumerable<IStream> Streams)
+		{
+			// Find the minimum changelist number to query
+			Dictionary<IStream, int> StreamToFirstChange = new Dictionary<IStream, int>();
+			foreach (IStream Stream in Streams)
+			{
+				RedisList<int> Changes = RedisStreamChanges(Stream.Id);
+
+				int FirstChange = await Changes.GetByIndexAsync(-1);
+				if (FirstChange != 0)
+				{
+					FirstChange++;
+				}
+
+				StreamToFirstChange[Stream] = FirstChange;
+			}
+
+			// Update the database with all the commits
+			await foreach (NewCommit NewCommit in FindCommitsForClusterAsync(ClusterName, StreamToFirstChange))
+			{
+				ICommit Commit = await CommitCollection.AddOrReplaceAsync(NewCommit);
+
+				RedisList<int> StreamCommitsKey = RedisStreamChanges(Commit.StreamId);
+				await StreamCommitsKey.RightPushAsync(Commit.Change);
+
+				await RedisDirtyStreams.AddAsync(Commit.StreamId);
+				await Redis.PublishAsync(RedisUpdateChannel, Commit.StreamId);
+			}
+		}
+
+		/// <summary>
+		/// Enumerate new commits to the given streams, using a stream view to deduplicate changes which affect multiple branches.
+		/// </summary>
+		/// <param name="ClusterName"></param>
+		/// <param name="StreamToFirstChange"></param>
+		/// <returns>List of new commits</returns>
+		public async IAsyncEnumerable<NewCommit> FindCommitsForClusterAsync(string ClusterName, Dictionary<IStream, int> StreamToFirstChange)
 		{
 			// Create a connection to the server
-			IPerforceConnection? Connection = await PerforceService.GetServiceUserConnection(ClusterName);
+			using IPerforceConnection? Connection = await PerforceService.GetServiceUserConnection(ClusterName);
 			if (Connection == null)
 			{
 				throw new PerforceException($"Unable to create cluster connection for {ClusterName}");
@@ -161,7 +197,7 @@ namespace HordeServer.Commits.Impl
 
 			// Get the view for each stream
 			List<StreamInfo> StreamInfoList = new List<StreamInfo>();
-			foreach (IStream Stream in Streams)
+			foreach (IStream Stream in StreamToFirstChange.Keys)
 			{
 				ViewMap View = await GetStreamViewAsync(Connection, Stream.Name);
 				StreamInfoList.Add(new StreamInfo(Stream, View));
@@ -198,18 +234,11 @@ namespace HordeServer.Commits.Impl
 			int MinChange = int.MaxValue;
 			foreach (StreamInfo StreamInfo in StreamInfoList)
 			{
-				RedisList<int> Changes = RedisStreamChanges(StreamInfo.Stream.Id);
-
-				int FirstChange = await Changes.GetByIndexAsync(-1);
+				int FirstChange = StreamToFirstChange[StreamInfo.Stream];
 				if (FirstChange == 0)
 				{
 					FirstChange = await GetFirstCommitToReplicateAsync(Connection, StreamInfo.View, ServerInfo.Utf8PathComparer);
 				}
-				else
-				{
-					FirstChange++;
-				}
-
 				if (FirstChange != 0)
 				{
 					MinChange = Math.Min(MinChange, FirstChange);
@@ -236,13 +265,12 @@ namespace HordeServer.Commits.Impl
 					Utf8String BasePath = GetBasePath(DescribeRecord, StreamInfo.View, ServerInfo.Utf8PathComparer);
 					if (!BasePath.IsEmpty)
 					{
-						await AddCommitAsync(Stream, DescribeRecord, BasePath.ToString());
+						IUser Author = await UserCollection.FindOrAddUserByLoginAsync(DescribeRecord.User);
+						IUser Owner = (await ParseRobomergeOwnerAsync(DescribeRecord.Description)) ?? Author;
 
-						RedisList<int> StreamCommitsKey = RedisStreamChanges(Stream.Id);
-						await StreamCommitsKey.RightPushAsync(DescribeRecord.Number);
+						int OriginalChange = ParseRobomergeSource(DescribeRecord.Description) ?? DescribeRecord.Number;
 
-						await RedisDirtyStreams.AddAsync(Stream.Id);
-						await Redis.PublishAsync(RedisUpdateChannel, Stream.Id);
+						yield return new NewCommit(Stream.Id, DescribeRecord.Number, OriginalChange, Author.Id, Owner.Id, DescribeRecord.Description, BasePath.ToString(), DescribeRecord.Time);
 					}
 				}
 			}
@@ -313,24 +341,6 @@ namespace HordeServer.Commits.Impl
 			return A.Substring(0, Index);
 		}
 
-		/// <summary>
-		/// Adds metadata for a particular commit
-		/// </summary>
-		/// <param name="Stream"></param>
-		/// <param name="Changelist"></param>
-		/// <param name="BasePath"></param>
-		/// <returns></returns>
-		async Task<ICommit> AddCommitAsync(IStream Stream, DescribeRecord Changelist, string BasePath)
-		{
-			IUser Author = await UserCollection.FindOrAddUserByLoginAsync(Changelist.User);
-			IUser Owner = (await ParseRobomergeOwnerAsync(Changelist.Description)) ?? Author;
-
-			int OriginalChange = ParseRobomergeSource(Changelist.Description) ?? Changelist.Number;
-
-			NewCommit NewCommit = new NewCommit(Stream.Id, Changelist.Number, OriginalChange, Author.Id, Owner.Id, Changelist.Description, BasePath, Changelist.Time);
-			return await CommitCollection.AddOrReplaceAsync(NewCommit);
-		}
-
 		/// <inheritdoc/>
 		public async Task<List<ICommit>> FindCommitsAsync(StreamId StreamId, int? MinChange = null, int? MaxChange = null, string[]? Paths = null, int? Index = null, int? Count = null)
 		{
