Horde: Add a test for Swagger generator

Ensures REST API code/docs can be generated for Swagger. Without this, problem will only manifest itself when calling the Swagger endpoints (via a browser usually).

diff --git a/Engine/Source/Programs/Horde/Horde.Build.Tests/SwaggerTest.cs b/Engine/Source/Programs/Horde/Horde.Build.Tests/SwaggerTest.cs
new file mode 100644
--- /dev/null
+++ b/Engine/Source/Programs/Horde/Horde.Build.Tests/SwaggerTest.cs
@@ -0,0 +1,29 @@
+﻿// Copyright Epic Games, Inc. All Rights Reserved.
+
+using System;
+using System.Text.Json;
+using System.Threading.Tasks;
+using HordeServerTests;
+using Microsoft.VisualStudio.TestTools.UnitTesting;
+
+namespace Horde.Build.Tests
+{
+    [TestClass]
+    public class SwaggerTest : ControllerIntegrationTest
+    {
+        [TestMethod]
+        [Ignore("Until clashes with Horde.Build and Horde.Storage are resolved")]
+        public async Task ValidateSwagger()
+        {
+	        var Res = await client.GetAsync("/swagger/v1/swagger.json");
+	        if (!Res.IsSuccessStatusCode)
+	        {
+		        string RawJson = await Res.Content.ReadAsStringAsync();
+		        JsonElement TempElement = JsonSerializer.Deserialize<JsonElement>(RawJson);
+		        string FormattedJson = JsonSerializer.Serialize(TempElement, new JsonSerializerOptions { WriteIndented = true });
+		        Console.Error.WriteLine("Error result:\n" + FormattedJson);
+		        Res.EnsureSuccessStatusCode();
+	        }
+        }
+    }
+}
\ No newline at end of file
