Horde Storage - Use the replication state from scylla if it is available

diff --git a/Engine/Source/Programs/Horde/HordeStorage/Horde.Storage/Implementation/Replication/RefsReplicator.cs b/Engine/Source/Programs/Horde/HordeStorage/Horde.Storage/Implementation/Replication/RefsReplicator.cs
--- a/Engine/Source/Programs/Horde/HordeStorage/Horde.Storage/Implementation/Replication/RefsReplicator.cs
+++ b/Engine/Source/Programs/Horde/HordeStorage/Horde.Storage/Implementation/Replication/RefsReplicator.cs
@@ -57,7 +57,7 @@ namespace Horde.Storage.Implementation
             DirectoryInfo stateRoot = new DirectoryInfo(replicationSettings.CurrentValue.GetStateRoot());
             _stateFile = new FileInfo(Path.Combine(stateRoot.FullName, stateFileName));
 
-            ReplicatorState? replicatorState = replicationLog.GetReplicatorState(replicatorSettings.NamespaceToReplicate, _name).Result;
+            ReplicatorState? replicatorState = _replicationLog.GetReplicatorState(_replicatorSettings.NamespaceToReplicate, _name).Result;
             if (replicatorState == null)
             {
                 if (_stateFile.Exists)
@@ -121,16 +121,29 @@ namespace Horde.Storage.Implementation
             }
 
             // read the state again to allow it to be modified by the admin controller / other instances of horde-storage connected to the same filesystem
-            _stateFile.Refresh();
-            if (_stateFile.Exists)
+            ReplicatorState? replicatorState = await _replicationLog.GetReplicatorState(_replicatorSettings.NamespaceToReplicate, _name);
+            if (replicatorState == null)
             {
-                _refsState = ReadState(_stateFile) ?? new RefsState();
+                _stateFile.Refresh();
+                if (_stateFile.Exists)
+                {
+                    _refsState = ReadState(_stateFile) ?? new RefsState();
+                }
+                else
+                {
+                    _refsState = new RefsState();
+                }
             }
             else
             {
-                _refsState = new RefsState();
+                _refsState = new RefsState()
+                {
+                    LastBucket = replicatorState.LastBucket,
+                    LastEvent = replicatorState.LastEvent,
+                };
             }
 
+
             LogReplicationHeartbeat(0);
 
             bool hasRun;
