Horde: Improve tracing and logging for AWS fleet manager

To help track down issues seen with shrinking of pools.

diff --git a/Engine/Source/Programs/Horde/Horde.Build/Fleet/Impl/Managers/Fleet/AwsFleetManager.cs b/Engine/Source/Programs/Horde/Horde.Build/Fleet/Impl/Managers/Fleet/AwsFleetManager.cs
--- a/Engine/Source/Programs/Horde/Horde.Build/Fleet/Impl/Managers/Fleet/AwsFleetManager.cs
+++ b/Engine/Source/Programs/Horde/Horde.Build/Fleet/Impl/Managers/Fleet/AwsFleetManager.cs
@@ -10,6 +10,9 @@ using System;
 using System.Collections.Generic;
 using System.Linq;
 using System.Threading.Tasks;
+using HordeServer.Utilities;
+using OpenTracing;
+using OpenTracing.Util;
 
 namespace HordeServer.Services.Impl
 {
@@ -50,25 +53,42 @@ namespace HordeServer.Services.Impl
 		/// <inheritdoc/>
 		public async Task ExpandPoolAsync(IPool Pool, IReadOnlyList<IAgent> Agents, int Count)
 		{
-			// Find stopped instances in the correct pool
-			DescribeInstancesRequest DescribeRequest = new DescribeInstancesRequest();
-			DescribeRequest.Filters = new List<Filter>();
-			DescribeRequest.Filters.Add(new Filter("instance-state-name", new List<string> { InstanceStateName.Stopped.Value }));
-			DescribeRequest.Filters.Add(new Filter("tag:" + PoolTagName, new List<string> { Pool.Name }));
-			DescribeInstancesResponse DescribeResponse = await Client.DescribeInstancesAsync(DescribeRequest);
+			using IScope Scope = GlobalTracer.Instance.BuildSpan("ExpandPool").StartActive();
+			Scope.Span.SetTag("poolName", Pool.Name);
+			Scope.Span.SetTag("numAgents", Agents.Count);
+			Scope.Span.SetTag("count", Count);
 
-			// Try to start the given instances
-			StartInstancesRequest StartRequest = new StartInstancesRequest();
-			StartRequest.InstanceIds.AddRange(DescribeResponse.Reservations.SelectMany(x => x.Instances).Select(x => x.InstanceId).Take(Count));
+			DescribeInstancesResponse DescribeResponse;
+			using (IScope DescribeScope = GlobalTracer.Instance.BuildSpan("DescribeInstances").StartActive())
+			{
+				// Find stopped instances in the correct pool
+				DescribeInstancesRequest DescribeRequest = new DescribeInstancesRequest();
+				DescribeRequest.Filters = new List<Filter>();
+				DescribeRequest.Filters.Add(new Filter("instance-state-name", new List<string> { InstanceStateName.Stopped.Value }));
+				DescribeRequest.Filters.Add(new Filter("tag:" + PoolTagName, new List<string> { Pool.Name }));
+				DescribeResponse = await Client.DescribeInstancesAsync(DescribeRequest);
+				DescribeScope.Span.SetTag("res.statusCode", (int)DescribeResponse.HttpStatusCode);
+				DescribeScope.Span.SetTag("res.numReservations", DescribeResponse.Reservations.Count);
+			}
 
-			if (StartRequest.InstanceIds.Count > 0)
+			using (IScope StartScope = GlobalTracer.Instance.BuildSpan("StartInstances").StartActive())
 			{
-				StartInstancesResponse StartResponse = await Client.StartInstancesAsync(StartRequest);
-				if ((int)StartResponse.HttpStatusCode >= 200 && (int)StartResponse.HttpStatusCode <= 299)
+				// Try to start the given instances
+				StartInstancesRequest StartRequest = new StartInstancesRequest();
+				StartRequest.InstanceIds.AddRange(DescribeResponse.Reservations.SelectMany(x => x.Instances).Select(x => x.InstanceId).Take(Count));
+				
+				StartScope.Span.SetTag("req.instanceIds", string.Join(",", StartRequest.InstanceIds));
+				if (StartRequest.InstanceIds.Count > 0)
 				{
-					foreach (string InstanceId in StartRequest.InstanceIds)
+					StartInstancesResponse StartResponse = await Client.StartInstancesAsync(StartRequest);
+					StartScope.Span.SetTag("res.statusCode", (int)StartResponse.HttpStatusCode);
+					StartScope.Span.SetTag("res.numInstances", StartResponse.StartingInstances.Count);
+					if ((int)StartResponse.HttpStatusCode >= 200 && (int)StartResponse.HttpStatusCode <= 299)
 					{
-						Logger.LogInformation("Starting instance {InstanceId} for pool {PoolId}", InstanceId, Pool.Id);
+						foreach (InstanceStateChange InstanceChange in StartResponse.StartingInstances)
+						{
+							Logger.LogInformation("Starting instance {InstanceId} for pool {PoolId} (prev state {PrevState}, current state {CurrentState}", InstanceChange.InstanceId, Pool.Id, InstanceChange.PreviousState, InstanceChange.CurrentState);
+						}
 					}
 				}
 			}
@@ -77,26 +97,32 @@ namespace HordeServer.Services.Impl
 		/// <inheritdoc/>
 		public async Task ShrinkPoolAsync(IPool Pool, IReadOnlyList<IAgent> Agents, int Count)
 		{
+			using IScope Scope = GlobalTracer.Instance.BuildSpan("ShrinkPool").StartActive();
+			Scope.Span.SetTag("poolName", Pool.Name);
+			Scope.Span.SetTag("count", Count);
+			
 			string AwsTagProperty = $"{AwsTagPropertyName}={PoolTagName}:{Pool.Name}";
-
+			
 			// Sort the agents by number of active leases. It's better to shutdown agents currently doing nothing.
-			List<IAgent> SortedAgents = Agents.OrderBy(x => x.Leases.Count).ToList();
-
-			// Try to remove the agents
-			for (int Idx = 0; Idx < SortedAgents.Count && Count > 0; Idx++)
+			List<IAgent> FilteredAgents = Agents.OrderBy(x => x.Leases.Count).ToList();
+			List<IAgent> AgentsWithAwsTags = FilteredAgents.Where(x => x.HasProperty(AwsTagProperty)).ToList(); 
+			List<IAgent> AgentsLimitedByCount = AgentsWithAwsTags.Take(Count).ToList();
+			
+			Scope.Span.SetTag("agents.num", Agents.Count);
+			Scope.Span.SetTag("agents.filtered.num", FilteredAgents.Count);
+			Scope.Span.SetTag("agents.withAwsTags.num", AgentsWithAwsTags.Count);
+			Scope.Span.SetTag("agents.limitedByCount.num", AgentsLimitedByCount.Count);
+
+			foreach (IAgent Agent in AgentsLimitedByCount)
 			{
-				IAgent Agent = SortedAgents[Idx];
-				if (Agent.HasProperty(AwsTagProperty))
+				IAuditLogChannel<AgentId> AgentLogger = AgentCollection.GetLogger(Agent.Id);
+				if (await AgentCollection.TryUpdateSettingsAsync(Agent, bRequestShutdown: true, ShutdownReason: "Autoscaler") != null)
 				{
-					for (IAgent? NewAgent = Agent; NewAgent != null; NewAgent = await AgentCollection.GetAsync(Agent.Id))
-					{
-						if (await AgentCollection.TryUpdateSettingsAsync(NewAgent, bRequestShutdown: true, ShutdownReason: "Autoscaler") != null)
-						{
-							AgentCollection.GetLogger(Agent.Id).LogInformation("Marked for shutdown due to autoscalar (currently {NumLeases} leases outstanding)", NewAgent.Leases.Count);
-							Count--;
-							break;
-						}
-					}
+					AgentLogger.LogInformation("Marked for shutdown due to autoscaling (currently {NumLeases} leases outstanding)", Agent.Leases.Count);
+				}
+				else
+				{
+					AgentLogger.LogError("Unable to mark agent for shutdown due to autoscaling");
 				}
 			}
 		}
