Horde Storage - Fixed so scylla can store a replicator state that is reset (has made no progress) - this fixes failing tests

diff --git a/Engine/Source/Programs/Horde/HordeStorage/Horde.Storage/Implementation/TransactionLog/ScyllaReplicationLog.cs b/Engine/Source/Programs/Horde/HordeStorage/Horde.Storage/Implementation/TransactionLog/ScyllaReplicationLog.cs
--- a/Engine/Source/Programs/Horde/HordeStorage/Horde.Storage/Implementation/TransactionLog/ScyllaReplicationLog.cs
+++ b/Engine/Source/Programs/Horde/HordeStorage/Horde.Storage/Implementation/TransactionLog/ScyllaReplicationLog.cs
@@ -275,12 +275,7 @@ namespace Horde.Storage.Implementation
 
         public async Task UpdateReplicatorState(NamespaceId ns, string replicatorName, ReplicatorState newState)
         {
-            if (newState.LastBucket == null || newState.LastEvent == null)
-            {
-                // no point in saving a invalid state
-                return;
-            }
-            await _mapper.UpdateAsync<ScyllaReplicationState>(new ScyllaReplicationState(ns, replicatorName, newState.LastBucket, newState.LastEvent.Value));
+            await _mapper.UpdateAsync<ScyllaReplicationState>(new ScyllaReplicationState(ns, replicatorName, newState.LastBucket, newState.LastEvent));
         }
 
         public async Task<ReplicatorState?> GetReplicatorState(NamespaceId ns, string name)
@@ -416,7 +411,7 @@ namespace Horde.Storage.Implementation
             LastBucket = null!;
         }
 
-        public ScyllaReplicationState(NamespaceId ns, string name, string lastBucket, Guid lastEvent)
+        public ScyllaReplicationState(NamespaceId ns, string name, string? lastBucket, Guid? lastEvent)
         {
             Namespace = ns.ToString();
             Name = name;
@@ -433,10 +428,10 @@ namespace Horde.Storage.Implementation
         public string Name { get; set; }
 
         [Cassandra.Mapping.Attributes.Column("last_bucket")]
-        public string LastBucket { get; set; }
+        public string? LastBucket { get; set; }
 
         [Cassandra.Mapping.Attributes.Column("last_event")]
-        public Guid LastEvent { get; set; }
+        public Guid? LastEvent { get; set; }
 
         public ReplicatorState ToReplicatorState()
         {
