Horde: Only suppress systemic issue creation when there are non-systemic errors present
#rnx
#skipci

diff --git a/Engine/Source/Programs/Horde/Horde.Build/Issues/Impl/IssueService.cs b/Engine/Source/Programs/Horde/Horde.Build/Issues/Impl/IssueService.cs
--- a/Engine/Source/Programs/Horde/Horde.Build/Issues/Impl/IssueService.cs
+++ b/Engine/Source/Programs/Horde/Horde.Build/Issues/Impl/IssueService.cs
@@ -653,6 +653,7 @@ namespace HordeServer.Services.Impl
 			List<ILogEvent> StepEvents = await LogFileService.FindLogEventsAsync(LogFile);
 
 			// Create a new list of event fingerprints for each one
+			bool NonSystemicErrors = false;
 			Dictionary<NewIssueFingerprint, NewEventGroup> FingerprintToEventGroup = new Dictionary<NewIssueFingerprint, NewEventGroup>();
 			foreach (ILogEvent StepEvent in StepEvents)
 			{
@@ -665,6 +666,14 @@ namespace HordeServer.Services.Impl
 				{
 					if (Matcher.TryGetFingerprint(Job, Node, StepEventData, out Fingerprint))
 					{
+						if (Fingerprint.Type != "Systemic")
+						{
+							if (StepEventData.Severity == EventSeverity.Error)
+							{
+								NonSystemicErrors = true;
+							}
+						}
+
 						break;
 					}
 				}
@@ -681,15 +690,18 @@ namespace HordeServer.Services.Impl
 					EventGroup.Events.Add(new NewEvent(StepEvent, StepEventData));
 				}
 			}
-			
-			List<NewIssueFingerprint> SystemicFingerprints = FingerprintToEventGroup.Keys.Where(x => x.Type == "Systemic").ToList();
-
+						
 			// Only generate user issues for systemic issues where there are no other failures, otherwise this leads to many false positives (Experimental)
-			if ((SystemicFingerprints.Count != 0) && (SystemicFingerprints.Count != FingerprintToEventGroup.Values.Count))
+			if (NonSystemicErrors)
 			{
-				for (int i = 0; i < SystemicFingerprints.Count; i++)
+				List<NewIssueFingerprint> SystemicFingerprints = FingerprintToEventGroup.Keys.Where(x => x.Type == "Systemic").ToList();
+
+				if ((SystemicFingerprints.Count != 0) && (SystemicFingerprints.Count != FingerprintToEventGroup.Values.Count))
 				{
-					FingerprintToEventGroup.Remove(SystemicFingerprints[i]);
+					for (int i = 0; i < SystemicFingerprints.Count; i++)
+					{
+						FingerprintToEventGroup.Remove(SystemicFingerprints[i]);
+					}
 				}
 			}
 
