Horde: Support for replicating Perforce content to Horde Storage. Currently disabled by default.

diff --git a/Engine/Source/Programs/Horde/Horde.Build/Commands/Perforce/ReplicateCommand.cs b/Engine/Source/Programs/Horde/Horde.Build/Commands/Perforce/ReplicateCommand.cs
--- a/Engine/Source/Programs/Horde/Horde.Build/Commands/Perforce/ReplicateCommand.cs
+++ b/Engine/Source/Programs/Horde/Horde.Build/Commands/Perforce/ReplicateCommand.cs
@@ -1,6 +1,9 @@
 // Copyright Epic Games, Inc. All Rights Reserved.
 
 using EpicGames.Core;
+using EpicGames.Horde.Storage;
+using EpicGames.Horde.Storage.Impl;
+using EpicGames.Serialization;
 using HordeServer.Collections;
 using HordeServer.Commits;
 using HordeServer.Commits.Impl;
@@ -22,15 +25,24 @@ namespace HordeServer.Commands
 	[Command("perforce", "replicate", "Replicates commits for a particular change of changes from Perforce")]
 	class ReplicateCommand : Command
 	{
-		[CommandLine("-StreamId=", Required = true)]
-		public List<string> StreamIds { get; set; } = new List<string>();
+		[CommandLine("-Stream=", Required = true)]
+		public string StreamId { get; set; } = String.Empty;
 
 		[CommandLine(Required = true)]
 		public int Change { get; set; }
 
+		[CommandLine]
+		public int BaseChange { get; set; }
+
 		[CommandLine]
 		public int Count { get; set; } = 1;
 
+		[CommandLine]
+		public bool Content { get; set; }
+
+		[CommandLine]
+		public DirectoryReference? OutputDir { get; set; }
+
 		IConfiguration Configuration;
 
 		public ReplicateCommand(IConfiguration Configuration)
@@ -40,28 +52,62 @@ namespace HordeServer.Commands
 
 		public override async Task<int> ExecuteAsync(ILogger Logger)
 		{
-			IServiceProvider ServiceProvider = Startup.CreateServiceProvider(Configuration);
-			CommitService CommitService = ServiceProvider.GetRequiredService<CommitService>();
+			IServiceCollection Services = new ServiceCollection();
+			Startup.AddServices(Services, Configuration);
+
+			OutputDir ??= DirectoryReference.Combine(Program.DataDir, "p4");
+			Logger.LogInformation("Writing output to {OutputDir}", OutputDir);
+
+			DirectoryReference BlobsDir = DirectoryReference.Combine(OutputDir, "blobs");
+			Services.AddSingleton<IStorageClient, FileStorageClient>(SP => new FileStorageClient(BlobsDir, Logger));
+
+			DirectoryReference RootsDir = DirectoryReference.Combine(OutputDir, "roots");
+			DirectoryReference.CreateDirectory(RootsDir);
 
+			IServiceProvider ServiceProvider = Services.BuildServiceProvider();
+			CommitService CommitService = ServiceProvider.GetRequiredService<CommitService>();
+			ICommitCollection CommitCollection = ServiceProvider.GetRequiredService<ICommitCollection>();
 			IStreamCollection StreamCollection = ServiceProvider.GetRequiredService<IStreamCollection>();
+			IStorageClient StorageClient = ServiceProvider.GetRequiredService<IStorageClient>();
+
+			IStream? Stream = await StreamCollection.GetAsync(new StreamId(StreamId));
+			if (Stream == null)
+			{
+				throw new FatalErrorException($"Stream '{Stream}' not found");
+			}
 
 			Dictionary<IStream, int> StreamToFirstChange = new Dictionary<IStream, int>();
-			foreach(string StreamId in StreamIds)
+			StreamToFirstChange[Stream] = Change;
+
+			CommitTree? BaseTree = null; 
+			if (BaseChange != 0)
 			{
-				IStream? Stream = await StreamCollection.GetAsync(new StreamId(StreamId));
-				if (Stream == null)
+				ICommit? BaseCommit = await CommitCollection.GetCommitAsync(Stream.Id, BaseChange);
+				if (BaseCommit == null)
 				{
-					throw new FatalErrorException($"Stream '{Stream}' not found");
+					throw new FatalErrorException($"Unable to find base commit {BaseChange}");
 				}
-				StreamToFirstChange[Stream] = Change;
+
+				FileReference RootFile = FileReference.Combine(RootsDir, $"{BaseCommit.Change}.ref");
+				byte[] RootData = await FileReference.ReadAllBytesAsync(RootFile);
+				BaseTree = new CommitTree(BaseCommit.Change, TreePackObject.Parse(RootData));
 			}
 
-			string ClusterName = StreamToFirstChange.First().Key.ClusterName;
-			await foreach (NewCommit Commit in CommitService.FindCommitsForClusterAsync(ClusterName, StreamToFirstChange).Take(Count))
+			await foreach (NewCommit NewCommit in CommitService.FindCommitsForClusterAsync(Stream.ClusterName, StreamToFirstChange).Take(Count))
 			{
-				string BriefSummary = Commit.Description.Replace('\n', ' ').Substring(0, 50);
-				Logger.LogInformation("Commit {Change} by {AuthorId}: {Summary}", Commit.Change, Commit.AuthorId, BriefSummary);
-				Logger.LogInformation(" - Base path: {BasePath}", Commit.BasePath);
+				string BriefSummary = NewCommit.Description.Replace('\n', ' ').Substring(0, 50);
+				Logger.LogInformation("Commit {Change} by {AuthorId}: {Summary}", NewCommit.Change, NewCommit.AuthorId, BriefSummary);
+				Logger.LogInformation(" - Base path: {BasePath}", NewCommit.BasePath);
+
+				if (Content)
+				{
+					BaseTree = await CommitService.FindCommitTreeAsync(Stream, NewCommit.Change, BaseTree);
+
+					FileReference RootFile = FileReference.Combine(RootsDir, $"{NewCommit.Change}.ref");
+					ReadOnlyMemory<byte> RootData = BaseTree.RootObject.ToCbObject().GetView();
+					await FileReference.WriteAllBytesAsync(RootFile, RootData.ToArray());
+					Logger.LogInformation("Written root to {RootFile} ({Size:n0} bytes)", RootFile, RootFile.ToFileInfo().Length);
+				}
 			}
 
 			return 0;
diff --git a/Engine/Source/Programs/Horde/Horde.Build/Perforce/ICommit.cs b/Engine/Source/Programs/Horde/Horde.Build/Perforce/ICommit.cs
--- a/Engine/Source/Programs/Horde/Horde.Build/Perforce/ICommit.cs
+++ b/Engine/Source/Programs/Horde/Horde.Build/Perforce/ICommit.cs
@@ -70,7 +70,7 @@ namespace HordeServer.Commits
 		/// <summary>
 		/// Name of a reference describing the tree at this changelist
 		/// </summary>
-		public string? TreeRef { get; }
+		public RefId? TreeRefId { get; }
 	}
 
 	/// <summary>
@@ -102,8 +102,8 @@ namespace HordeServer.Commits
 		/// <inheritdoc cref="ICommit.DateUtc"/>
 		public DateTime DateUtc { get; set; }
 
-		/// <inheritdoc cref="ICommit.TreeRef"/>
-		public string? TreeRef { get; set; }
+		/// <inheritdoc cref="ICommit.TreeRefId"/>
+		public RefId? TreeRefId { get; set; }
 
 		/// <summary>
 		/// Constructor
@@ -111,7 +111,7 @@ namespace HordeServer.Commits
 		public NewCommit(ICommit Commit)
 			: this(Commit.StreamId, Commit.Change, Commit.OriginalChange, Commit.AuthorId, Commit.OwnerId, Commit.Description, Commit.BasePath, Commit.DateUtc)
 		{
-			this.TreeRef = Commit.TreeRef;
+			this.TreeRefId = Commit.TreeRefId;
 		}
 
 		/// <summary>
@@ -129,20 +129,4 @@ namespace HordeServer.Commits
 			this.DateUtc = DateUtc;
 		}
 	}
-
-	/// <summary>
-	/// Extension methods for commits
-	/// </summary>
-	static class CommitExtensions
-	{
-		/// <summary>
-		/// Get the ref id for a particular commit
-		/// </summary>
-		/// <param name="Commit"></param>
-		/// <returns></returns>
-		public static RefId GetRefId(this ICommit Commit)
-		{
-			return new RefId(IoHash.Compute(Encoding.UTF8.GetBytes(Commit.TreeRef!)));
-		}
-	}
 }
diff --git a/Engine/Source/Programs/Horde/Horde.Build/Perforce/Impl/CommitCollection.cs b/Engine/Source/Programs/Horde/Horde.Build/Perforce/Impl/CommitCollection.cs
--- a/Engine/Source/Programs/Horde/Horde.Build/Perforce/Impl/CommitCollection.cs
+++ b/Engine/Source/Programs/Horde/Horde.Build/Perforce/Impl/CommitCollection.cs
@@ -11,6 +11,7 @@ using System.Collections.Generic;
 using System.Linq;
 using System.Threading.Tasks;
 using MongoDB.Bson.Serialization.Attributes;
+using EpicGames.Horde.Storage;
 
 namespace HordeServer.Commits.Impl
 {
@@ -48,8 +49,8 @@ namespace HordeServer.Commits.Impl
 			[BsonElement("p")]
 			public string BasePath { get; set; } = String.Empty;
 
-			[BsonElement("tr")]
-			public string? TreeRef { get; set; }
+			[BsonElement("r")]
+			public RefId? TreeRefId { get; set; }
 
 			[BsonElement("t")]
 			public DateTime DateUtc { get; set; }
@@ -76,7 +77,7 @@ namespace HordeServer.Commits.Impl
 				}
 				this.Description = NewCommit.Description;
 				this.BasePath = NewCommit.BasePath;
-				this.TreeRef = NewCommit.TreeRef;
+				this.TreeRefId = NewCommit.TreeRefId;
 				this.DateUtc = NewCommit.DateUtc;
 			}
 		}
diff --git a/Engine/Source/Programs/Horde/Horde.Build/Perforce/Impl/CommitService.cs b/Engine/Source/Programs/Horde/Horde.Build/Perforce/Impl/CommitService.cs
--- a/Engine/Source/Programs/Horde/Horde.Build/Perforce/Impl/CommitService.cs
+++ b/Engine/Source/Programs/Horde/Horde.Build/Perforce/Impl/CommitService.cs
@@ -31,6 +31,7 @@ using System.IO;
 using HordeCommon;
 using EpicGames.Horde.Storage;
 using System.Text;
+using Microsoft.Extensions.Options;
 
 namespace HordeServer.Commits.Impl
 {
@@ -38,13 +39,77 @@ namespace HordeServer.Commits.Impl
 	using StreamId = StringId<IStream>;
 	using IRef = EpicGames.Horde.Storage.IRef;
 
+	/// <summary>
+	/// Options for the commit service
+	/// </summary>
+	public class CommitServiceOptions
+	{
+		/// <summary>
+		/// Whether to mirror commit metadata to the database
+		/// </summary>
+		public bool Metadata { get; set; }
+
+		/// <summary>
+		/// Specific list of streams to enable commit mirroring for
+		/// </summary>
+		public List<StreamId> MetadataStreams { get; set; } = new List<StreamId>();
+
+		/// <summary>
+		/// Whether to mirror content to Horde storage
+		/// </summary>
+		public bool Content { get; set; }
+
+		/// <summary>
+		/// List of streams for which content mirroring should be enabled
+		/// </summary>
+		public List<StreamId> ContentStreams { get; set; } = new List<StreamId>();
+
+		/// <summary>
+		/// Namespace to store content replicated from Perforce
+		/// </summary>
+		public NamespaceId NamespaceId { get; set; } = new NamespaceId("horde.p4");
+
+		/// <summary>
+		/// Options for how objects are packed together
+		/// </summary>
+		public TreePackOptions Packing { get; set; } = new TreePackOptions();
+	}
+
+	/// <summary>
+	/// Information about a commit tree
+	/// </summary>
+	class CommitTree
+	{
+		/// <summary>
+		/// The changelist number
+		/// </summary>
+		public int Change;
+
+		/// <summary>
+		/// Root object describing the tree contents
+		/// </summary>
+		public TreePackObject RootObject;
+
+		/// <summary>
+		/// Constructor
+		/// </summary>
+		public CommitTree(int Change, TreePackObject RootObject)
+		{
+			this.Change = Change;
+			this.RootObject = RootObject;
+		}
+
+		/// <summary>
+		/// Gets the directory at the root of this tree
+		/// </summary>
+		public TreePackDirNode RootDirNode => TreePackDirNode.Parse(RootObject.GetRootNode());
+	}
+
 	/// <summary>
 	/// Service which mirrors changes from Perforce
 	/// </summary>
 	class CommitService : ICommitService, IHostedService, IDisposable
 	{
-		bool EnableUpdates = false;
-
 		/// <summary>
 		/// Metadata about a stream that needs to have commits mirrored
 		/// </summary>
@@ -60,8 +125,7 @@ namespace HordeServer.Commits.Impl
 			}
 		}
 
-		NamespaceId NamespaceId { get; } = new NamespaceId("default");
-		BucketId BucketId { get; } = new BucketId("commits");
+		public CommitServiceOptions Options { get; }
 
 		// Redis
 		IDatabase Redis;
@@ -73,15 +137,16 @@ namespace HordeServer.Commits.Impl
 		IPerforceService PerforceService;
 		IUserCollection UserCollection;
 		ILogger<CommitService> Logger;
-		ITicker UpdateCommitsTicker;
 
 		const int MaxBackgroundTasks = 2;
 
 		/// <summary>
 		/// Constructor
 		/// </summary>
-		public CommitService(IDatabase Redis, ICommitCollection CommitCollection, IStorageClient StorageClient, IStreamCollection StreamCollection, IPerforceService PerforceService, IUserCollection UserCollection, IClock Clock, ILogger<CommitService> Logger)
+		public CommitService(IDatabase Redis, ICommitCollection CommitCollection, IStorageClient StorageClient, IStreamCollection StreamCollection, IPerforceService PerforceService, IUserCollection UserCollection, IClock Clock, IOptions<CommitServiceOptions> Options, ILogger<CommitService> Logger)
 		{
+			this.Options = Options.Value;
+
 			this.Redis = Redis;
 			this.RedisDirtyStreams = new RedisSet<StreamId>(Redis, RedisBaseKey.Append("streams"));
 			this.RedisReservations = new RedisSortedSet<StreamId>(Redis, RedisBaseKey.Append("reservations"));
@@ -92,47 +157,70 @@ namespace HordeServer.Commits.Impl
 			this.PerforceService = PerforceService;
 			this.UserCollection = UserCollection;
 			this.Logger = Logger;
-			this.UpdateCommitsTicker = Clock.AddSharedTicker<CommitService>(TimeSpan.FromSeconds(30.0), UpdateCommitsAsync, Logger);
+			this.UpdateMetadataTicker = Clock.AddSharedTicker<CommitService>(TimeSpan.FromSeconds(30.0), UpdateMetadataAsync, Logger);
 		}
 
 		/// <inheritdoc/>
-		public void Dispose() => UpdateCommitsTicker.Dispose();
+		public void Dispose() => UpdateMetadataTicker.Dispose();
 
 		/// <inheritdoc/>
 		public async Task StartAsync(CancellationToken CancellationToken)
 		{
-			if (EnableUpdates)
+			if (Options.Metadata)
+			{
+				await StartMetadataReplicationAsync();
+			}
+			if (Options.Content)
 			{
-				await UpdateCommitsTicker.StartAsync();
-				await StartTreeUpdatesAsync();
+				await StartContentReplicationAsync();
 			}
 		}
 
 		/// <inheritdoc/>
 		public async Task StopAsync(CancellationToken CancellationToken)
 		{
-			if (EnableUpdates)
+			if (Options.Content)
 			{
-				await StopTreeUpdatesAsync();
-				await UpdateCommitsTicker.StopAsync();
+				await StopContentReplicationAsync();
+			}
+			if (Options.Metadata)
+			{
+				await StopMetadataReplicationAsync();
 			}
 		}
 
-		#region Commit updates
+		#region Metadata updates
+
+		ITicker UpdateMetadataTicker;
+
+		Task StartMetadataReplicationAsync()
+		{
+			return UpdateMetadataTicker.StartAsync();
+		}
+
+		Task StopMetadataReplicationAsync()
+		{
+			return UpdateMetadataTicker.StopAsync();
+		}
 
 		/// <summary>
 		/// Polls Perforce for submitted changes
 		/// </summary>
 		/// <param name="CancellationToken"></param>
 		/// <returns></returns>
-		async ValueTask UpdateCommitsAsync(CancellationToken CancellationToken)
+		async ValueTask UpdateMetadataAsync(CancellationToken CancellationToken)
 		{
 			List<IStream> Streams = await StreamCollection.FindAllAsync();
+			if (Options.MetadataStreams.Count > 0)
+			{
+				Streams.RemoveAll(x => !Options.MetadataStreams.Contains(x.Id));
+			}
+
 			foreach (IGrouping<string, IStream> Group in Streams.GroupBy(x => x.ClusterName, StringComparer.OrdinalIgnoreCase))
 			{
 				try
 				{
-					await UpdateCommitsForClusterAsync(Group.Key, Group);
+					await UpdateMetadataForClusterAsync(Group.Key, Group);
 				}
 				catch (Exception Ex)
 				{
@@ -142,12 +230,12 @@ namespace HordeServer.Commits.Impl
 		}
 
 		/// <summary>
-		/// Updates all the streams on a particular server cluster
+		/// Updates metadata for all the streams on a particular server cluster
 		/// </summary>
 		/// <param name="ClusterName"></param>
 		/// <param name="Streams"></param>
 		/// <returns></returns>
-		async Task UpdateCommitsForClusterAsync(string ClusterName, IEnumerable<IStream> Streams)
+		async Task UpdateMetadataForClusterAsync(string ClusterName, IEnumerable<IStream> Streams)
 		{
 			// Find the minimum changelist number to query
 			Dictionary<IStream, int> StreamToFirstChange = new Dictionary<IStream, int>();
@@ -399,9 +487,33 @@ namespace HordeServer.Commits.Impl
 				return null;
 			}
 		}
+
+		/// <summary>
+		/// Gets the view for a stream
+		/// </summary>
+		/// <param name="Connection">The Perforce connection</param>
+		/// <param name="StreamName">Name of the stream to be mirrored</param>
+		static async Task<ViewMap> GetStreamViewAsync(IPerforceConnection Connection, string StreamName)
+		{
+			StreamRecord Record = await Connection.GetStreamAsync(StreamName, true);
+
+			ViewMap View = new ViewMap();
+			foreach (string ViewLine in Record.View)
+			{
+				Match Match = Regex.Match(ViewLine, "^(-?)([^ ]+) *([^ ]+)$");
+				if (!Match.Success)
+				{
+					throw new PerforceException($"Unable to parse stream view: '{ViewLine}'");
+				}
+				View.Entries.Add(new ViewMapEntry(Match.Groups[1].Length == 0, Match.Groups[2].Value, Match.Groups[3].Value));
+			}
+
+			return View;
+		}
+
 		#endregion
 
-		#region Tree Updates
+		#region Content Updates
 
 		RedisKey RedisBaseKey { get; } = new RedisKey("commits/");
 		RedisChannel<StreamId> RedisUpdateChannel { get; } = new RedisChannel<StreamId>("commits/streams");
@@ -410,16 +522,17 @@ namespace HordeServer.Commits.Impl
 		Channel<StreamId>? RedisUpdateStreams;
 		RedisChannelSubscription<StreamId>? RedisUpdateSubscription;
 		Task? StreamUpdateTask;
+
 		RedisList<int> RedisStreamChanges(StreamId StreamId) => new RedisList<int>(Redis, RedisBaseKey.Append($"stream/{StreamId}/changes"));
 
-		async Task StartTreeUpdatesAsync()
+		async Task StartContentReplicationAsync()
 		{
 			RedisUpdateStreams = Channel.CreateUnbounded<StreamId>();
 			RedisUpdateSubscription = await Redis.Multiplexer.GetSubscriber().SubscribeAsync(RedisUpdateChannel, (_, StreamId) => RedisUpdateStreams.Writer.TryWrite(StreamId));
-			StreamUpdateTask = Task.Run(() => UpdateTreesAsync());
+			StreamUpdateTask = Task.Run(() => UpdateContentAsync());
 		}
 
-		async Task StopTreeUpdatesAsync()
+		async Task StopContentReplicationAsync()
 		{
 			if (RedisUpdateStreams != null)
 			{
@@ -434,14 +547,14 @@ namespace HordeServer.Commits.Impl
 			await StreamUpdateTask!;
 		}
 
-		async Task UpdateTreesAsync()
+		async Task UpdateContentAsync()
 		{
 			List<(StreamId, Task)> BackgroundTasks = new List<(StreamId, Task)>(MaxBackgroundTasks);
 			while (!RedisUpdateStreams!.Reader.Completion.IsCompleted || BackgroundTasks.Count > 0)
 			{
 				try
 				{
-					await UpdateTreesInternalAsync(BackgroundTasks);
+					await UpdateContentInternalAsync(BackgroundTasks);
 				}
 				catch (Exception Ex)
 				{
@@ -450,7 +563,7 @@ namespace HordeServer.Commits.Impl
 			}
 		}
 
-		async Task UpdateTreesInternalAsync(List<(StreamId, Task)> BackgroundTasks)
+		async Task UpdateContentInternalAsync(List<(StreamId, Task)> BackgroundTasks)
 		{
 			// Remove any complete background tasks
 			for (int Idx = BackgroundTasks.Count - 1; Idx >= 0; Idx--)
@@ -488,7 +601,7 @@ namespace HordeServer.Commits.Impl
 						{
 							if (await RedisReservations.AddAsync(CheckStream, NewScore, When.NotExists))
 							{
-								Task NewTask = Task.Run(() => UpdateStreamTreesAsync(CheckStream));
+								Task NewTask = Task.Run(() => UpdateStreamContentAsync(CheckStream));
 								BackgroundTasks.Add((CheckStream, NewTask));
 
 								if (BackgroundTasks.Count == MaxBackgroundTasks)
@@ -526,7 +639,7 @@ namespace HordeServer.Commits.Impl
 			}
 		}
 
-		async Task UpdateStreamTreesAsync(StreamId StreamId)
+		async Task UpdateStreamContentAsync(StreamId StreamId)
 		{
 			Stopwatch Timer = Stopwatch.StartNew();
 
@@ -534,7 +647,7 @@ namespace HordeServer.Commits.Impl
 			for (; ; )
 			{
 				// Update the stream, updating the reservation every 30 seconds
-				Task InternalTask = Task.Run(() => UpdateStreamTreesInternalAsync(StreamId, StreamChanges));
+				Task InternalTask = Task.Run(() => UpdateStreamContentInternalAsync(StreamId, StreamChanges));
 				while (!InternalTask.IsCompleted)
 				{
 					Task DelayTask = Task.Delay(TimeSpan.FromSeconds(15.0));
@@ -559,7 +672,7 @@ namespace HordeServer.Commits.Impl
 			await RedisReservations.RemoveAsync(StreamId);
 		}
 
-		async Task UpdateStreamTreesInternalAsync(StreamId StreamId, RedisList<int> Changes)
+		async Task UpdateStreamContentInternalAsync(StreamId StreamId, RedisList<int> Changes)
 		{
 			IStream? Stream = await StreamCollection.GetAsync(StreamId);
 			if (Stream == null)
@@ -567,600 +680,288 @@ namespace HordeServer.Commits.Impl
 				return;
 			}
 
-			using IPerforceConnection? Perforce = await PerforceService.GetServiceUserConnection(Stream.ClusterName);
-			if (Perforce == null)
-			{
-				return;
-			}
-
-			InfoRecord ServerInfo = await Perforce.GetInfoAsync(InfoOptions.ShortOutput);
-
-			ViewMap View = await GetStreamViewAsync(Perforce, Stream.Name);
-
-#if ENABLE_TREE_MIRRORING
 			ICommit? PrevCommit = null;
-			StreamTreeRef? PrevRoot = null;
-#endif
-			ObjectSet ObjectSet = new ObjectSet(StorageClient, NamespaceId, 256 * 1024, DateTime.UtcNow);
 			for (; ; )
 			{
-				// Get the first two changelists to update
+				// Get the first two changes to be mirrored. The first one should already exist, unless it's the start of replication for this stream.
 				int[] Values = await Changes.RangeAsync(0, 1);
 				if (Values.Length < 2)
 				{
 					break;
 				}
-#if ENABLE_TREE_MIRRORING
-				// Clear the previous commit if it no longer matches
 				if (PrevCommit != null && PrevCommit.Change != Values[0])
 				{
 					PrevCommit = null;
 				}
 
-				// Get the previous tree if we don't have it from the previous iteration
-				if (PrevCommit == null)
-				{
-					PrevCommit = await CommitCollection.GetCommitAsync(Stream.Id, Values[0]);
-					if (PrevCommit == null)
-					{
-						await Changes.LeftPopAsync();
-						continue;
-					}
-
-					bool HasTree = false;
-					if (PrevCommit.TreeRef != null)
-					{
-						CommitTree? PrevCommitTree = await ReadCommitTreeAsync(PrevCommit);
-						if (PrevCommitTree != null)
-						{
-							try
-							{
-								throw new NotImplementedException();
-//								await Builder.ResetAsync(PrevCommitTree);
-//								HasTree = true;
-							}
-							catch (Exception Ex)
-							{
-								Logger.LogError(Ex, "Unable to read commit tree for {StreamId} CL {Change}", Stream.Id, PrevCommit.Change);
-							}
-						}
-					}
+				// Get the commit we want to mirror
+				int Change = Values[(PrevCommit == null) ? 0 : 1];
 
-					if (!HasTree)
-					{
-						PrevRoot = await AddTreeRefAsync(Perforce, Stream, View, null, null, PrevCommit, ObjectSet, ServerInfo.Utf8PathComparer);
-					}
+				// Check if the tree already exists. If it does, we don't need to run replication again.
+				ICommit? Commit = await CommitCollection.GetCommitAsync(StreamId, Change);
+				if (Commit == null)
+				{
+					Logger.LogWarning("Missing commit for {Stream} change {Change}", StreamId, Change);
 				}
-
-				// Move to the next commit
-				ICommit? Commit = await CommitCollection.GetCommitAsync(StreamId, Values[1]);
-				StreamTreeRef? Root = null;
-				if (Commit != null)
+				else if (Commit.TreeRefId != null)
 				{
-					Root = await AddTreeRefAsync(Perforce, Stream, View, PrevCommit, PrevRoot, Commit, ObjectSet, ServerInfo.Utf8PathComparer);
+					Logger.LogInformation("Skipping replication for {Stream} change {Change}; tree already exists", StreamId, Change);
+				}
+				else
+				{
+					Commit = await UpdateCommitTreeAsync(Stream, Commit, PrevCommit);
 				}
 
-				// Move to the next commit
-				PrevCommit = Commit;
-				PrevRoot = Root;
-#endif
 				// Remove the first change number from this queue
 				await Changes.LeftPopAsync();
+				PrevCommit = Commit;
 			}
 		}
 
-		async Task<CommitTree?> GetCommitTreeAsync(ICommit Commit)
+		/// <summary>
+		/// Gets the tree for a particular commit
+		/// </summary>
+		/// <param name="Commit">The commit instance</param>
+		/// <returns>The corresponding commit tree</returns>
+		public async Task<CommitTree?> GetTreeAsync(ICommit Commit)
 		{
-			if (Commit.TreeRef == null)
+			if (Commit.TreeRefId == null)
 			{
 				return null;
 			}
 
-			IRef? Ref = await StorageClient.GetRefAsync(NamespaceId, BucketId, Commit.GetRefId());
-			if (Ref == null)
+			IRef? Ref = await StorageClient.GetRefAsync(Options.NamespaceId, GetStreamBucketId(Commit.StreamId), Commit.TreeRefId.Value);
+			if(Ref == null)
 			{
 				return null;
 			}
 
-			return CbSerializer.Deserialize<CommitTree>(Ref.Value.AsField());
-		}
-
-		async Task<StreamTreeRef> AddTreeRefAsync(IPerforceConnection Perforce, IStream Stream, ViewMap View, ICommit? PrevCommit, StreamTreeRef? PrevRoot, ICommit Commit, ObjectSet ObjectSet, Utf8StringComparer PathComparer)
-		{
-			StreamTreeRef Root;
-			if (PrevCommit == null || PrevRoot == null)
-			{
-				Logger.LogInformation("Performing snapshot of {StreamId} at CL {Change}", Commit.StreamId, Commit.Change);
-				Root = await CreateTreeAsync(Perforce, View, Commit.Change, ObjectSet, PathComparer);
-			}
-			else
-			{
-				Logger.LogInformation("Performing incremental update of {StreamId} from CL {PrevChange} to CL {Change}", Commit.StreamId, PrevCommit.Change, Commit.Change);
-				Root = await UpdateTreeAsync(PrevRoot, Perforce, View, Commit.Change, ObjectSet, PathComparer);
-			}
-
-			CommitTree Tree = await FlushAsync(ObjectSet, Root);
-			string RefName = $"tree_{Commit.StreamId}_{Commit.Change}";
-			RefId RefId = new RefId(IoHash.Compute(Encoding.UTF8.GetBytes(RefName)));
-			await StorageClient.SetRefAsync(NamespaceId, BucketId, RefId, Tree.Serialize(Stream.Name));
-
-			NewCommit NewCommit = new NewCommit(Commit);
-			NewCommit.TreeRef = RefName;
-			await CommitCollection.AddOrReplaceAsync(NewCommit);
-
-			return Tree.Root;
+			return new CommitTree(Commit.Change, TreePackObject.Parse(Ref));
 		}
 
-		#endregion
-
-		#region Tree Snapshots
-
 		/// <summary>
-		/// Creates a snapshot of a stream at a particular changelist
+		/// Replicates the contents of a stream to Horde storage, optionally using the given change as a starting point
 		/// </summary>
-		/// <param name="Connection">The Perforce repository to mirror</param>
-		/// <param name="View">View for the stream</param>
-		/// <param name="Change">Changelist to replicate</param>
-		/// <param name="ObjectSet"></param>
-		/// <param name="PathComparer">Comparison type for paths</param>
-		async Task<StreamTreeRef> CreateTreeAsync(IPerforceConnection Connection, ViewMap View, int Change, ObjectSet ObjectSet, Utf8StringComparer PathComparer)
+		/// <param name="Stream">The stream to replicate</param>
+		/// <param name="Commit">Commit to store the tree ref</param>
+		/// <param name="BaseCommit">Previous commit to use as a base</param>
+		/// <returns>New ref id for the given commit</returns>
+		public async Task<ICommit> UpdateCommitTreeAsync(IStream Stream, ICommit Commit, ICommit? BaseCommit = null)
 		{
-			StreamTreeBuilder Root = new StreamTreeBuilder();
-
-			// Optimize the view by removing unnecessary inclusions of the same path
-			View = new ViewMap(View);
-			for (int Idx = 0; Idx < View.Entries.Count; Idx++)
+			// Get the ref corresponding to BaseChange
+			CommitTree? BaseTree = null;
+			if (BaseCommit != null)
 			{
-				ViewMapEntry Entry = View.Entries[Idx];
-				if (Entry.Include && Entry.IsPathWildcard())
+				if (BaseCommit.TreeRefId == null)
 				{
-					while (Idx + 1 < View.Entries.Count && View.Entries[Idx + 1].Include && View.Entries[Idx + 1].SourcePrefix.StartsWith(Entry.SourcePrefix, PathComparer))
-					{
-						View.Entries.RemoveAt(Idx + 1);
-					}
+					Logger.LogWarning("Base commit for stream {StreamId} change {Change} does not have a mirrored tree", Stream.Id, BaseCommit.Change);
 				}
-			}
-
-			// In order to optimize for the common case where we're adding multiple files to the same directory, we keep an open list of
-			// StreamTreeBuilder objects down to the last modified node. Whenever we need to move to a different node, we write the 
-			// previous path to the object store.
-			for(int Idx = 0; Idx < View.Entries.Count; Idx++)
-			{
-				ViewMapEntry Entry = View.Entries[Idx];
-				if (Entry.Include)
+				else
 				{
-					ViewMap FilteredView = new ViewMap();
-					FilteredView.Entries.Add(Entry);
-
-					for (int OtherIdx = Idx + 1; OtherIdx < View.Entries.Count; OtherIdx++)
-					{
-						ViewMapEntry OtherEntry = View.Entries[OtherIdx];
-						if (!OtherEntry.Include)
-						{
-							if (OtherEntry.SourcePrefix.StartsWith(Entry.SourcePrefix, PathComparer) || Entry.SourcePrefix.StartsWith(OtherEntry.SourcePrefix, PathComparer))
-							{
-								FilteredView.Entries.Add(OtherEntry);
-							}
-						}
-					}
-
-					Logger.LogInformation("Adding {Source}@{Change}", Entry.Source, Change);
-					await AddDepotFilesToSnapshotAsync(Root, Connection, $"{Entry.Source}@{Change}", View, PathComparer, ObjectSet);
+					BaseTree = new CommitTree(BaseCommit.Change, TreePackObject.Parse(await StorageClient.GetRefAsync(Options.NamespaceId, GetStreamBucketId(Stream.Id), BaseCommit.TreeRefId.Value)));
 				}
 			}
 
-			return Root.EncodeRef(x => WriteTree(ObjectSet, x));
-		}
-
-		/// <summary>
-		/// Record returned by the fstat command
-		/// </summary>
-		public class Utf8FStatRecord
-		{
-			[PerforceTag("depotFile", Optional = true)]
-			public Utf8String DepotFile { get; set; }
-
-			[PerforceTag("headType", Optional = true)]
-			public Utf8String HeadType { get; set; }
-
-			[PerforceTag("headRev", Optional = true)]
-			public int HeadRevision { get; set; }
+			// Create the new root object
+			CommitTree Tree = await FindCommitTreeAsync(Stream, Commit.Change, BaseTree);
 
-			[PerforceTag("digest", Optional = true)]
-			public Utf8String Digest { get; set; }
+			// Write the new ref
+			string RefName = $"{Commit.Change}";
+			RefId TreeRefId = new RefId(IoHash.Compute(Encoding.UTF8.GetBytes(RefName)));
+			await StorageClient.SetRefAsync(Options.NamespaceId, new BucketId(Stream.Id.ToString()), TreeRefId, Tree.RootObject.ToCbObject());
 
-			[PerforceTag("fileSize", Optional = true)]
-			public long FileSize { get; set; }
-
-			[PerforceTag("type", Optional = true)]
-			public Utf8String Type { get; set; }
-		}
-
-		/// <summary>
-		/// Adds a set of files in the depot to a snapshot
-		/// </summary>
-		/// <param name="Root"></param>
-		/// <param name="Connection"></param>
-		/// <param name="FileSpec">The files to add</param>
-		/// <param name="View">View for the stream</param>
-		/// <param name="Comparer">Comparison type to use for names</param>
-		/// <param name="ObjectSet"></param>
-		static async Task AddDepotFilesToSnapshotAsync(StreamTreeBuilder Root, IPerforceConnection Connection, string FileSpec, ViewMap View, Utf8StringComparer Comparer, ObjectSet ObjectSet)
-		{
-			const string Filter = "^headAction=delete ^headAction=move/delete";
-			const string Fields = "depotFile,fileSize,digest,headType,headRev";
-
-			List<string> Arguments = new List<string>();
-			Arguments.Add("-Ol");
-			Arguments.Add("-Op");
-			Arguments.Add("-Os");
-			Arguments.Add("-F");
-			Arguments.Add(Filter);
-			Arguments.Add("-T");
-			Arguments.Add(String.Join(",", Fields));
-			Arguments.Add(FileSpec);
-
-			List<Utf8String> Segments = new List<Utf8String>();
-
-			IAsyncEnumerable<PerforceResponse<Utf8FStatRecord>> Responses = Connection.StreamCommandAsync<Utf8FStatRecord>("fstat", Arguments);
-			await foreach (PerforceResponse<Utf8FStatRecord> Response in Responses)
-			{
-				Utf8FStatRecord File = Response.Data;
-				if (File.Digest != Utf8String.Empty)
-				{
-					Utf8String TargetFile;
-					if (File.DepotFile != Utf8String.Empty && View.TryMapFile(File.DepotFile, Comparer, out TargetFile))
-					{
-						await AddFileAsync(Root, TargetFile, File, Segments, ObjectSet);
-					}
-				}
-			}
+			// Update the commit
+			NewCommit NewCommit = new NewCommit(Commit);
+			NewCommit.TreeRefId = TreeRefId;
+			return await CommitCollection.AddOrReplaceAsync(NewCommit);
 		}
 
 		/// <summary>
-		/// Adds a single depot file to a snapshot
+		/// Replicates the contents of a stream to Horde storage, optionally using the given change as a starting point
 		/// </summary>
-		/// <param name="Root">The root tree builder</param>
-		/// <param name="Path">The stream path for the file</param>
-		/// <param name="MetaData">The file metadata</param>
-		/// <param name="Segments"></param>
-		/// <param name="ObjectSet"></param>
-		static async Task AddFileAsync(StreamTreeBuilder Root, Utf8String Path, Utf8FStatRecord MetaData, List<Utf8String> Segments, ObjectSet ObjectSet)
+		/// <param name="Stream">The stream to replicate</param>
+		/// <param name="Change">Commit to store the tree ref</param>
+		/// <param name="BaseTree">The base tree to update from</param>
+		/// <returns>Root tree object</returns>
+		public async Task<CommitTree> FindCommitTreeAsync(IStream Stream, int Change, CommitTree? BaseTree)
 		{
-			StreamTreeBuilder Node = Root;
-
-			// Split the path into segments
-			Segments.Clear();
-			for (int Pos = 0; ;)
-			{
-				int NextPos = Path.IndexOf('/', Pos);
-				if (NextPos == -1)
-				{
-					Segments.Add(Path.Substring(Pos));
-					break;
-				}
-				else
-				{
-					Segments.Add(Path.Substring(Pos, NextPos - Pos));
-					Pos = NextPos + 1;
-				}
-			}
+			// Create a client to replicate from this stream, and connect to the server
+			ReplicationClient ClientInfo = await FindOrAddReplicationClientAsync(Stream);
+			using IPerforceConnection Perforce = await PerforceConnection.CreateAsync(ClientInfo.Settings, Logger);
 
-			// Match as many nodes as we can from the existing path
-			int SegmentIdx = 0;
-			for(; SegmentIdx < Segments.Count - 1; SegmentIdx++)
+			// Get the initial directory state
+			TreePackDirWriter DirWriter;
+			if (BaseTree == null)
 			{
-				Utf8String Segment = Segments[SegmentIdx];
-				if (!Node.NameToTreeBuilder.TryGetValue(Segment, out StreamTreeBuilder? NextNode))
-				{
-					break;
-				}
-				Node = NextNode;
+				await FlushWorkspaceAsync(ClientInfo, Perforce, 0);
+				DirWriter = new TreePackDirWriter(ClientInfo.TreePack);
 			}
-
-			// Collapse the rest of the current path
-			await CollapseChildrenAsync(ObjectSet, Node);
-
-			// Expand the rest of the path
-			for(; SegmentIdx < Segments.Count - 1; SegmentIdx++)
+			else
 			{
-				Utf8String Segment = Segments[SegmentIdx];
-				Node = await ExpandChildAsync(ObjectSet, Node, Segment);
+				await FlushWorkspaceAsync(ClientInfo, Perforce, BaseTree.Change);
+				DirWriter = new TreePackDirWriter(ClientInfo.TreePack, BaseTree.RootDirNode);
 			}
 
-			// Add the filename
-			Utf8String FileName = Segments[Segments.Count - 1];
-			Node.NameToFile[FileName] = CreateStreamFile(MetaData);
-		}
-
-#endregion
+			// Apply all the updates
+			Logger.LogInformation("Updating client {Client} to changelist {Change}", ClientInfo.Client.Name, Change);
+			ClientInfo.Change = -1;
 
-#region Tree Incremental Updates
-
-		/// <summary>
-		/// Creates a snapshot for the given stream, given an initial state and list of modified files
-		/// </summary>
-		/// <param name="RootRef"></param>
-		/// <param name="Perforce"></param>
-		/// <param name="Change"></param>
-		/// <param name="View"></param>
-		/// <param name="ObjectSet"></param>
-		/// <param name="PathComparer"></param>
-		/// <returns></returns>
-		static async Task<StreamTreeRef> UpdateTreeAsync(StreamTreeRef RootRef, IPerforceConnection Perforce, ViewMap View, int Change, ObjectSet ObjectSet, Utf8StringComparer PathComparer)
-		{
-			List<FStatRecord> Records = await Perforce.FStatAsync(-1, Change, null, null, -1, FStatOptions.IncludeFileSizes, FileSpecList.Any).ToListAsync();
-			return await UpdateTreeAsync(RootRef, Records, View, ObjectSet, PathComparer);
-		}
-
-		/// <summary>
-		/// Creates a snapshot for the given stream, given an initial state and list of modified files
-		/// </summary>
-		/// <param name="RootRef"></param>
-		/// <param name="Files"></param>
-		/// <param name="View"></param>
-		/// <param name="ObjectSet"></param>
-		/// <param name="PathComparer"></param>
-		/// <returns></returns>
-		static async Task<StreamTreeRef> UpdateTreeAsync(StreamTreeRef RootRef, IList<FStatRecord> Files, ViewMap View, ObjectSet ObjectSet, Utf8StringComparer PathComparer)
-		{
-			// Read the root tree
-			StreamTreeBuilder Root = new StreamTreeBuilder(await ReadTreeAsync(ObjectSet, RootRef));
-
-			// Update the tree
-			foreach (FStatRecord File in Files)
+			Dictionary<int, (string Path, TreePackFileWriter Writer)> Files = new Dictionary<int, (string, TreePackFileWriter)>();
+			await foreach (PerforceResponse Response in Perforce.StreamCommandAsync("sync", Array.Empty<string>(), new string[] { $"//...@{Change}" }, null, typeof(SyncRecord), true, default))
 			{
-				Utf8String LocalPathStr;
-				if (File.DepotFile != null && View.TryMapFile(File.DepotFile, PathComparer, out LocalPathStr))
+				PerforceIo? Io = Response.Io;
+				if (Io != null)
 				{
-					Utf8String LocalPath = LocalPathStr;
-
-					StreamTreeBuilder Node = Root;
-					for (; ; )
+					if (Io.Command == PerforceIoCommand.Open)
 					{
-						int NextIdx = LocalPath.IndexOf('/');
-						if (NextIdx == -1)
-						{
-							break;
-						}
-
-						Utf8String Name = LocalPath[0..NextIdx];
-						Node = await ExpandChildAsync(ObjectSet, Node, Name);
-
-						LocalPath = LocalPath[(NextIdx + 1)..];
+						string Path = GetClientRelativePath(Io.Payload, ClientInfo.Client.Root);
+						TreePackFileWriter FileWriter = new TreePackFileWriter(ClientInfo.TreePack);
+						Files[Io.File] = (Path, FileWriter);
 					}
-
-					FileAction Action = File.Action;
-					if (Action == FileAction.None)
+					else if (Io.Command == PerforceIoCommand.Write)
 					{
-						Action = File.HeadAction;
-						if (Action == FileAction.None)
-						{
-							throw new NotImplementedException();
-						}
+						TreePackFileWriter FileWriter = Files[Io.File].Writer;
+						await FileWriter.WriteAsync(Io.Payload, false);
 					}
-
-					Utf8String FileName = LocalPath;
-					if (Action != FileAction.Delete && Action != FileAction.MoveDelete)
+					else if (Io.Command == PerforceIoCommand.Close)
 					{
-						Node.NameToFile[FileName] = CreateStreamFile(File);
+						(string Path, TreePackFileWriter FileWriter) = Files[Io.File];
+						IoHash FileHash = await FileWriter.FinalizeAsync();
+						await DirWriter.FindOrAddFileByPathAsync(Path, TreePackDirEntryFlags.File, FileHash, Sha1Hash.Zero);
+						Files.Remove(Io.File);
 					}
-					else
+					else if (Io.Command == PerforceIoCommand.Unlink)
 					{
-						if (!Node.NameToFile.Remove(FileName))
-						{
-							// TODO: Handle mismatched case?
-							throw new NotImplementedException();
-						}
+						string Path = GetClientRelativePath(Io.Payload, ClientInfo.Client.Root);
+						await DirWriter.RemoveFileByPathAsync(Path);
 					}
 				}
 			}
 
-			return Root.EncodeRef(x => WriteTree(ObjectSet, x));
-		}
-
-#endregion
-
-#region Tree Builder
-
-		/// <summary>
-		/// Adds or expands a tree under the given node
-		/// </summary>
-		/// <param name="ObjectSet"></param>
-		/// <param name="Tree"></param>
-		/// <param name="Name"></param>
-		/// <returns></returns>
-		static async Task<StreamTreeBuilder> ExpandChildAsync(ObjectSet ObjectSet, StreamTreeBuilder Tree, Utf8String Name)
-		{
-			StreamTreeBuilder? Result;
-			if (Tree.NameToTreeBuilder.TryGetValue(Name, out Result))
-			{
-				return Result;
-			}
+			ClientInfo.Change = Change;
 
-			StreamTreeRef? TreeRef;
-			if (Tree.NameToTree.TryGetValue(Name, out TreeRef))
-			{
-				Result = new StreamTreeBuilder(await ReadTreeAsync(ObjectSet, TreeRef));
-
-				Tree.NameToTree.Remove(Name);
-				Tree.NameToTreeBuilder.Add(Name, Result);
-
-				return Result;
-			}
-
-			Result = new StreamTreeBuilder();
-			Tree.NameToTreeBuilder[Name] = Result;
-
-			return Result;
+			// Return the new root object
+			IoHash RootHash = await DirWriter.FlushAsync();
+			TreePackObject RootObject = await ClientInfo.TreePack.FlushAsync(RootHash, DateTime.UtcNow);
+			return new CommitTree(Change, RootObject);
 		}
 
-		/// <summary>
-		/// Collapse children underneath the given tree builder, allowing data to be purged from memory and pushed
-		/// into persistent storage.
-		/// </summary>
-		/// <param name="ObjectSet"></param>
-		/// <param name="Tree"></param>
-		/// <returns></returns>
-		static async Task CollapseChildrenAsync(ObjectSet ObjectSet, StreamTreeBuilder Tree)
+		async Task FlushWorkspaceAsync(ReplicationClient ClientInfo, IPerforceConnection Perforce, int Change)
 		{
-			foreach ((Utf8String ChildName, StreamTreeBuilder ChildBuilder) in Tree.NameToTreeBuilder)
+			if (ClientInfo.Change != Change)
 			{
-				await CollapseChildrenAsync(ObjectSet, ChildBuilder);
-				if (ChildBuilder.NameToFile.Count > 0 || ChildBuilder.NameToTree.Count > 0)
+				ClientInfo.Change = -1;
+				if (Change == 0)
 				{
-					StreamTreeRef ChildTreeRef = ChildBuilder.EncodeRef(x => WriteTree(ObjectSet, x));
-					Tree.NameToTree[ChildName] = ChildTreeRef;
+					Logger.LogInformation("Flushing have table for {Client}", Perforce.Settings.ClientName);
+					await Perforce.SyncQuietAsync(SyncOptions.Force | SyncOptions.KeepWorkspaceFiles, -1, $"//...#0");
 				}
+				else
+				{
+					Logger.LogInformation("Flushing have table for {Client} to change {Change}", Perforce.Settings.ClientName, Change);
+					await Perforce.SyncQuietAsync(SyncOptions.Force | SyncOptions.KeepWorkspaceFiles, -1, $"//...@{Change}");
+				}
+				ClientInfo.Change = Change;
 			}
-			Tree.NameToTreeBuilder.Clear();
 		}
 
-		/// <summary>
-		/// Flush all the current data to storage, optimizing the underlying blobs if necessary
-		/// </summary>
-		/// <returns>The new tree object</returns>
-		static async Task<CommitTree> FlushAsync(ObjectSet ObjectSet, StreamTreeRef RootRef)
+		static string GetClientRelativePath(ReadOnlyMemory<byte> Data, string ClientRoot)
 		{
-			// Find the live set of objects
-			ObjectSet.RootSet.Clear();
-			ObjectSet.RootSet.Add(RootRef.Hash);
+			ReadOnlySpan<byte> Span = Data.Span;
 
-			// Flush all the remaining objects
-			await ObjectSet.FlushAsync();
-
-			// Wait for all the blobs to finish writing and write the new tree
-			return new CommitTree(RootRef, new List<CbBinaryAttachment>(), ObjectSet.PackIndexes.ConvertAll(x => (CbBinaryAttachment)x.DataHash));
-		}
-
-		/// <summary>
-		/// Reads a referenced tree from storage
-		/// </summary>
-		/// <param name="ObjectSet">The object packer</param>
-		/// <param name="TreeRef">The tree reference</param>
-		/// <returns></returns>
-		static async Task<StreamTree> ReadTreeAsync(ObjectSet ObjectSet, StreamTreeRef TreeRef)
-		{
-			ReadOnlyMemory<byte> Data = await ObjectSet.GetObjectDataAsync(TreeRef.Hash);
-			if (Data.Length == 0)
+			string Path = Encoding.UTF8.GetString(Span.Slice(0, Span.IndexOf((byte)0)));
+			if (!Path.StartsWith(ClientRoot, StringComparison.Ordinal))
 			{
-				throw new Exception($"Unable to find tree with hash {TreeRef.Hash} for {TreeRef.Path}");
+				throw new ArgumentException($"Unable to make path {Path} relative to client root {ClientRoot}");
 			}
-			return new StreamTree(TreeRef.Path, new CbObject(Data));
-		}
-
-		/// <summary>
-		/// Serializes a tree to an object packer
-		/// </summary>
-		/// <param name="ObjectSet"></param>
-		/// <param name="Tree"></param>
-		/// <returns></returns>
-		static IoHash WriteTree(ObjectSet ObjectSet, StreamTree Tree)
-		{
-			CbObject Object = Tree.ToCbObject();
-			ReadOnlyMemory<byte> Data = Object.GetView();
 
-			List<IoHash> Refs = new List<IoHash>();
-			Object.IterateAttachments(Field => Refs.Add(Field.AsHash()));
-
-			IoHash Hash = IoHash.Compute(Data.Span);
-			ObjectSet.Add(Hash, Data.Span, Refs.ToArray());
-			return Hash;
+			return Path.Substring(ClientRoot.Length).Replace('\\', '/').TrimStart('/');
 		}
 
-		/// <summary>
-		/// Prints the state of the tree
-		/// </summary>
-		/// <param name="ObjectSet"></param>
-		/// <param name="TreeRef"></param>
-		/// <param name="Logger"></param>
-		/// <returns></returns>
-		public Task PrintAsync(ObjectSet ObjectSet, StreamTreeRef TreeRef, ILogger Logger) => PrintInternalAsync(ObjectSet, TreeRef, String.Empty, Logger);
-
-		async Task PrintInternalAsync(ObjectSet ObjectSet, StreamTreeRef TreeRef, string Prefix, ILogger Logger)
+		class ReplicationClient
 		{
-			StreamTree Tree = await ReadTreeAsync(ObjectSet, TreeRef);
+			public PerforceSettings Settings { get; }
+			public string ClusterName { get; }
+			public InfoRecord ServerInfo { get; }
+			public ClientRecord Client { get; }
+			public int Change { get; set; }
+			public TreePack TreePack { get; }
 
-			foreach ((Utf8String Name, StreamFile File) in Tree.NameToFile)
+			public ReplicationClient(PerforceSettings Settings, string ClusterName, InfoRecord ServerInfo, ClientRecord Client, int Change, TreePack TreePack)
 			{
-				Logger.LogInformation("{Prefix}/{Name} = {Path}#{Revision}", Prefix, Name, File.Path, File.Revision);
-			}
-			foreach ((Utf8String Name, StreamTreeRef ChildTreeRef) in Tree.NameToTree)
-			{
-				await PrintInternalAsync(ObjectSet, ChildTreeRef, $"{Prefix}/{Name}", Logger);
+				this.Settings = Settings;
+				this.ClusterName = ClusterName;
+				this.ServerInfo = ServerInfo;
+				this.Client = Client;
+				this.Change = Change;
+				this.TreePack = TreePack;
 			}
 		}
 
-#endregion
+		Dictionary<StreamId, ReplicationClient> CachedPerforceClients = new Dictionary<StreamId, ReplicationClient>();
 
-#region Tree Storage
-
-		/// <summary>
-		/// Reads the tree for a given commit
-		/// </summary>
-		/// <param name="Commit"></param>
-		async Task<CommitTree> ReadCommitTreeAsync(ICommit Commit)
+		async Task<ReplicationClient?> FindReplicationClientAsync(IStream Stream)
 		{
-			IRef? Ref = await StorageClient.GetRefAsync(NamespaceId, BucketId, Commit.GetRefId());
-			return CbSerializer.Deserialize<CommitTree>(Ref!.Value.AsField());
-		}
+			ReplicationClient? ClientInfo;
+			if (CachedPerforceClients.TryGetValue(Stream.Id, out ClientInfo))
+			{
+				if (!String.Equals(ClientInfo.ClusterName, Stream.ClusterName, StringComparison.Ordinal) && String.Equals(ClientInfo.Client.Stream, Stream.Name, StringComparison.Ordinal))
+				{
+					PerforceSettings ServerSettings = new PerforceSettings(ClientInfo.Settings);
+					ServerSettings.ClientName = null;
 
-#endregion
+					using IPerforceConnection Perforce = await PerforceConnection.CreateAsync(Logger);
+					await Perforce.DeleteClientAsync(DeleteClientOptions.None, ClientInfo.Client.Name);
 
-#region Misc tree manipulation
+					CachedPerforceClients.Remove(Stream.Id);
+					ClientInfo = null;
+				}
+			}
+			return ClientInfo;
+		}
 
-		/// <summary>
-		/// Gets the view for a stream
-		/// </summary>
-		/// <param name="Connection">The Perforce connection</param>
-		/// <param name="StreamName">Name of the stream to be mirrored</param>
-		static async Task<ViewMap> GetStreamViewAsync(IPerforceConnection Connection, string StreamName)
+		async Task<ReplicationClient> FindOrAddReplicationClientAsync(IStream Stream)
 		{
-			StreamRecord Record = await Connection.GetStreamAsync(StreamName, true);
-
-			ViewMap View = new ViewMap();
-			foreach (string ViewLine in Record.View)
+			ReplicationClient? ClientInfo = await FindReplicationClientAsync(Stream);
+			if (ClientInfo == null)
 			{
-				Match Match = Regex.Match(ViewLine, "^(-?)([^ ]+) *([^ ]+)$");
-				if (!Match.Success)
+				using IPerforceConnection? Perforce = await PerforceService.GetServiceUserConnection(Stream.ClusterName);
+				if (Perforce == null)
 				{
-					throw new PerforceException($"Unable to parse stream view: '{ViewLine}'");
+					throw new PerforceException($"Unable to create connection to Perforce server");
 				}
-				View.Entries.Add(new ViewMapEntry(Match.Groups[1].Length == 0, Match.Groups[2].Value, Match.Groups[3].Value));
-			}
 
-			return View;
-		}
+				InfoRecord ServerInfo = await Perforce.GetInfoAsync(InfoOptions.ShortOutput);
 
-		/// <summary>
-		/// Creates a StreamFile from a FileMetaData object
-		/// </summary>
-		/// <param name="MetaData"></param>
-		/// <returns></returns>
-		static StreamFile CreateStreamFile(Utf8FStatRecord MetaData)
-		{
-			Utf8String FileType = Utf8String.Empty;
-			if (MetaData.Type != Utf8String.Empty)
-			{
-				FileType = MetaData.Type;
-			}
-			else if (MetaData.HeadType != Utf8String.Empty)
-			{
-				FileType = MetaData.HeadType;
+				ClientRecord NewClient = new ClientRecord($"Horde.Build_{ServerInfo.ClientHost}_{Stream.Id}", Perforce.Settings.UserName, "/p4/");
+				NewClient.Description = "Created to mirror Perforce content to Horde Storage";
+				NewClient.Owner = Perforce.Settings.UserName;
+				NewClient.Host = ServerInfo.ClientHost;
+				NewClient.Stream = Stream.Name;
+				await Perforce.CreateClientAsync(NewClient);
+				Logger.LogInformation("Created client {ClientName} for {StreamName}", NewClient.Name, Stream.Name);
+
+				PerforceSettings Settings = new PerforceSettings(Perforce.Settings);
+				Settings.ClientName = NewClient.Name;
+				Settings.PreferNativeClient = true;
+
+				TreePack TreePack = new TreePack(StorageClient, Options.NamespaceId, Options.Packing);
+				ClientInfo = new ReplicationClient(Settings, Stream.ClusterName, ServerInfo, NewClient, -1, TreePack);
+				CachedPerforceClients.Add(Stream.Id, ClientInfo);
 			}
-			return new StreamFile(MetaData.DepotFile!, MetaData.FileSize, new FileContentId(Md5Hash.Parse(MetaData.Digest!), FileType), MetaData.HeadRevision);
+			return ClientInfo;
 		}
 
 		/// <summary>
-		/// Creates a StreamFile from a FileMetaData object
+		/// Get the bucket used for particular stream refs
 		/// </summary>
-		/// <param name="MetaData"></param>
+		/// <param name="StreamId"></param>
 		/// <returns></returns>
-		static StreamFile CreateStreamFile(FStatRecord MetaData)
+		static BucketId GetStreamBucketId(StreamId StreamId)
 		{
-			string FileType = MetaData.Type ?? MetaData.HeadType ?? String.Empty;
-			return new StreamFile(MetaData.DepotFile!, MetaData.FileSize, new FileContentId(Md5Hash.Parse(MetaData.Digest!), FileType.ToString()), MetaData.HeadRevision);
+			return new BucketId(StreamId.ToString());
 		}
-#endregion
+
+		#endregion
 	}
 }
diff --git a/Engine/Source/Programs/Horde/Horde.Build/Perforce/Impl/CommitTree.cs b/Engine/Source/Programs/Horde/Horde.Build/Perforce/Impl/CommitTree.cs
deleted file mode 100644
diff --git a/Engine/Source/Programs/Horde/Horde.Build/Startup.cs b/Engine/Source/Programs/Horde/Horde.Build/Startup.cs
--- a/Engine/Source/Programs/Horde/Horde.Build/Startup.cs
+++ b/Engine/Source/Programs/Horde/Horde.Build/Startup.cs
@@ -255,7 +255,9 @@ namespace HordeServer
 			// Settings used for configuring services
 			IConfigurationSection ConfigSection = Configuration.GetSection("Horde");
 			ServerSettings Settings = new ServerSettings();
-			ConfigSection.Bind(Settings);			
+			ConfigSection.Bind(Settings);
+
+			Services.Configure<CommitServiceOptions>(ConfigSection.GetSection("Replication"));
 
 			if (Settings.GlobalThreadPoolMinSize != null)
 			{
@@ -945,13 +947,16 @@ namespace HordeServer
 			return LogEventLevel.Information;
 		}
 
-		public static IServiceProvider CreateServiceProvider(IConfiguration Configuration)
+		public static void AddServices(IServiceCollection ServiceCollection, IConfiguration Configuration)
 		{
-			IServiceCollection ServiceCollection = new ServiceCollection();
-
 			Startup Startup = new Startup(Configuration);
 			Startup.ConfigureServices(ServiceCollection);
+		}
 
+		public static IServiceProvider CreateServiceProvider(IConfiguration Configuration)
+		{
+			IServiceCollection ServiceCollection = new ServiceCollection();
+			AddServices(ServiceCollection, Configuration);
 			return ServiceCollection.BuildServiceProvider();
 		}
 	}
