Horde: Narrow exit code error matching
#rnx
#skipci

diff --git a/Engine/Source/Programs/Horde/Horde.Build/Issues/Impl/Handlers/SystemicIssueHandler.cs b/Engine/Source/Programs/Horde/Horde.Build/Issues/Impl/Handlers/SystemicIssueHandler.cs
--- a/Engine/Source/Programs/Horde/Horde.Build/Issues/Impl/Handlers/SystemicIssueHandler.cs
+++ b/Engine/Source/Programs/Horde/Horde.Build/Issues/Impl/Handlers/SystemicIssueHandler.cs
@@ -40,7 +40,7 @@ namespace HordeServer.IssueHandlers.Impl
 		static HashSet<EventId?> KnownSystemic = new HashSet<EventId?> {  KnownLogEvents.Systemic, KnownLogEvents.Systemic_Xge, KnownLogEvents.Systemic_Xge_Standalone, 
 																		KnownLogEvents.Systemic_Xge_ServiceNotRunning, KnownLogEvents.Systemic_Xge_BuildFailed, 
 																		KnownLogEvents.Systemic_SlowDDC, KnownLogEvents.Systemic_Horde, KnownLogEvents.Systemic_Horde_ArtifactUpload,
-																		KnownLogEvents.Horde, KnownLogEvents.Horde_InvalidPreflight, KnownLogEvents.ExitCode};
+																		KnownLogEvents.Horde, KnownLogEvents.Horde_InvalidPreflight};
 
 		/// <summary>
 		/// Determines if the given event id matches
@@ -60,9 +60,22 @@ namespace HordeServer.IssueHandlers.Impl
 		/// <inheritdoc/>
 		public bool TryGetFingerprint(IJob Job, INode Node, ILogEventData EventData, [NotNullWhen(true)] out NewIssueFingerprint? Fingerprint)
 		{
-			if (!IsMatchingEventId(EventData.EventId))
+			Fingerprint = null;
+
+			if (EventData.EventId == KnownLogEvents.ExitCode)
 			{
-				Fingerprint = null;
+				for (int i = 0; i < EventData.Lines.Count; i++)
+				{
+					if (EventData.Lines[i].Message.Contains("AutomationTool exiting with ExitCode", StringComparison.InvariantCultureIgnoreCase))
+					{
+						Fingerprint = new NewIssueFingerprint(Type, new[] { Node.Name }, null);
+						return true;
+					}
+				}
+			}
+
+			if (!IsMatchingEventId(EventData.EventId))
+			{				
 				return false;
 			}
 
