Horde: Enable commit metadata mirroring by default (no content yet), and add support for subscribing to new commits being submitted.

diff --git a/Engine/Source/Programs/Horde/Horde.Build/Perforce/ICommit.cs b/Engine/Source/Programs/Horde/Horde.Build/Perforce/ICommit.cs
--- a/Engine/Source/Programs/Horde/Horde.Build/Perforce/ICommit.cs
+++ b/Engine/Source/Programs/Horde/Horde.Build/Perforce/ICommit.cs
@@ -14,18 +14,19 @@ using System.Threading.Tasks;
 
 namespace HordeServer.Commits
 {
+	using CommitId = ObjectId<ICommit>;
 	using StreamId = StringId<IStream>;
 	using UserId = ObjectId<IUser>;
 
 	/// <summary>
 	/// Stores metadata about a commit
 	/// </summary>
-	interface ICommit
+	public interface ICommit
 	{
 		/// <summary>
 		/// Unique id for this document
 		/// </summary>
-		public ObjectId Id { get; }
+		public CommitId Id { get; }
 
 		/// <summary>
 		/// The stream id
diff --git a/Engine/Source/Programs/Horde/Horde.Build/Perforce/ICommitCollection.cs b/Engine/Source/Programs/Horde/Horde.Build/Perforce/ICommitCollection.cs
--- a/Engine/Source/Programs/Horde/Horde.Build/Perforce/ICommitCollection.cs
+++ b/Engine/Source/Programs/Horde/Horde.Build/Perforce/ICommitCollection.cs
@@ -12,6 +12,7 @@ using HordeServer.Utilities;
 namespace HordeServer.Commits
 {
 	using StreamId = StringId<IStream>;
+	using CommitId = ObjectId<ICommit>;
 
 	/// <summary>
 	/// Stores a collection of commits
@@ -29,7 +30,7 @@ namespace HordeServer.Commits
 		/// Gets a single commit
 		/// </summary>
 		/// <param name="Id">Identifier for the commit</param>
-		Task<ICommit?> GetCommitAsync(ObjectId Id);
+		Task<ICommit?> GetCommitAsync(CommitId Id);
 
 		/// <summary>
 		/// Finds commits matching certain criteria
diff --git a/Engine/Source/Programs/Horde/Horde.Build/Perforce/ICommitService.cs b/Engine/Source/Programs/Horde/Horde.Build/Perforce/ICommitService.cs
--- a/Engine/Source/Programs/Horde/Horde.Build/Perforce/ICommitService.cs
+++ b/Engine/Source/Programs/Horde/Horde.Build/Perforce/ICommitService.cs
@@ -11,23 +11,16 @@ using HordeServer.Utilities;
 
 namespace HordeServer.Commits
 {
-	using StreamId = StringId<IStream>;
-
 	/// <summary>
 	/// Provides information about commits
 	/// </summary>
-	interface ICommitService
+	public interface ICommitService
 	{
 		/// <summary>
-		/// Finds commits matching a set of criteria
+		/// Registers a delegate to be called when a new commit is added
 		/// </summary>
-		/// <param name="StreamId">The stream to query</param>
-		/// <param name="MinChange">Minimum change to query</param>
-		/// <param name="MaxChange">Maximum change to query</param>
-		/// <param name="Paths">Paths to filter by</param>
-		/// <param name="Index">Index of the first result to return</param>
-		/// <param name="Count">Maximum number of results to return</param>
-		/// <returns>List of commit objects</returns>
-		Task<List<ICommit>> FindCommitsAsync(StreamId StreamId, int? MinChange = null, int? MaxChange = null, string[]? Paths = null, int? Index = null, int? Count = null);
+		/// <param name="OnAddCommit">Callback for a new commit being added</param>
+		/// <returns>Disposable handler.</returns>
+		IDisposable AddListener(Action<ICommit> OnAddCommit);
 	}
 }
diff --git a/Engine/Source/Programs/Horde/Horde.Build/Perforce/Impl/CommitCollection.cs b/Engine/Source/Programs/Horde/Horde.Build/Perforce/Impl/CommitCollection.cs
--- a/Engine/Source/Programs/Horde/Horde.Build/Perforce/Impl/CommitCollection.cs
+++ b/Engine/Source/Programs/Horde/Horde.Build/Perforce/Impl/CommitCollection.cs
@@ -15,6 +15,7 @@ using EpicGames.Horde.Storage;
 
 namespace HordeServer.Commits.Impl
 {
+	using CommitId = ObjectId<ICommit>;
 	using StreamId = StringId<IStream>;
 	using UserId = ObjectId<IUser>;
 
@@ -26,7 +27,7 @@ namespace HordeServer.Commits.Impl
 		class Commit : ICommit
 		{
 			[BsonIgnoreIfDefault]
-			public ObjectId Id { get; set; }
+			public CommitId Id { get; set; }
 
 			[BsonElement("s")]
 			public StreamId StreamId { get; set; }
@@ -106,7 +107,7 @@ namespace HordeServer.Commits.Impl
 		}
 
 		/// <inheritdoc/>
-		public async Task<ICommit?> GetCommitAsync(ObjectId Id)
+		public async Task<ICommit?> GetCommitAsync(CommitId Id)
 		{
 			return await Commits.Find(x => x.Id == Id).FirstOrDefaultAsync();
 		}
diff --git a/Engine/Source/Programs/Horde/Horde.Build/Perforce/Impl/CommitService.cs b/Engine/Source/Programs/Horde/Horde.Build/Perforce/Impl/CommitService.cs
--- a/Engine/Source/Programs/Horde/Horde.Build/Perforce/Impl/CommitService.cs
+++ b/Engine/Source/Programs/Horde/Horde.Build/Perforce/Impl/CommitService.cs
@@ -32,10 +32,12 @@ using HordeCommon;
 using EpicGames.Horde.Storage;
 using System.Text;
 using Microsoft.Extensions.Options;
+using Horde.Build.Utilities;
 
 namespace HordeServer.Commits.Impl
 {
 	using P4 = Perforce.P4;
+	using CommitId = ObjectId<ICommit>;
 	using StreamId = StringId<IStream>;
 	using IRef = EpicGames.Horde.Storage.IRef;
 
@@ -47,7 +49,7 @@ namespace HordeServer.Commits.Impl
 		/// <summary>
 		/// Whether to mirror commit metadata to the database
 		/// </summary>
-		public bool Metadata { get; set; }
+		public bool Metadata { get; set; } = true;
 
 		/// <summary>
 		/// Specific list of streams to enable commit mirroring for
@@ -125,6 +127,34 @@ namespace HordeServer.Commits.Impl
 			}
 		}
 
+		/// <summary>
+		/// A registered listener for new commits
+		/// </summary>
+		class ListenerInfo : IDisposable
+		{
+			List<ListenerInfo> Listeners;
+			public Action<ICommit> Callback;
+
+			public ListenerInfo(List<ListenerInfo> Listeners, Action<ICommit> Callback)
+			{
+				this.Listeners = Listeners;
+				this.Callback = Callback;
+
+				lock (Listeners)
+				{
+					Listeners.Add(this);
+				}
+			}
+
+			public void Dispose()
+			{
+				lock (Listeners)
+				{
+					Listeners.Remove(this);
+				}
+			}
+		}
+
 		public CommitServiceOptions Options { get; }
 
 		// Redis
@@ -166,6 +196,7 @@ namespace HordeServer.Commits.Impl
 		/// <inheritdoc/>
 		public async Task StartAsync(CancellationToken CancellationToken)
 		{
+			await StartNotificationsAsync();
 			if (Options.Metadata)
 			{
 				await StartMetadataReplicationAsync();
@@ -187,8 +218,57 @@ namespace HordeServer.Commits.Impl
 			{
 				await StopMetadataReplicationAsync();
 			}
+			await StopNotificationsAsync();
 		}
 
+		#region Notifications
+
+		List<ListenerInfo> Listeners = new List<ListenerInfo>();
+		RedisChannel<CommitId> CommitNotifyChannel { get; } = new RedisChannel<CommitId>("commits/notify");
+		RedisChannelSubscription<CommitId>? CommitNotifySubscription;
+
+		async Task StartNotificationsAsync()
+		{
+			ISubscriber Subscriber = Redis.Multiplexer.GetSubscriber();
+			CommitNotifySubscription = await Subscriber.SubscribeAsync(CommitNotifyChannel, DispatchCommitNotification);
+		}
+
+		async Task StopNotificationsAsync()
+		{
+			if (CommitNotifySubscription != null)
+			{
+				await CommitNotifySubscription.DisposeAsync();
+				CommitNotifySubscription = null;
+			}
+		}
+
+		async void DispatchCommitNotification(RedisChannel<CommitId> Channel, CommitId CommitId)
+		{
+			ICommit? Commit = await CommitCollection.GetCommitAsync(CommitId);
+			if (Commit != null)
+			{
+				lock (Listeners)
+				{
+					foreach (ListenerInfo Listener in Listeners)
+					{
+						Listener.Callback(Commit);
+					}
+				}
+			}
+		}
+
+		async Task NotifyListeners(CommitId CommitId)
+		{
+			await Redis.PublishAsync(CommitNotifyChannel, CommitId);
+		}
+
+		public IDisposable AddListener(Action<ICommit> OnAddCommit)
+		{
+			return new ListenerInfo(Listeners, OnAddCommit);
+		}
+
+		#endregion
+
 		#region Metadata updates
 
 		ITicker UpdateMetadataTicker;
@@ -256,6 +336,7 @@ namespace HordeServer.Commits.Impl
 			await foreach (NewCommit NewCommit in FindCommitsForClusterAsync(ClusterName, StreamToFirstChange))
 			{
 				ICommit Commit = await CommitCollection.AddOrReplaceAsync(NewCommit);
+				await NotifyListeners(Commit.Id);
 
 				RedisList<int> StreamCommitsKey = RedisStreamChanges(Commit.StreamId);
 				await StreamCommitsKey.RightPushAsync(Commit.Change);
diff --git a/Engine/Source/Programs/Horde/Horde.Build/Startup.cs b/Engine/Source/Programs/Horde/Horde.Build/Startup.cs
--- a/Engine/Source/Programs/Horde/Horde.Build/Startup.cs
+++ b/Engine/Source/Programs/Horde/Horde.Build/Startup.cs
@@ -375,6 +375,7 @@ namespace HordeServer
 				return new NoOpDogStatsd();
 			});
 			Services.AddSingleton<CommitService>();
+			Services.AddSingleton<ICommitService>(SP => SP.GetRequiredService<CommitService>());
 			Services.AddSingleton<IClock, Clock>();
 			Services.AddSingleton<IDowntimeService, DowntimeService>();
 			Services.AddSingleton<IIssueService, IssueService>();
