Horde: Add commands to analyze and extract contents of tree packs, and add support for creating "pointer files" in packs to Perforce depot/revision numbers rather than actually including the file contents themselves.

diff --git a/Engine/Source/Programs/Horde/Horde.Build/Commands/Perforce/ReplicateCommand.cs b/Engine/Source/Programs/Horde/Horde.Build/Commands/Perforce/ReplicateCommand.cs
--- a/Engine/Source/Programs/Horde/Horde.Build/Commands/Perforce/ReplicateCommand.cs
+++ b/Engine/Source/Programs/Horde/Horde.Build/Commands/Perforce/ReplicateCommand.cs
@@ -55,14 +55,9 @@ namespace HordeServer.Commands
 			IServiceCollection Services = new ServiceCollection();
 			Startup.AddServices(Services, Configuration);
 
-			OutputDir ??= DirectoryReference.Combine(Program.DataDir, "p4");
+			OutputDir ??= DirectoryReference.Combine(Program.DataDir, "Storage");
 			Logger.LogInformation("Writing output to {OutputDir}", OutputDir);
-
-			DirectoryReference BlobsDir = DirectoryReference.Combine(OutputDir, "blobs");
-			Services.AddSingleton<IStorageClient, FileStorageClient>(SP => new FileStorageClient(BlobsDir, Logger));
-
-			DirectoryReference RootsDir = DirectoryReference.Combine(OutputDir, "roots");
-			DirectoryReference.CreateDirectory(RootsDir);
+			Services.AddSingleton<IStorageClient, FileStorageClient>(SP => new FileStorageClient(OutputDir, Logger));
 
 			IServiceProvider ServiceProvider = Services.BuildServiceProvider();
 			CommitService CommitService = ServiceProvider.GetRequiredService<CommitService>();
@@ -82,15 +77,7 @@ namespace HordeServer.Commands
 			CommitTree? BaseTree = null; 
 			if (BaseChange != 0)
 			{
-				ICommit? BaseCommit = await CommitCollection.GetCommitAsync(Stream.Id, BaseChange);
-				if (BaseCommit == null)
-				{
-					throw new FatalErrorException($"Unable to find base commit {BaseChange}");
-				}
-
-				FileReference RootFile = FileReference.Combine(RootsDir, $"{BaseCommit.Change}.ref");
-				byte[] RootData = await FileReference.ReadAllBytesAsync(RootFile);
-				BaseTree = new CommitTree(BaseCommit.Change, TreePackObject.Parse(RootData));
+				BaseTree = await CommitService.ReadTreeAsync(Stream.Id, BaseChange);
 			}
 
 			await foreach (NewCommit NewCommit in CommitService.FindCommitsForClusterAsync(Stream.ClusterName, StreamToFirstChange).Take(Count))
@@ -102,11 +89,7 @@ namespace HordeServer.Commands
 				if (Content)
 				{
 					BaseTree = await CommitService.FindCommitTreeAsync(Stream, NewCommit.Change, BaseTree);
-
-					FileReference RootFile = FileReference.Combine(RootsDir, $"{NewCommit.Change}.ref");
-					ReadOnlyMemory<byte> RootData = BaseTree.RootObject.ToCbObject().GetView();
-					await FileReference.WriteAllBytesAsync(RootFile, RootData.ToArray());
-					Logger.LogInformation("Written root to {RootFile} ({Size:n0} bytes)", RootFile, RootFile.ToFileInfo().Length);
+					await CommitService.WriteTreeAsync(Stream.Id, BaseTree);
 				}
 			}
 
diff --git a/Engine/Source/Programs/Horde/Horde.Build/Commands/Tree/AnalyzeCommand.cs b/Engine/Source/Programs/Horde/Horde.Build/Commands/Tree/AnalyzeCommand.cs
new file mode 100644
--- /dev/null
+++ b/Engine/Source/Programs/Horde/Horde.Build/Commands/Tree/AnalyzeCommand.cs
@@ -0,0 +1,66 @@
+// Copyright Epic Games, Inc. All Rights Reserved.
+
+using EpicGames.Core;
+using EpicGames.Horde.Storage;
+using EpicGames.Horde.Storage.Impl;
+using EpicGames.Serialization;
+using HordeServer.Collections;
+using HordeServer.Commits;
+using HordeServer.Commits.Impl;
+using HordeServer.Controllers;
+using HordeServer.Models;
+using HordeServer.Utilities;
+using Microsoft.Extensions.Configuration;
+using Microsoft.Extensions.DependencyInjection;
+using Microsoft.Extensions.Logging;
+using System;
+using System.Collections.Generic;
+using System.Linq;
+using System.Threading.Tasks;
+
+namespace HordeServer.Commands
+{
+	using StreamId = StringId<IStream>;
+
+	[Command("tree", "analyze", "Indexes a tree and creates a table of contents for each blob describing the contents")]
+	class AnalyzeCommand : Command
+	{
+		IConfiguration Configuration;
+
+		[CommandLine("-Namespace=", Required = true)]
+		public NamespaceId NamespaceId { get; set; }
+
+		[CommandLine("-Bucket=", Required = true)]
+		public BucketId BucketId { get; set; }
+
+		[CommandLine("-Ref=", Required = true)]
+		public RefId RefId { get; set; }
+
+		public AnalyzeCommand(IConfiguration Configuration)
+		{
+			this.Configuration = Configuration;
+		}
+
+		public override async Task<int> ExecuteAsync(ILogger Logger)
+		{
+			IServiceCollection Services = new ServiceCollection();
+			Startup.AddServices(Services, Configuration);
+
+			DirectoryReference OutputDir = DirectoryReference.Combine(Program.DataDir, "Storage");
+			Logger.LogInformation("Writing output to {OutputDir}", OutputDir);
+
+			Services.AddSingleton<IStorageClient, FileStorageClient>(SP => new FileStorageClient(OutputDir, Logger));
+
+			IServiceProvider ServiceProvider = Services.BuildServiceProvider();
+			IStorageClient StorageClient = ServiceProvider.GetRequiredService<IStorageClient>();
+
+			IRef Ref = await StorageClient.GetRefAsync(NamespaceId, BucketId, RefId);
+
+			TreePack Pack = new TreePack(StorageClient, NamespaceId);
+			IoHash RootHash = Pack.AddRootObject(Ref).GetRootHash();
+			await Pack.WriteTreeSummaryAsync(RootHash, DirectoryReference.Combine(OutputDir, NamespaceId.ToString()), Logger);
+
+			return 0;
+		}
+	}
+}
diff --git a/Engine/Source/Programs/Horde/Horde.Build/Commands/Tree/ExtractCommand.cs b/Engine/Source/Programs/Horde/Horde.Build/Commands/Tree/ExtractCommand.cs
new file mode 100644
--- /dev/null
+++ b/Engine/Source/Programs/Horde/Horde.Build/Commands/Tree/ExtractCommand.cs
@@ -0,0 +1,90 @@
+// Copyright Epic Games, Inc. All Rights Reserved.
+
+using EpicGames.Core;
+using EpicGames.Horde.Storage;
+using EpicGames.Horde.Storage.Impl;
+using EpicGames.Serialization;
+using HordeServer.Collections;
+using HordeServer.Commits;
+using HordeServer.Commits.Impl;
+using HordeServer.Controllers;
+using HordeServer.Models;
+using HordeServer.Utilities;
+using Microsoft.Extensions.Configuration;
+using Microsoft.Extensions.DependencyInjection;
+using Microsoft.Extensions.Logging;
+using System;
+using System.Collections.Generic;
+using System.Linq;
+using System.Threading.Tasks;
+
+namespace HordeServer.Commands
+{
+	[Command("tree", "extract", "Extracts data from a tree to the local hard drive")]
+	class ExtractCommand : Command
+	{
+		IConfiguration Configuration;
+
+		[CommandLine("-Namespace=", Required = true)]
+		public NamespaceId NamespaceId { get; set; }
+
+		[CommandLine("-Bucket=", Required = true)]
+		public BucketId BucketId { get; set; }
+
+		[CommandLine("-Ref=", Required = true)]
+		public RefId RefId { get; set; }
+
+		[CommandLine("-Output=")]
+		public DirectoryReference? OutputDir { get; set; }
+
+		public ExtractCommand(IConfiguration Configuration)
+		{
+			this.Configuration = Configuration;
+		}
+
+		public override async Task<int> ExecuteAsync(ILogger Logger)
+		{
+			IServiceCollection Services = new ServiceCollection();
+			Startup.AddServices(Services, Configuration);
+
+			DirectoryReference StorageDir = DirectoryReference.Combine(Program.DataDir, "Storage");
+			Services.AddSingleton<IStorageClient, FileStorageClient>(SP => new FileStorageClient(StorageDir, Logger));
+
+			OutputDir ??= DirectoryReference.Combine(Program.AppDir, "Output");
+			Logger.LogInformation("Writing output to {OutputDir}", OutputDir);
+
+			IServiceProvider ServiceProvider = Services.BuildServiceProvider();
+			IStorageClient StorageClient = ServiceProvider.GetRequiredService<IStorageClient>();
+
+			IRef Ref = await StorageClient.GetRefAsync(NamespaceId, BucketId, RefId);
+
+			TreePack Pack = new TreePack(StorageClient, NamespaceId);
+			IoHash RootHash = Pack.AddRootObject(Ref).GetRootHash();
+
+			TreePackDirNode RootNode = TreePackDirNode.Parse(await Pack.GetDataAsync(RootHash));
+			await ExtractDataAsync(Pack, RootNode, OutputDir);
+
+			return 0;
+		}
+
+		public async Task ExtractDataAsync(TreePack Pack, TreePackDirNode DirNode, DirectoryReference OutputDir)
+		{
+			DirectoryReference.CreateDirectory(OutputDir);
+			foreach (TreePackDirEntry Entry in DirNode.Entries)
+			{
+				if ((Entry.Flags & TreePackDirEntryFlags.Directory) != 0)
+				{
+					DirectoryReference SubOutputDir = DirectoryReference.Combine(OutputDir, Entry.Name.ToString());
+					TreePackDirNode SubDirNode = TreePackDirNode.Parse(await Pack.GetDataAsync(Entry.Hash));
+					await ExtractDataAsync(Pack, SubDirNode, SubOutputDir);
+				}
+				else if ((Entry.Flags & TreePackDirEntryFlags.File) != 0)
+				{
+					FileReference File = FileReference.Combine(OutputDir, Entry.Name.ToString());
+					ReadOnlyMemory<byte> Data = await Pack.GetDataAsync(Entry.Hash);
+					await FileReference.WriteAllBytesAsync(File, Data.ToArray());
+				}
+			}
+		}
+	}
+}
diff --git a/Engine/Source/Programs/Horde/Horde.Build/Perforce/IPerforceService.cs b/Engine/Source/Programs/Horde/Horde.Build/Perforce/IPerforceService.cs
--- a/Engine/Source/Programs/Horde/Horde.Build/Perforce/IPerforceService.cs
+++ b/Engine/Source/Programs/Horde/Horde.Build/Perforce/IPerforceService.cs
@@ -117,15 +117,6 @@ namespace HordeServer.Services
 		/// <returns></returns>
 		public Task<IStreamView> GetStreamViewAsync(string ClusterName, string StreamName);
 
-		/// <summary>
-		/// Gets the state of all files in a stream
-		/// </summary>
-		/// <param name="ClusterName">Name of the Perforce cluster</param>
-		/// <param name="StreamName">The stream name</param>
-		/// <param name="Change">The changelist to query</param>
-		/// <returns></returns>
-		public Task<List<ChangeFile>> GetStreamSnapshotAsync(string ClusterName, string StreamName, int Change);
-
 		/// <summary>
 		/// Create a new changelist by submitting the given file
 		/// </summary>
@@ -204,14 +195,6 @@ namespace HordeServer.Services
 		/// <returns>Commit details</returns>
 		public Task<List<ChangeDetails>> GetChangeDetailsAsync(string ClusterName, string StreamName, IReadOnlyList<int> ChangeNumbers, string? ImpersonateUser);
 
-		/// <summary>
-		/// Create a ticket as the specified user
-		/// </summary>
-		/// <param name="ClusterName">Name of the Perforce cluster</param>
-		/// <param name="ImpersonateUser">Name of the user to impersonate</param>
-		/// <returns>Commit details</returns>
-		public Task<string> CreateTicket(string ClusterName, string ImpersonateUser);
-
 		/// <summary>
 		/// Gets the latest changes for a set of depot paths
 		/// </summary>
diff --git a/Engine/Source/Programs/Horde/Horde.Build/Perforce/Impl/CommitService.cs b/Engine/Source/Programs/Horde/Horde.Build/Perforce/Impl/CommitService.cs
--- a/Engine/Source/Programs/Horde/Horde.Build/Perforce/Impl/CommitService.cs
+++ b/Engine/Source/Programs/Horde/Horde.Build/Perforce/Impl/CommitService.cs
@@ -66,6 +66,11 @@ namespace HordeServer.Commits.Impl
 		/// </summary>
 		public List<StreamId> ContentStreams { get; set; } = new List<StreamId>();
 
+		/// <summary>
+		/// Write Perforce stub files rather than including the actual content.
+		/// </summary>
+		public bool WriteStubFiles { get; set; } = true;
+
 		/// <summary>
 		/// Namespace to store content replicated from Perforce
 		/// </summary>
@@ -82,25 +87,41 @@ namespace HordeServer.Commits.Impl
 	/// </summary>
 	class CommitTree
 	{
+		/// <summary>
+		/// The stream id
+		/// </summary>
+		public StreamId StreamId { get; }
+
 		/// <summary>
 		/// The changelist number
 		/// </summary>
-		public int Change;
+		public int Change { get; }
 
 		/// <summary>
 		/// Root object describing the tree contents
 		/// </summary>
-		public TreePackObject RootObject;
+		public TreePackObject RootObject { get; }
 
 		/// <summary>
 		/// Constructor
 		/// </summary>
-		public CommitTree(int Change, TreePackObject RootObject)
+		public CommitTree(StreamId StreamId, int Change, TreePackObject RootObject)
 		{
+			this.StreamId = StreamId;
 			this.Change = Change;
 			this.RootObject = RootObject;
 		}
 
+		/// <summary>
+		/// The ref id for storing this tree
+		/// </summary>
+		public RefId RefId => GetRefId(Change);
+
+		/// <summary>
+		/// The ref id for storing this tree
+		/// </summary>
+		public static RefId GetRefId(int Change) => new RefId(IoHash.Compute(Encoding.UTF8.GetBytes($"{Change}")));
+
 		/// <summary>
 		/// Gets the directory at the root of this tree
 		/// </summary>
@@ -817,7 +838,20 @@ namespace HordeServer.Commits.Impl
 				return null;
 			}
 
-			return new CommitTree(Commit.Change, TreePackObject.Parse(Ref));
+			return new CommitTree(Commit.StreamId, Commit.Change, TreePackObject.Parse(Ref));
+		}
+
+		public async Task<CommitTree> ReadTreeAsync(StreamId StreamId, int Change)
+		{
+			RefId TreeRefId = new RefId(IoHash.Compute(Encoding.UTF8.GetBytes($"{Change}")));
+			IRef Ref = await StorageClient.GetRefAsync(Options.NamespaceId, new BucketId(StreamId.ToString()), TreeRefId);
+			return new CommitTree(StreamId, Change, TreePackObject.Parse(Ref));
+		}
+
+		public async Task WriteTreeAsync(StreamId StreamId, CommitTree Tree)
+		{
+			CbObject Object = Tree.RootObject.ToCbObject();
+			await StorageClient.SetRefAsync(Options.NamespaceId, new BucketId(StreamId.ToString()), Tree.RefId, Object);
 		}
 
 		/// <summary>
@@ -839,7 +873,7 @@ namespace HordeServer.Commits.Impl
 				}
 				else
 				{
-					BaseTree = new CommitTree(BaseCommit.Change, TreePackObject.Parse(await StorageClient.GetRefAsync(Options.NamespaceId, GetStreamBucketId(Stream.Id), BaseCommit.TreeRefId.Value)));
+					BaseTree = new CommitTree(BaseCommit.StreamId, BaseCommit.Change, TreePackObject.Parse(await StorageClient.GetRefAsync(Options.NamespaceId, GetStreamBucketId(Stream.Id), BaseCommit.TreeRefId.Value)));
 				}
 			}
 
@@ -887,34 +921,59 @@ namespace HordeServer.Commits.Impl
 			Logger.LogInformation("Updating client {Client} to changelist {Change}", ClientInfo.Client.Name, Change);
 			ClientInfo.Change = -1;
 
-			Dictionary<int, (string Path, TreePackFileWriter Writer)> Files = new Dictionary<int, (string, TreePackFileWriter)>();
-			await foreach (PerforceResponse Response in Perforce.StreamCommandAsync("sync", Array.Empty<string>(), new string[] { $"//...@{Change}" }, null, typeof(SyncRecord), true, default))
+			if (Options.WriteStubFiles)
 			{
-				PerforceIo? Io = Response.Io;
-				if (Io != null)
+				await foreach (PerforceResponse<SyncRecord> Record in Perforce.StreamCommandAsync<SyncRecord>("sync", new[] { "-k", $"//...@{Change}" }, null, default))
 				{
-					if (Io.Command == PerforceIoCommand.Open)
-					{
-						string Path = GetClientRelativePath(Io.Payload, ClientInfo.Client.Root);
-						TreePackFileWriter FileWriter = new TreePackFileWriter(ClientInfo.TreePack);
-						Files[Io.File] = (Path, FileWriter);
-					}
-					else if (Io.Command == PerforceIoCommand.Write)
-					{
-						TreePackFileWriter FileWriter = Files[Io.File].Writer;
-						await FileWriter.WriteAsync(Io.Payload, false);
-					}
-					else if (Io.Command == PerforceIoCommand.Close)
+					SyncRecord SyncRecord = Record.Data;
+					if (!SyncRecord.Path.StartsWith(ClientInfo.Client.Root, StringComparison.Ordinal))
 					{
-						(string Path, TreePackFileWriter FileWriter) = Files[Io.File];
-						IoHash FileHash = await FileWriter.FinalizeAsync();
-						await DirWriter.FindOrAddFileByPathAsync(Path, TreePackDirEntryFlags.File, FileHash, Sha1Hash.Zero);
-						Files.Remove(Io.File);
+						throw new ArgumentException($"Unable to make path {ClientInfo.Client.Root} relative to client root {ClientInfo.Client.Root}");
 					}
-					else if (Io.Command == PerforceIoCommand.Unlink)
+
+					string RelativePath = SyncRecord.Path.Substring(ClientInfo.Client.Root.Length).Replace('\\', '/').TrimStart('/');
+
+					string DepotPath = $"{SyncRecord.DepotFile}#{SyncRecord.Revision}";
+					byte[] Data = new byte[Encoding.UTF8.GetMaxByteCount(DepotPath.Length) + 1];
+					Data[0] = (byte)TreePackNodeType.Binary;
+					int Length = Encoding.UTF8.GetBytes(DepotPath, Data.AsSpan(1));
+					ReadOnlyMemory<byte> EncodedData = Data.AsMemory(0, Length + 1);
+
+					IoHash Hash = await ClientInfo.TreePack.AddNodeAsync(EncodedData);
+					await DirWriter.FindOrAddFileByPathAsync(RelativePath, TreePackDirEntryFlags.File | TreePackDirEntryFlags.PerforceDepotPathAndRevision, Hash, Sha1Hash.Zero);
+				}
+			}
+			else
+			{
+				Dictionary<int, (string Path, TreePackFileWriter Writer)> Files = new Dictionary<int, (string, TreePackFileWriter)>();
+				await foreach (PerforceResponse Response in Perforce.StreamCommandAsync("sync", Array.Empty<string>(), new string[] { $"//...@{Change}" }, null, typeof(SyncRecord), true, default))
+				{
+					PerforceIo? Io = Response.Io;
+					if (Io != null)
 					{
-						string Path = GetClientRelativePath(Io.Payload, ClientInfo.Client.Root);
-						await DirWriter.RemoveFileByPathAsync(Path);
+						if (Io.Command == PerforceIoCommand.Open)
+						{
+							string Path = GetClientRelativePath(Io.Payload, ClientInfo.Client.Root);
+							TreePackFileWriter FileWriter = new TreePackFileWriter(ClientInfo.TreePack);
+							Files[Io.File] = (Path, FileWriter);
+						}
+						else if (Io.Command == PerforceIoCommand.Write)
+						{
+							TreePackFileWriter FileWriter = Files[Io.File].Writer;
+							await FileWriter.WriteAsync(Io.Payload, false);
+						}
+						else if (Io.Command == PerforceIoCommand.Close)
+						{
+							(string Path, TreePackFileWriter FileWriter) = Files[Io.File];
+							IoHash FileHash = await FileWriter.FinalizeAsync();
+							await DirWriter.FindOrAddFileByPathAsync(Path, TreePackDirEntryFlags.File, FileHash, Sha1Hash.Zero);
+							Files.Remove(Io.File);
+						}
+						else if (Io.Command == PerforceIoCommand.Unlink)
+						{
+							string Path = GetClientRelativePath(Io.Payload, ClientInfo.Client.Root);
+							await DirWriter.RemoveFileByPathAsync(Path);
+						}
 					}
 				}
 			}
@@ -924,7 +983,7 @@ namespace HordeServer.Commits.Impl
 			// Return the new root object
 			IoHash RootHash = await DirWriter.FlushAsync();
 			TreePackObject RootObject = await ClientInfo.TreePack.FlushAsync(RootHash, DateTime.UtcNow);
-			return new CommitTree(Change, RootObject);
+			return new CommitTree(Stream.Id, Change, RootObject);
 		}
 
 		async Task FlushWorkspaceAsync(ReplicationClient ClientInfo, IPerforceConnection Perforce, int Change)
diff --git a/Engine/Source/Programs/Horde/Horde.Build/Perforce/Impl/PerforceService.cs b/Engine/Source/Programs/Horde/Horde.Build/Perforce/Impl/PerforceService.cs
--- a/Engine/Source/Programs/Horde/Horde.Build/Perforce/Impl/PerforceService.cs
+++ b/Engine/Source/Programs/Horde/Horde.Build/Perforce/Impl/PerforceService.cs
@@ -785,43 +785,6 @@ namespace HordeServer.Services
 			return new ChangeFile(RelativePath, MetaData.Path.Path, Revision, MetaData.Size, Digest, MetaData.Type.ToString());
 		}
 
-		/// <inheritdoc/>
-		public async Task<List<ChangeFile>> GetStreamSnapshotAsync(string ClusterName, string StreamName, int Change)
-		{
-			using IScope Scope = GlobalTracer.Instance.BuildSpan("PerforceService.GetStreamSnapshotAsync").StartActive();
-			Scope.Span.SetTag("ClusterName", ClusterName);
-			Scope.Span.SetTag("StreamName", StreamName);
-			Scope.Span.SetTag("Change", Change);
-			
-			PerforceCluster Cluster = await GetClusterAsync(ClusterName);
-			using (P4.Repository Repository = await GetConnection(Cluster, Stream: StreamName))
-			{
-				P4.FileSpec FileSpec = new P4.FileSpec(new P4.DepotPath($"//UE5/Main/Engine/Build/..."), new P4.ChangelistIdVersion(Change));
-				Repository.Connection.Client.SyncFiles(new P4.SyncFilesCmdOptions(P4.SyncFilesCmdFlags.ServerOnly | P4.SyncFilesCmdFlags.Quiet), FileSpec);
-
-				string ClientPrefix = $"//{Repository.Connection.Client.Name}/";
-				P4.FileSpec StatSpec = new P4.FileSpec(new P4.ClientPath($"{ClientPrefix}..."));
-
-				IList<P4.FileMetaData> Files = Repository.GetFileMetaData(new P4.GetFileMetaDataCmdOptions(P4.GetFileMetadataCmdFlags.Synced | P4.GetFileMetadataCmdFlags.LocalPath | P4.GetFileMetadataCmdFlags.FileSize, null, null, 0, null, null, null), StatSpec);
-
-				List<ChangeFile> Results = new List<ChangeFile>();
-				foreach (P4.FileMetaData File in Files)
-				{
-					string RelativePath = File.ClientPath.Path;
-					if (RelativePath.StartsWith(ClientPrefix, StringComparison.OrdinalIgnoreCase))
-					{
-						RelativePath = RelativePath.Substring(ClientPrefix.Length);
-					}
-					else
-					{
-						throw new Exception($"Client path does not start with client name: {RelativePath}");
-					}
-					Results.Add(CreateChangeFile(RelativePath, File));
-				}
-				return Results;
-			}
-		}
-
 		/// <inheritdoc/>
 		public async Task<List<ChangeSummary>> GetChangesAsync(string ClusterName, int? MinChange, int? MaxChange, int MaxResults)
 		{
@@ -1075,25 +1038,6 @@ namespace HordeServer.Services
 			}
 		}
 
-		/// <inheritdoc />
-		public async Task<string> CreateTicket(string ClusterName, string ImpersonateUser)
-		{
-			using IScope Scope = GlobalTracer.Instance.BuildSpan("PerforceService.CreateTicket").StartActive();
-			Scope.Span.SetTag("ClusterName", ClusterName);
-			Scope.Span.SetTag("ImpersonateUser", ImpersonateUser);
-			
-			PerforceCluster Cluster = await GetClusterAsync(ClusterName);
-
-			CachedTicketInfo Credential = await GetImpersonateCredential(Cluster, ImpersonateUser);
-
-			if (Credential == null)
-			{
-				throw new Exception($"Unable to get ticket for user {ImpersonateUser}");
-			}
-
-			return Credential.Ticket.Ticket;
-		}
-
 		/// <inheritdoc/>
 		public async Task<List<FileSummary>> FindFilesAsync(string ClusterName, IEnumerable<string> Paths)
 		{
diff --git a/Engine/Source/Programs/Shared/EpicGames.Horde/Storage/Impl/FileStorageClient.cs b/Engine/Source/Programs/Shared/EpicGames.Horde/Storage/Impl/FileStorageClient.cs
--- a/Engine/Source/Programs/Shared/EpicGames.Horde/Storage/Impl/FileStorageClient.cs
+++ b/Engine/Source/Programs/Shared/EpicGames.Horde/Storage/Impl/FileStorageClient.cs
@@ -52,7 +52,7 @@ namespace EpicGames.Horde.Storage.Impl
 
 		FileReference GetBlobFile(NamespaceId NamespaceId, IoHash Hash)
 		{
-			return FileReference.Combine(BaseDir, $"{NamespaceId}-{Hash}.blob");
+			return FileReference.Combine(BaseDir, NamespaceId.ToString(), $"{Hash}.blob");
 		}
 
 		/// <inheritdoc/>
@@ -67,6 +67,8 @@ namespace EpicGames.Horde.Storage.Impl
 		public async Task WriteBlobAsync(NamespaceId NamespaceId, IoHash Hash, Stream Stream, CancellationToken CancellationToken = default)
 		{
 			FileReference File = GetBlobFile(NamespaceId, Hash);
+			DirectoryReference.CreateDirectory(File.Directory);
+
 			Logger.LogInformation("Writing {File} ({Size} bytes)", File, Stream.Length);
 			using (Stream OutputStream = FileReference.Open(File, FileMode.Create, FileAccess.Write, FileShare.Read))
 			{
@@ -94,7 +96,7 @@ namespace EpicGames.Horde.Storage.Impl
 
 		FileReference GetRefFile(NamespaceId NamespaceId, BucketId BucketId, RefId RefId)
 		{
-			return FileReference.Combine(BaseDir, $"{NamespaceId}-{BucketId}-{RefId}.ref");
+			return FileReference.Combine(BaseDir, NamespaceId.ToString(), BucketId.ToString(), $"{RefId}.ref");
 		}
 
 		/// <inheritdoc/>
@@ -122,6 +124,8 @@ namespace EpicGames.Horde.Storage.Impl
 		public async Task<List<IoHash>> TrySetRefAsync(NamespaceId NamespaceId, BucketId BucketId, RefId RefId, CbObject Value, CancellationToken CancellationToken = default)
 		{
 			FileReference File = GetRefFile(NamespaceId, BucketId, RefId);
+			DirectoryReference.CreateDirectory(File.Directory);
+
 			Logger.LogInformation("Writing {File} ({Size} bytes)", File, Value.GetView().Length);
 			await FileReference.WriteAllBytesAsync(File, Value.GetView().ToArray(), CancellationToken);
 			return new List<IoHash>();
diff --git a/Engine/Source/Programs/Shared/EpicGames.Horde/Storage/TreePack.cs b/Engine/Source/Programs/Shared/EpicGames.Horde/Storage/TreePack.cs
--- a/Engine/Source/Programs/Shared/EpicGames.Horde/Storage/TreePack.cs
+++ b/Engine/Source/Programs/Shared/EpicGames.Horde/Storage/TreePack.cs
@@ -71,6 +71,12 @@ namespace EpicGames.Horde.Storage
 		/// <returns></returns>
 		public static TreePackObject Parse(ReadOnlyMemory<byte> Data) => CbSerializer.Deserialize<TreePackObject>(Data);
 
+		/// <summary>
+		/// Serialize this object to compact binary
+		/// </summary>
+		/// <returns></returns>
+		public CbObject ToCbObject() => CbSerializer.Serialize(this);
+
 		/// <summary>
 		/// Gets the root hash for this object
 		/// </summary>
@@ -284,6 +290,11 @@ namespace EpicGames.Horde.Storage
 		/// The attached entry includes a Git SHA1 of the corresponding blob/tree contents.
 		/// </summary>
 		HasGitSha1 = 32,
+
+		/// <summary>
+		/// The data for this entry is a Perforce depot path and revision rather than the actual file contents.
+		/// </summary>
+		PerforceDepotPathAndRevision = 64,
 	}
 
 	/// <summary>
@@ -966,7 +977,8 @@ namespace EpicGames.Horde.Storage
 		{
 			if (VisitedBlobs.Add(Blob))
 			{
-				foreach (BlobInfo ReferencedBlob in Blob.ReferencedBy)
+				BlobInfo[] ReferencedBy = Blob.ReferencedBy.ToArray();
+				foreach (BlobInfo ReferencedBlob in ReferencedBy)
 				{
 					ExpandReferences(ReferencedBlob, VisitedBlobs);
 					Blob.ReferencedBy.UnionWith(ReferencedBlob.ReferencedBy);
@@ -1119,6 +1131,7 @@ namespace EpicGames.Horde.Storage
 				Node.CompressedData = ReadOnlyMemory<byte>.Empty;
 			}
 
+			Blobs[Blob.Hash] = Blob;
 			return Blob;
 		}
 
@@ -1264,6 +1277,42 @@ namespace EpicGames.Horde.Storage
 			return Probability * (DownloadInit + (Size / DownloadRate));
 		}
 
+		public async Task WriteTreeSummaryAsync(IoHash RootHash, DirectoryReference Directory, ILogger Logger)
+		{
+			DirectoryReference.CreateDirectory(Directory);
+
+			HashSet<NodeInfo> FoundNodes = new HashSet<NodeInfo>();
+			await FindNodesAsync(RootHash, FoundNodes);
+
+			HashSet<BlobInfo> FoundBlobs = new HashSet<BlobInfo>();
+			foreach (NodeInfo FoundNode in FoundNodes)
+			{
+				BlobInfo? FoundBlob = FoundNode.Blob;
+				if (FoundBlob != null && FoundBlobs.Add(FoundBlob) && !FoundBlob.IsLeaf)
+				{
+					TreePackObject Object = TreePackObject.Parse(FoundBlob.Data);
+
+					FileReference OutputFile = FileReference.Combine(Directory, $"{FoundBlob.Hash}.yml");
+					Logger.LogInformation("Writing summary: {File}", OutputFile);
+
+					WriteSummary(OutputFile, Object);
+				}
+			}
+		}
+
+		async Task FindNodesAsync(IoHash Hash, HashSet<NodeInfo> FoundNodes)
+		{
+			NodeInfo Node = Nodes[Hash];
+			if (FoundNodes.Add(Node) && (Node.Blob == null || !Node.Blob.IsLeaf))
+			{
+				ReadOnlyMemory<byte> Data = await GetDataAsync(Node.Hash);
+				foreach (IoHash RefHash in GetReferences(Data))
+				{
+					await FindNodesAsync(RefHash, FoundNodes);
+				}
+			}
+		}
+
 		public static void WriteSummary(FileReference File, TreePackObject Object)
 		{
 			Dictionary<IoHash, string> HashToLocator = new Dictionary<IoHash, string>();
diff --git a/Engine/Source/Programs/Shared/EpicGames.Horde/Storage/TreePackDirWriter.cs b/Engine/Source/Programs/Shared/EpicGames.Horde/Storage/TreePackDirWriter.cs
--- a/Engine/Source/Programs/Shared/EpicGames.Horde/Storage/TreePackDirWriter.cs
+++ b/Engine/Source/Programs/Shared/EpicGames.Horde/Storage/TreePackDirWriter.cs
@@ -103,11 +103,11 @@ namespace EpicGames.Horde.Storage
 		}
 
 		/// <summary>
-		/// Adds a directory to the current tree level with the given name
+		/// Gets a directory writer to the current tree level for the given directory name
 		/// </summary>
 		/// <param name="Name"></param>
-		/// <returns></returns>
-		public async ValueTask<TreePackDirWriter> FindOrAddDirAsync(Utf8String Name)
+		/// <returns>Writer for the given directory name. Null if a directory by this name does not exist.</returns>
+		public async ValueTask<TreePackDirWriter?> FindDirAsync(Utf8String Name)
 		{
 			TreePackDirWriter? Writer;
 			if (!NameToSubDir.TryGetValue(Name, out Writer))
@@ -116,17 +116,59 @@ namespace EpicGames.Horde.Storage
 				if (NameToEntry.TryGetValue(Name, out Entry) && (Entry.Flags & TreePackDirEntryFlags.Directory) != 0)
 				{
 					Writer = await FromDirNodeAsync(TreePack, Entry.Hash);
+					NameToEntry.Remove(Name);
+					NameToSubDir.Add(Name, Writer);
 				}
-				else
-				{
-					Writer = new TreePackDirWriter(TreePack);
-				}
-				NameToEntry.Remove(Name);
+			}
+			return Writer;
+		}
+
+		/// <summary>
+		/// Adds a directory to the current tree level with the given name
+		/// </summary>
+		/// <param name="Name"></param>
+		/// <returns></returns>
+		public async ValueTask<TreePackDirWriter> FindOrAddDirAsync(Utf8String Name)
+		{
+			TreePackDirWriter? Writer = await FindDirAsync(Name);
+			if (Writer == null)
+			{
+				Writer = new TreePackDirWriter(TreePack);
 				NameToSubDir.Add(Name, Writer);
 			}
 			return Writer;
 		}
 
+		/// <summary>
+		/// Adds a file to the tree pack with the given path. 
+		/// </summary>
+		/// <param name="Path">Path to the file. Should be specified using forward slashes as directory separators, and without a leading slash.</param>
+		public async Task RemoveFileByPathAsync(Utf8String Path)
+		{
+			int Index = Path.IndexOf('/');
+			if (Index == -1)
+			{
+				if (NameToEntry.TryGetValue(Path, out TreePackDirEntry? Entry) && (Entry.Flags & TreePackDirEntryFlags.File) != 0)
+				{
+					NameToEntry.Remove(Path);
+				}
+			}
+			else
+			{
+				Utf8String DirName = Path.Substring(0, Index);
+
+				TreePackDirWriter? SubDirWriter = await FindDirAsync(DirName);
+				if (SubDirWriter != null)
+				{
+					await SubDirWriter.RemoveFileByPathAsync(Path.Substring(Index + 1));
+					if (SubDirWriter.NameToEntry.Count == 0 && SubDirWriter.NameToSubDir.Count == 0)
+					{
+						NameToSubDir.Remove(DirName);
+					}
+				}
+			}
+		}
+
 		/// <summary>
 		/// Flushes the current tree state and writes it to the tree pack
 		/// </summary>
