Horde: Add general build failure exit code  to systemic  matcher
#rnx
#skipci

diff --git a/Engine/Source/Programs/Horde/Horde.Build/Issues/Impl/Handlers/SystemicIssueHandler.cs b/Engine/Source/Programs/Horde/Horde.Build/Issues/Impl/Handlers/SystemicIssueHandler.cs
--- a/Engine/Source/Programs/Horde/Horde.Build/Issues/Impl/Handlers/SystemicIssueHandler.cs
+++ b/Engine/Source/Programs/Horde/Horde.Build/Issues/Impl/Handlers/SystemicIssueHandler.cs
@@ -66,7 +66,8 @@ namespace HordeServer.IssueHandlers.Impl
 			{
 				for (int i = 0; i < EventData.Lines.Count; i++)
 				{
-					if (EventData.Lines[i].Message.Contains("AutomationTool exiting with ExitCode", StringComparison.InvariantCultureIgnoreCase))
+					string Message = EventData.Lines[i].Message;
+					if (Message.Contains("AutomationTool exiting with ExitCode", StringComparison.InvariantCultureIgnoreCase) || Message.Contains("BUILD FAILED", StringComparison.InvariantCultureIgnoreCase))
 					{
 						Fingerprint = new NewIssueFingerprint(Type, new[] { Node.Name }, null);
 						return true;
