Horde Storage - New snapshot format that does not require us to hold the full set of objects in memory, works by creating a temporary file that we can stream from.

diff --git a/Engine/Source/Programs/Horde/HordeStorage/Horde.Storage/Implementation/Replication/RefsReplicator.cs b/Engine/Source/Programs/Horde/HordeStorage/Horde.Storage/Implementation/Replication/RefsReplicator.cs
--- a/Engine/Source/Programs/Horde/HordeStorage/Horde.Storage/Implementation/Replication/RefsReplicator.cs
+++ b/Engine/Source/Programs/Horde/HordeStorage/Horde.Storage/Implementation/Replication/RefsReplicator.cs
@@ -296,8 +296,7 @@ namespace Horde.Storage.Implementation
                 HttpResponseMessage response = await _httpClient.SendAsync(request, HttpCompletionOption.ResponseHeadersRead, cancellationToken);
                 response.EnsureSuccessStatusCode();
                 await using Stream blobStream = await response.Content.ReadAsStreamAsync(cancellationToken);
-                using FilesystemBufferedPayload payload = await FilesystemBufferedPayload.Create(blobStream);
-                snapshot = await ReplicationLogSnapshot.DeserializeSnapshot(payload.GetStream());
+                snapshot = ReplicationLogFactory.DeserializeSnapshotFromStream(blobStream);
             }
 
             if (!snapshot.LastEvent.HasValue)
@@ -310,7 +309,7 @@ namespace Horde.Storage.Implementation
 
             // if MaxParallelReplications is set to not limit we use the default behavior of ParallelForEachAsync which is to limit based on CPUs 
             int maxParallelism = _replicatorSettings.MaxParallelReplications != -1 ? _replicatorSettings.MaxParallelReplications : 0;
-            await snapshot.LiveObjects.ParallelForEachAsync(async snapshotLiveObject =>
+            await snapshot.GetLiveObjects().ParallelForEachAsync(async snapshotLiveObject =>
             {
                 try
                 {
@@ -340,6 +339,8 @@ namespace Horde.Storage.Implementation
 
             }, maxParallelism, cancellationToken);
 
+            snapshot.Dispose();
+
             // snapshot processed, we proceed to reading the incremental state
             return (snapshotBucket, snapshotEvent, countOfObjectsReplicated);
         }
diff --git a/Engine/Source/Programs/Horde/HordeStorage/Horde.Storage/Implementation/TransactionLog/ReplicationLogSnapshot.cs b/Engine/Source/Programs/Horde/HordeStorage/Horde.Storage/Implementation/TransactionLog/ReplicationLogSnapshot.cs
--- a/Engine/Source/Programs/Horde/HordeStorage/Horde.Storage/Implementation/TransactionLog/ReplicationLogSnapshot.cs
+++ b/Engine/Source/Programs/Horde/HordeStorage/Horde.Storage/Implementation/TransactionLog/ReplicationLogSnapshot.cs
@@ -2,20 +2,163 @@
 
 using System;
 using System.Collections.Generic;
+using System.Diagnostics;
 using System.IO;
 using System.IO.Compression;
 using System.Linq;
+using System.Text;
 using System.Threading.Tasks;
+using EpicGames.Core;
 using EpicGames.Horde.Storage;
+using Jupiter.Common.Implementation;
 using Jupiter.Implementation;
 using Newtonsoft.Json;
 
 namespace Horde.Storage.Implementation.TransactionLog
 {
-    public class ReplicationLogSnapshot
+    public class JsonReplicationLogSnapshot : ReplicationLogSnapshot
     {
-        private readonly NamespaceId _namespace;
-        private readonly Dictionary<string, SnapshotLiveObject> _liveObjects = new Dictionary<string, SnapshotLiveObject>();
+        private readonly List<SnapshotLiveObject> _liveObjects;
+        JsonReplicationLogSnapshot(ReplicationLogSnapshotState snapshotState) : base(snapshotState.Header.Namespace, snapshotState.Header.LastBucket!, snapshotState.Header.LastEvent, (ulong)snapshotState.LiveObjects.Count)
+        {
+            _liveObjects = snapshotState.LiveObjects;
+        }
+
+        public static ReplicationLogSnapshot FromStream(Stream stream)
+        {
+            // legacy json snapshot found
+            using GZipStream gzipStream = new GZipStream(stream, CompressionMode.Decompress);
+            using TextReader textReader = new StreamReader(gzipStream);
+            using JsonReader reader = new JsonTextReader(textReader);
+            JsonSerializer serializer = JsonSerializer.Create();
+            ReplicationLogSnapshotState? snapshotState = serializer.Deserialize<ReplicationLogSnapshotState>(reader);
+            if (snapshotState == null)
+            {
+                throw new NotImplementedException();
+            }
+            return new JsonReplicationLogSnapshot(snapshotState);
+        }
+
+        public override IEnumerable<SnapshotLiveObject> GetLiveObjects()
+        {
+            return _liveObjects;
+        }
+    }
+
+    public class BinaryReplicationLogSnapshot : ReplicationLogSnapshot
+    {
+        private FilesystemBufferedPayload? _filesystemBufferedPayload;
+
+        private BinaryReplicationLogSnapshot(NamespaceId ns, string lastBucket, Guid lastEvent, ulong countOfObjects, FilesystemBufferedPayload payload) : base(ns, lastBucket, lastEvent, countOfObjects)
+        {
+            _filesystemBufferedPayload = payload; 
+        }
+
+        private BinaryReplicationLogSnapshot(NamespaceId ns) : base(ns)
+        {
+            _filesystemBufferedPayload = null;
+        }
+
+        public override void Dispose()
+        {
+            if (_filesystemBufferedPayload != null)
+            {
+                _filesystemBufferedPayload.Dispose();
+                _filesystemBufferedPayload = null;
+            }
+        }
+
+        public static ReplicationLogSnapshot FromStream(Stream stream)
+        {
+            FilesystemBufferedPayload payload = FilesystemBufferedPayload.Create(stream).Result;
+
+            using Stream payloadStream = payload.GetStream();
+            (NamespaceId ns, string lastBucket, Guid lastEvent, ulong countOfObjects) = ReadHeader(payloadStream);
+
+            return new BinaryReplicationLogSnapshot(ns, lastBucket, lastEvent, countOfObjects, payload);
+        }
+
+        private static (NamespaceId,string, Guid, ulong) ReadHeader(Stream stream)
+        {
+            byte[] bytes = new byte[4];
+            int bytesRead = stream.Read(bytes, 0, 4);
+            if (bytesRead != 4)
+                throw new Exception("Partial stream read when deserializing snapshot");
+
+            bool hasMagic = bytes[0] == 'S' && bytes[1] == 'N' && bytes[2] == 'A' && bytes[3] == 'P';
+            if (!hasMagic)
+                throw new Exception("Did not find magic prefix when deserializing snapshot, incorrect format?");
+
+            using BinaryReader reader = new BinaryReader(stream, Encoding.ASCII, leaveOpen: true);
+            string ns = reader.ReadString();
+            string lastBucket = reader.ReadString();
+            Guid lastEvent = new Guid(reader.ReadBytes(16));
+            ulong countOfObjects = reader.ReadUInt64();
+
+            return (new NamespaceId(ns), lastBucket, lastEvent, countOfObjects);
+        }
+
+
+        public override IEnumerable<SnapshotLiveObject> GetLiveObjects()
+        {
+            if (_filesystemBufferedPayload == null)
+                yield break;
+
+            using Stream payloadStream = _filesystemBufferedPayload.GetStream();
+
+            (NamespaceId ns, string lastBucket, Guid lastEvent, ulong countOfObjects) = ReadHeader(payloadStream);
+
+            using GZipStream gzipStream = new GZipStream(payloadStream, CompressionMode.Decompress);
+            using BinaryReader readerGzip = new BinaryReader(gzipStream);
+            for (ulong i = 0; i < countOfObjects; i++)
+            {
+                BucketId bucket = new BucketId(readerGzip.ReadString());
+                IoHashKey key = new IoHashKey(StringUtils.FormatHexString(readerGzip.ReadBytes(20)));
+                BlobIdentifier blobIdentifier = new BlobIdentifier(readerGzip.ReadBytes(20));
+
+                yield return new SnapshotLiveObject(bucket, key, blobIdentifier);
+            }
+        }
+
+        public static ReplicationLogSnapshot CreateEmptySnapshot(NamespaceId ns)
+        {
+            return new BinaryReplicationLogSnapshot(ns);
+        }
+    }
+
+    public static class ReplicationLogFactory
+    {
+        public static ReplicationLogSnapshot DeserializeSnapshotFromStream(Stream stream)
+        {
+            if (!stream.CanSeek)
+                throw new Exception("Seekable stream required");
+
+            byte[] bytes = new byte[4];
+            int bytesRead = stream.Read(bytes, 0, 4);
+            if (bytesRead != 4)
+                throw new Exception("Partial stream read when deserializing snapshot");
+
+            bool hasMagic = bytes[0] == 'S' && bytes[1] == 'N' && bytes[2] == 'A' && bytes[3] == 'P';
+            stream.Seek(0, SeekOrigin.Begin);
+
+            if (hasMagic)
+            {
+                return BinaryReplicationLogSnapshot.FromStream(stream);
+            }
+
+            return JsonReplicationLogSnapshot.FromStream(stream);
+        }
+
+        public static ReplicationLogSnapshot CreateEmptySnapshot(NamespaceId ns)
+        {
+            return BinaryReplicationLogSnapshot.CreateEmptySnapshot(ns);
+        }
+    }
+
+    public abstract class ReplicationLogSnapshot: IDisposable
+    {
+        private readonly List<SnapshotLiveObject> _addedObjects = new List<SnapshotLiveObject>();
+        private readonly HashSet<(BucketId, IoHashKey)> _removedObjects = new HashSet<(BucketId, IoHashKey)>();
 
         public class ReplicationLogSnapshotState
         {
@@ -32,7 +175,7 @@ namespace Horde.Storage.Implementation.TransactionLog
         public class SnapshotHeader
         {
             public NamespaceId Namespace { get; set; }
-            public string? LastBucket { get; set; }
+            public string LastBucket { get; set; } = null!;
             public Guid LastEvent { get; set; }
         }
 
@@ -52,70 +195,87 @@ namespace Horde.Storage.Implementation.TransactionLog
 
         public ReplicationLogSnapshot(NamespaceId ns)
         {
-            _namespace = ns;
+            Namespace = ns;
         }
 
-        private ReplicationLogSnapshot(ReplicationLogSnapshotState snapshotState)
+        protected ReplicationLogSnapshot(NamespaceId ns, string lastBucket, Guid lastEvent, ulong liveObjectsCount)
         {
-            _namespace = snapshotState.Header.Namespace;
-            LastEvent = snapshotState.Header.LastEvent;
-            LastBucket = snapshotState.Header.LastBucket;
-
-            _liveObjects = snapshotState.LiveObjects.ToDictionary(o => BuildObjectKey(o.Bucket, o.Key));
+            Namespace = ns;
+            LastEvent = lastEvent;
+            LastBucket = lastBucket;
+            LiveObjectsCount = liveObjectsCount;
         }
 
-        public Guid? LastEvent { get; set; }
-        public string? LastBucket { get; set; }
+        public NamespaceId Namespace { get; init; }
+        public Guid? LastEvent { get; private set; }
+        public string? LastBucket { get; private set; }
+        public ulong LiveObjectsCount { get; set; }
 
-        public IEnumerable<SnapshotLiveObject> LiveObjects
-        {
-            get { return _liveObjects.Values; }
-        }
+        public abstract IEnumerable<SnapshotLiveObject> GetLiveObjects();
 
-        public int LiveObjectsCount
+        public void Serialize(Stream stream)
         {
-            get { return _liveObjects.Values.Count; }
-        }
+            if (LastBucket == null)
+                throw new Exception("No last bucket found when serializing state, did you really have events?");
+            if (LastEvent == null)
+                throw new Exception("No last event found when serializing state, did you really have events?");
 
-        public static async Task<ReplicationLogSnapshot> DeserializeSnapshot(Stream stream)
-        {
-            await using GZipStream gzipStream = new GZipStream(stream, CompressionMode.Decompress);
-            using TextReader textReader = new StreamReader(gzipStream);
-            using JsonReader reader = new JsonTextReader(textReader);
-            JsonSerializer serializer = JsonSerializer.Create();
-            ReplicationLogSnapshotState? snapshotState = serializer.Deserialize<ReplicationLogSnapshotState>(reader);
-            if (snapshotState == null)
+            void WriteHeader(Stream s, NamespaceId ns, string lastBucket, Guid lastEvent, ulong countOfObjects)
             {
-                throw new NotImplementedException();
+                // write magic bytes
+                s.Write(Encoding.ASCII.GetBytes("SNAP"));
+
+                using BinaryWriter writer = new BinaryWriter(s, Encoding.ASCII, leaveOpen: true);
+                writer.Write(ns.ToString());
+                writer.Write(lastBucket);
+                writer.Write(lastEvent.ToByteArray());
+                writer.Write((ulong)countOfObjects);
             }
-            return new ReplicationLogSnapshot(snapshotState);
 
-        }
+            void WriteLiveObject(BinaryWriter writer, SnapshotLiveObject liveObject)
+            {
+                writer.Write(liveObject.Bucket.ToString());
+                // a io hash key is 20 bytes so we just write the binary format of it to keep the file small
+                // IoHashKey should likely not use a string as its internal representation
+                writer.Write(StringUtils.ParseHexString(liveObject.Key.ToString()), 0, 20);
+                writer.Write(liveObject.Blob.HashData, 0, 20);
+            }
 
-        public async Task Serialize(Stream stream)
-        {
-            // apply gzip on the stream as we have a lot of redundant text from bucket, but also within the keys
-            await using GZipStream gzipStream = new GZipStream(stream, CompressionMode.Compress);
-            await using TextWriter textWriter = new StreamWriter(gzipStream);
-            JsonSerializer serializer = JsonSerializer.Create();
+            // write header without the correct number of objects
+            WriteHeader(stream, Namespace, LastBucket, LastEvent.Value, 0);
 
-            List<SnapshotLiveObject> liveObjects = LiveObjects.ToList();
-            if (!liveObjects.Any())
-                throw new Exception("You must have at least one live object to build a snapshot");
-            if (LastBucket == null)
-                throw new Exception("No last bucket found when serializing state, did you really have events?");
-            if (LastEvent == null)
-                throw new Exception("No last event found when serializing state, did you really have events?");
+            ulong countOfLiveObjects = 0;
 
-            ReplicationLogSnapshotState state = new ReplicationLogSnapshotState(new SnapshotHeader()
             {
-                Namespace = _namespace,
-                LastBucket = LastBucket,
-                LastEvent = LastEvent.Value,
-            }, liveObjects);
-            serializer.Serialize(textWriter, state);
+                using GZipStream gzipStream = new GZipStream(stream, CompressionMode.Compress, leaveOpen: true);
+                using BinaryWriter zipWriter = new BinaryWriter(gzipStream, Encoding.ASCII, leaveOpen: true);
+                foreach (SnapshotLiveObject liveObject in GetLiveObjects())
+                {
+                    bool isRemoved = _removedObjects.Contains((liveObject.Bucket, liveObject.Key));
+                    if (isRemoved)
+                    {
+                        _removedObjects.Remove((liveObject.Bucket, liveObject.Key));
+                        // this object was deleted so we skip it
+                        continue;
+                    }
+
+                    WriteLiveObject(zipWriter, liveObject);
+                    countOfLiveObjects++;
+                }
+
+                foreach (SnapshotLiveObject liveObject in _addedObjects)
+                {
+                    WriteLiveObject(zipWriter, liveObject);
+                    countOfLiveObjects++;
+                }
+            }
+
+            // seek back to the start and update the header
+            stream.Seek(0, SeekOrigin.Begin);
+            WriteHeader(stream, Namespace, LastBucket, LastEvent.Value, countOfLiveObjects);
         }
 
+
         public void ProcessEvent(ReplicationLogEvent entry)
         {
             LastBucket = entry.TimeBucket;
@@ -134,19 +294,18 @@ namespace Horde.Storage.Implementation.TransactionLog
             }
         }
 
-
-        private string BuildObjectKey(BucketId bucket, IoHashKey key)
-        {
-            return $"{bucket}.{key}";
-        }
-
         private void ProcessAddEvent(BucketId bucket, IoHashKey key, BlobIdentifier blob)
         {
-            _liveObjects.TryAdd(BuildObjectKey(bucket, key), new SnapshotLiveObject(bucket, key, blob));
+            _addedObjects.Add(new SnapshotLiveObject(bucket, key, blob));
         }
         private void ProcessDeleteEvent(BucketId bucket, IoHashKey key)
         {
-            _liveObjects.Remove(BuildObjectKey(bucket, key));
+            _removedObjects.Add((bucket, key));
+        }
+
+        public virtual void Dispose()
+        {
+            
         }
     }
 }
diff --git a/Engine/Source/Programs/Horde/HordeStorage/Horde.Storage/Implementation/TransactionLog/ReplicationLogSnapshotBuilder.cs b/Engine/Source/Programs/Horde/HordeStorage/Horde.Storage/Implementation/TransactionLog/ReplicationLogSnapshotBuilder.cs
--- a/Engine/Source/Programs/Horde/HordeStorage/Horde.Storage/Implementation/TransactionLog/ReplicationLogSnapshotBuilder.cs
+++ b/Engine/Source/Programs/Horde/HordeStorage/Horde.Storage/Implementation/TransactionLog/ReplicationLogSnapshotBuilder.cs
@@ -7,6 +7,7 @@ using System.Linq;
 using System.Threading;
 using System.Threading.Tasks;
 using EpicGames.Horde.Storage;
+using Jupiter.Common.Implementation;
 using Jupiter.Implementation;
 
 namespace Horde.Storage.Implementation.TransactionLog
@@ -42,13 +43,13 @@ namespace Horde.Storage.Implementation.TransactionLog
                 await using BlobContents blobContents = await _blobService.GetObject(snapshotInfo.BlobNamespace, snapshotInfo.SnapshotBlob);
                 if (cancellationToken.IsCancellationRequested)
                     throw new TaskCanceledException();
-                snapshot = await ReplicationLogSnapshot.DeserializeSnapshot(blobContents.Stream);
+                snapshot = ReplicationLogFactory.DeserializeSnapshotFromStream(blobContents.Stream);
                 lastBucket = snapshot.LastBucket;
                 lastEvent = snapshot.LastEvent;
             }
             else
             {
-                snapshot = new ReplicationLogSnapshot(ns);
+                snapshot = ReplicationLogFactory.CreateEmptySnapshot(ns);
 
                 lastBucket = null;
                 lastEvent = null;
@@ -66,17 +67,20 @@ namespace Horde.Storage.Implementation.TransactionLog
             }
 
 
-            byte[] buf;
+            FileInfo tempFile = new FileInfo(Path.GetTempFileName());
             {
-                await using MemoryStream ms = new MemoryStream();
-                await snapshot.Serialize(ms);
-                await ms.FlushAsync(cancellationToken);
-                buf = ms.ToArray();
+                await using FileStream fs = tempFile.OpenWrite();
+                snapshot.Serialize(fs);
+                await fs.FlushAsync(cancellationToken);
             }
-            
-            
+
+            using FilesystemBufferedPayload payload = FilesystemBufferedPayload.FromTempFile(tempFile);
             {
-                BlobIdentifier blobIdentifier = BlobIdentifier.FromBlob(buf);
+                BlobIdentifier blobIdentifier;
+                {
+                    await using Stream stream = payload.GetStream();
+                    blobIdentifier = await BlobIdentifier.FromStream(stream);
+                }
 
                 CompactBinaryWriter writer = new CompactBinaryWriter();
                 writer.BeginObject();
@@ -89,10 +93,11 @@ namespace Horde.Storage.Implementation.TransactionLog
 
                 if (cancellationToken.IsCancellationRequested)
                     throw new TaskCanceledException();
-                
+
                 // upload the attachment first so we are not missing any references when we go to create the ref
-                await _blobService.PutObject(storeInNamespace, buf, blobIdentifier);
-                
+                await _blobService.PutObject(storeInNamespace, payload, blobIdentifier);
+            
+                tempFile.Delete();
                 (ContentId[] missingContentIds, BlobIdentifier[] missingBlobs) = await _objectService.Put(storeInNamespace, new BucketId("snapshot"), new IoHashKey(blobIdentifier.ToString()), cbBlobId, CompactBinaryObject.Load(cbObjectBytes));
                 List<ContentHash> missingHashes = new List<ContentHash>(missingContentIds);
                 missingHashes.AddRange(missingBlobs);
diff --git a/Engine/Source/Programs/Horde/HordeStorage/Horde.Storage/Tests/Functional/Objects/RefReplicatorTests.cs b/Engine/Source/Programs/Horde/HordeStorage/Horde.Storage/Tests/Functional/Objects/RefReplicatorTests.cs
--- a/Engine/Source/Programs/Horde/HordeStorage/Horde.Storage/Tests/Functional/Objects/RefReplicatorTests.cs
+++ b/Engine/Source/Programs/Horde/HordeStorage/Horde.Storage/Tests/Functional/Objects/RefReplicatorTests.cs
@@ -185,7 +185,7 @@ namespace Horde.Storage.FunctionalTests.Replication
             }
 
             // Build snapshot
-            ReplicationLogSnapshot snapshot = new ReplicationLogSnapshot(TestNamespace);
+            ReplicationLogSnapshot snapshot = ReplicationLogFactory.CreateEmptySnapshot(TestNamespace);
             foreach (ReplicationLogEvent logEvent in replicationEvents)
             {
                 snapshot.ProcessEvent(logEvent);
@@ -194,7 +194,7 @@ namespace Horde.Storage.FunctionalTests.Replication
             byte[] snapshotContent;
             {
                 await using MemoryStream ms = new MemoryStream();
-                await snapshot.Serialize(ms);
+                snapshot.Serialize(ms);
                 snapshotContent = ms.ToArray();
             }
 
@@ -270,7 +270,7 @@ namespace Horde.Storage.FunctionalTests.Replication
             }
             
             // Build snapshot
-            ReplicationLogSnapshot snapshot = new ReplicationLogSnapshot(TestNamespace);
+            ReplicationLogSnapshot snapshot = ReplicationLogFactory.CreateEmptySnapshot(TestNamespace);
             foreach (ReplicationLogEvent logEvent in snapshotEvents)
             {
                 snapshot.ProcessEvent(logEvent);
@@ -279,7 +279,7 @@ namespace Horde.Storage.FunctionalTests.Replication
             byte[] snapshotContent;
             {
                 await using MemoryStream ms = new MemoryStream();
-                await snapshot.Serialize(ms);
+                snapshot.Serialize(ms);
                 snapshotContent = ms.ToArray();
             }
 
@@ -350,7 +350,7 @@ namespace Horde.Storage.FunctionalTests.Replication
             }
 
             // Build snapshot
-            ReplicationLogSnapshot snapshot = new ReplicationLogSnapshot(TestNamespace);
+            ReplicationLogSnapshot snapshot = ReplicationLogFactory.CreateEmptySnapshot(TestNamespace);
             foreach (ReplicationLogEvent logEvent in snapshotEvents)
             {
                 snapshot.ProcessEvent(logEvent);
@@ -359,7 +359,7 @@ namespace Horde.Storage.FunctionalTests.Replication
             byte[] snapshotContent;
             {
                 await using MemoryStream ms = new MemoryStream();
-                await snapshot.Serialize(ms);
+                snapshot.Serialize(ms);
                 snapshotContent = ms.ToArray();
             }
 
diff --git a/Engine/Source/Programs/Horde/HordeStorage/Horde.Storage/Tests/Functional/Objects/ReplicationTests.cs b/Engine/Source/Programs/Horde/HordeStorage/Horde.Storage/Tests/Functional/Objects/ReplicationTests.cs
--- a/Engine/Source/Programs/Horde/HordeStorage/Horde.Storage/Tests/Functional/Objects/ReplicationTests.cs
+++ b/Engine/Source/Programs/Horde/HordeStorage/Horde.Storage/Tests/Functional/Objects/ReplicationTests.cs
@@ -593,15 +593,16 @@ namespace Horde.Storage.FunctionalTests.References
             Assert.AreEqual(snapshotBlobId, snapshotInfo.SnapshotBlob);
 
             BlobContents blobContents = await _blobStore.GetObject(SnapshotNamespace, snapshotBlobId);
-            ReplicationLogSnapshot snapshot = await ReplicationLogSnapshot.DeserializeSnapshot(blobContents.Stream);
+            ReplicationLogSnapshot snapshot = ReplicationLogFactory.DeserializeSnapshotFromStream(blobContents.Stream);
 
             Assert.AreEqual(lastEventId, snapshot.LastEvent);
             Assert.AreEqual(lastEventBucket, snapshot.LastBucket);
-            
-            Assert.IsTrue(snapshot.LiveObjects.Any(o => o.Bucket == TestBucket && o.Key == IoHashKey.FromName("firstObject")));
-            Assert.IsTrue(snapshot.LiveObjects.Any(o => o.Bucket == TestBucket && o.Key == IoHashKey.FromName("secondObject")));
-            Assert.IsTrue(snapshot.LiveObjects.Any(o => o.Bucket == TestBucket && o.Key == IoHashKey.FromName("thirdObject")));
-            Assert.IsTrue(snapshot.LiveObjects.Any(o => o.Bucket == TestBucket && o.Key == IoHashKey.FromName("fourthObject")));
+
+            List<ReplicationLogSnapshot.SnapshotLiveObject> liveObjects = snapshot.GetLiveObjects().ToList();
+            Assert.IsTrue(liveObjects.Any(o => o.Bucket == TestBucket && o.Key == IoHashKey.FromName("firstObject")));
+            Assert.IsTrue(liveObjects.Any(o => o.Bucket == TestBucket && o.Key == IoHashKey.FromName("secondObject")));
+            Assert.IsTrue(liveObjects.Any(o => o.Bucket == TestBucket && o.Key == IoHashKey.FromName("thirdObject")));
+            Assert.IsTrue(liveObjects.Any(o => o.Bucket == TestBucket && o.Key == IoHashKey.FromName("fourthObject")));
         }
 
 
@@ -654,7 +655,7 @@ namespace Horde.Storage.FunctionalTests.References
 
                 Assert.AreEqual(snapshotBlobId, foundSnapshot.SnapshotBlob);
                 BlobContents blobContents = await _blobStore.GetObject(SnapshotNamespace, snapshotBlobId);
-                ReplicationLogSnapshot snapshot = await ReplicationLogSnapshot.DeserializeSnapshot(blobContents.Stream);
+                ReplicationLogSnapshot snapshot = ReplicationLogFactory.DeserializeSnapshotFromStream(blobContents.Stream);
 
                 Assert.AreEqual(lastEventBucket, snapshot.LastBucket);
                 Assert.AreEqual(lastEventId, snapshot.LastEvent);
@@ -696,7 +697,7 @@ namespace Horde.Storage.FunctionalTests.References
             Assert.AreEqual(snapshotBlobId, snapshotInfo.SnapshotBlob);
 
             BlobContents blobContents = await _blobStore.GetObject(SnapshotNamespace, snapshotBlobId);
-            ReplicationLogSnapshot snapshot = await ReplicationLogSnapshot.DeserializeSnapshot(blobContents.Stream);
+            ReplicationLogSnapshot snapshot = ReplicationLogFactory.DeserializeSnapshotFromStream(blobContents.Stream);
 
             // insert more events
             await _replicationLog.InsertAddEvent(TestNamespace, TestBucket, IoHashKey.FromName("fifthObject"), objectHash, oldestTimestamp.AddDays(0.91));
diff --git a/Engine/Source/Programs/Horde/HordeStorage/Jupiter.Common/Implementation/StreamingBodyPayload.cs b/Engine/Source/Programs/Horde/HordeStorage/Jupiter.Common/Implementation/StreamingBodyPayload.cs
--- a/Engine/Source/Programs/Horde/HordeStorage/Jupiter.Common/Implementation/StreamingBodyPayload.cs
+++ b/Engine/Source/Programs/Horde/HordeStorage/Jupiter.Common/Implementation/StreamingBodyPayload.cs
@@ -60,6 +60,12 @@ namespace Jupiter.Common.Implementation
             _tempFile = new FileInfo(Path.GetTempFileName());
         }
 
+        private FilesystemBufferedPayload(FileInfo fi)
+        {
+            _tempFile = fi;
+            _length = fi.Length;
+        }
+
         public static async Task<FilesystemBufferedPayload> Create(Stream s)
         {
             FilesystemBufferedPayload payload = new FilesystemBufferedPayload();
@@ -88,6 +94,11 @@ namespace Jupiter.Common.Implementation
         {
             get { return _length; }
         }
+
+        public static FilesystemBufferedPayload FromTempFile(FileInfo tempFile)
+        {
+            return new FilesystemBufferedPayload(tempFile);
+        }
     }
 
     public class BufferedPayloadFactory
