Horde Storage - Fixed issue in snapshot generation - we require a seekable stream to be able to determine format. So streaming the snapshot to disk, even though we will actually copy it once more if its the new format.

diff --git a/Engine/Source/Programs/Horde/HordeStorage/Horde.Storage/Implementation/TransactionLog/ReplicationLogSnapshotBuilder.cs b/Engine/Source/Programs/Horde/HordeStorage/Horde.Storage/Implementation/TransactionLog/ReplicationLogSnapshotBuilder.cs
--- a/Engine/Source/Programs/Horde/HordeStorage/Horde.Storage/Implementation/TransactionLog/ReplicationLogSnapshotBuilder.cs
+++ b/Engine/Source/Programs/Horde/HordeStorage/Horde.Storage/Implementation/TransactionLog/ReplicationLogSnapshotBuilder.cs
@@ -43,7 +43,8 @@ namespace Horde.Storage.Implementation.TransactionLog
                 await using BlobContents blobContents = await _blobService.GetObject(snapshotInfo.BlobNamespace, snapshotInfo.SnapshotBlob);
                 if (cancellationToken.IsCancellationRequested)
                     throw new TaskCanceledException();
-                snapshot = ReplicationLogFactory.DeserializeSnapshotFromStream(blobContents.Stream);
+                using FilesystemBufferedPayload snapshotPayload = await FilesystemBufferedPayload.Create(blobContents.Stream);
+                snapshot = ReplicationLogFactory.DeserializeSnapshotFromStream(snapshotPayload.GetStream());
                 lastBucket = snapshot.LastBucket;
                 lastEvent = snapshot.LastEvent;
             }
