Horde Storage - Fixed failing unit tests

diff --git a/Engine/Source/Programs/Horde/HordeStorage/Horde.Storage/Implementation/Replication/IReplicator.cs b/Engine/Source/Programs/Horde/HordeStorage/Horde.Storage/Implementation/Replication/IReplicator.cs
--- a/Engine/Source/Programs/Horde/HordeStorage/Horde.Storage/Implementation/Replication/IReplicator.cs
+++ b/Engine/Source/Programs/Horde/HordeStorage/Horde.Storage/Implementation/Replication/IReplicator.cs
@@ -46,7 +46,7 @@ namespace Horde.Storage.Implementation
         /// <summary>
         /// Forcefully set the replication offset to a new value, should only be used to recover a replicator stuck in a bad state
         /// </summary>
-        void SetReplicationOffset(long state);
+        void SetReplicationOffset(long? state);
 
         Task StopReplicating();
 
diff --git a/Engine/Source/Programs/Horde/HordeStorage/Horde.Storage/Implementation/Replication/RefsReplicator.cs b/Engine/Source/Programs/Horde/HordeStorage/Horde.Storage/Implementation/Replication/RefsReplicator.cs
--- a/Engine/Source/Programs/Horde/HordeStorage/Horde.Storage/Implementation/Replication/RefsReplicator.cs
+++ b/Engine/Source/Programs/Horde/HordeStorage/Horde.Storage/Implementation/Replication/RefsReplicator.cs
@@ -581,7 +581,7 @@ namespace Horde.Storage.Implementation
             _logger.Information("{Name} starting replication. Last transaction was {TransactionId} {Generation}", _name, State.ReplicatorOffset.GetValueOrDefault(0L), State.ReplicatingGeneration.GetValueOrDefault(Guid.Empty) );
         }
 
-        public void SetReplicationOffset(long state)
+        public void SetReplicationOffset(long? state)
         {
             throw new NotImplementedException();
         }
diff --git a/Engine/Source/Programs/Horde/HordeStorage/Horde.Storage/Implementation/Replication/Replicator.cs b/Engine/Source/Programs/Horde/HordeStorage/Horde.Storage/Implementation/Replication/Replicator.cs
--- a/Engine/Source/Programs/Horde/HordeStorage/Horde.Storage/Implementation/Replication/Replicator.cs
+++ b/Engine/Source/Programs/Horde/HordeStorage/Horde.Storage/Implementation/Replication/Replicator.cs
@@ -448,7 +448,7 @@ namespace Horde.Storage.Implementation
             await _replicationFinishedEvent.WaitAsync();
         }
 
-        public void SetReplicationOffset(long offset)
+        public void SetReplicationOffset(long? offset)
         {
             State.ReplicatorOffset = offset;
             SaveState(_stateFile, State);
diff --git a/Engine/Source/Programs/Horde/HordeStorage/Horde.Storage/Tests/UnitTests/ReplicatorTests.cs b/Engine/Source/Programs/Horde/HordeStorage/Horde.Storage/Tests/UnitTests/ReplicatorTests.cs
--- a/Engine/Source/Programs/Horde/HordeStorage/Horde.Storage/Tests/UnitTests/ReplicatorTests.cs
+++ b/Engine/Source/Programs/Horde/HordeStorage/Horde.Storage/Tests/UnitTests/ReplicatorTests.cs
@@ -314,7 +314,7 @@ namespace Horde.Storage.UnitTests
             Mock<HttpMessageHandler> handler = new Mock<HttpMessageHandler>();
             IHttpClientFactory httpClientFactory = handler.CreateClientFactory();
             using IReplicator replicator = new ReplicatorV1(replicatorSettings, replicationSettingsMonitor, jupiterSettingsMonitor, blobStoreMock.Object, transactionLogWriter.Object, remoteClientMock.Object, serviceCredentials, httpClientFactory);
-
+            replicator.SetReplicationOffset(null);
             Assert.IsNull(replicator.State.ReplicatorOffset, "Expected state to have been reset during test initialize");
 
             bool ran = await replicator.TriggerNewReplications();
