Horde: Ensure job-queue scaled pools can grow even with a low queue count

Previously, the agent delta got rounded down when combined with the scale out factor.
This will prevent that behavior for pools where no agents are online.

diff --git a/Engine/Source/Programs/Horde/Horde.Build.Tests/Fleet/JobQueueStrategyTest.cs b/Engine/Source/Programs/Horde/Horde.Build.Tests/Fleet/JobQueueStrategyTest.cs
--- a/Engine/Source/Programs/Horde/Horde.Build.Tests/Fleet/JobQueueStrategyTest.cs
+++ b/Engine/Source/Programs/Horde/Horde.Build.Tests/Fleet/JobQueueStrategyTest.cs
@@ -42,6 +42,12 @@ namespace Horde.Build.Tests.Fleet
 			await AssertAgentCount(2, 0);
 		}
 		
+		[TestMethod]
+		public async Task NoAgentsInPool()
+		{
+			await AssertAgentCount(1, 1, true, 0);
+		}
+		
 		[TestMethod]
 		public async Task BatchesNotWaitingLongEnough()
 		{
@@ -66,9 +72,9 @@ namespace Horde.Build.Tests.Fleet
 			await AssertAgentCount(25, 6);
 		}
 
-		public async Task AssertAgentCount(int NumBatchesReady, int ExpectedAgentDelta, bool WaitedBeyondThreshold = true)
+		public async Task AssertAgentCount(int NumBatchesReady, int ExpectedAgentDelta, bool WaitedBeyondThreshold = true, int NumAgents = 10)
 		{
-			(JobQueueStrategy Strategy, PoolSizeData PoolSizeData) = await SetUpJobsAsync(1, NumBatchesReady);
+			(JobQueueStrategy Strategy, PoolSizeData PoolSizeData) = await SetUpJobsAsync(1, NumBatchesReady, NumAgents);
 			TimeSpan TimeToWait = WaitedBeyondThreshold
 				? Strategy.ReadyTimeThreshold + TimeSpan.FromSeconds(5)
 				: TimeSpan.FromSeconds(15);
@@ -85,11 +91,11 @@ namespace Horde.Build.Tests.Fleet
 		/// </summary>
 		/// <param name="NumBatchesRunning">Num of job batches that should be in state running</param>
 		/// <param name="NumBatchesReady">Num of job batches that should be in state waiting</param>
-		private async Task<(JobQueueStrategy, PoolSizeData)> SetUpJobsAsync(int NumBatchesRunning, int NumBatchesReady)
+		private async Task<(JobQueueStrategy, PoolSizeData)> SetUpJobsAsync(int NumBatchesRunning, int NumBatchesReady, int NumAgents = 10)
 		{
 			IPool Pool1 = await PoolService.CreatePoolAsync("bogusPool1", null, true, 0, 0);
 			List<IAgent> Agents = new();
-			for (int i = 0; i < 10; i++)
+			for (int i = 0; i < NumAgents; i++)
 			{
 				Agents.Add(await CreateAgentAsync(Pool1));
 			}
diff --git a/Engine/Source/Programs/Horde/Horde.Build/Fleet/Autoscale/JobQueueStrategy.cs b/Engine/Source/Programs/Horde/Horde.Build/Fleet/Autoscale/JobQueueStrategy.cs
--- a/Engine/Source/Programs/Horde/Horde.Build/Fleet/Autoscale/JobQueueStrategy.cs
+++ b/Engine/Source/Programs/Horde/Horde.Build/Fleet/Autoscale/JobQueueStrategy.cs
@@ -48,7 +48,7 @@ namespace Horde.Build.Fleet.Autoscale
 		///
 		/// Useful to avoid a very small queue size triggering any scaling. 
 		/// </summary>
-		private readonly int MinQueueSizeForScaleOut = 3;
+		private readonly int MinQueueSizeForScaleOut = 0;
 
 		/// <summary>
 		/// Factor translating queue size to additional agents to grow the pool with
@@ -148,6 +148,16 @@ namespace Horde.Build.Fleet.Autoscale
 				if (QueueSize > 0)
 				{
 					int AdditionalAgentCount = (int)Math.Round(QueueSize * ScaleOutFactor);
+
+					int NumAgentsEnabled = Current.Agents.Count(x => x.Enabled);
+					if (NumAgentsEnabled == 0)
+					{
+						// Ensure pool can grow even if no agents are online combined with low queue count
+						// A low queue count with default scale out factor will yield an increase below 0.5,
+						// which is rounded down (see above).
+						AdditionalAgentCount = Math.Max(AdditionalAgentCount, 1);	
+					}
+					
 					if (QueueSize < MinQueueSizeForScaleOut)
 						AdditionalAgentCount = 0;
 
