Horde Storage - Added debug logging to track down issue with 0 byte files in the filesystem cache

diff --git a/Engine/Source/Programs/Horde/HordeStorage/Horde.Storage/Controllers/ObjectController.cs b/Engine/Source/Programs/Horde/HordeStorage/Horde.Storage/Controllers/ObjectController.cs
--- a/Engine/Source/Programs/Horde/HordeStorage/Horde.Storage/Controllers/ObjectController.cs
+++ b/Engine/Source/Programs/Horde/HordeStorage/Horde.Storage/Controllers/ObjectController.cs
@@ -173,6 +173,10 @@ namespace Horde.Storage.Controllers
             }
            
             byte[] blobContents = await blob.Stream.ToByteArray();
+            if (blobContents.Length == 0)
+            {
+                _logger.Warning("0 byte object found for {Id} {Namespace}", id, ns);
+            }
             CompactBinaryObject compactBinaryObject = CompactBinaryObject.Load(blobContents);
 
             try
diff --git a/Engine/Source/Programs/Horde/HordeStorage/Horde.Storage/Implementation/Blob/FileSystemStore.cs b/Engine/Source/Programs/Horde/HordeStorage/Horde.Storage/Implementation/Blob/FileSystemStore.cs
--- a/Engine/Source/Programs/Horde/HordeStorage/Horde.Storage/Implementation/Blob/FileSystemStore.cs
+++ b/Engine/Source/Programs/Horde/HordeStorage/Horde.Storage/Implementation/Blob/FileSystemStore.cs
@@ -64,6 +64,9 @@ namespace Horde.Storage.Implementation
                 await fs.FlushAsync(cancellationToken);
             }
 
+            if (content.Length == 0)
+                _logger.Warning("0 byte file written as {Id} {Namespace}", blobIdentifier, ns);
+
             UpdateLastWriteTime(filePath.FullName, DateTime.UnixEpoch);
             return blobIdentifier;
         }
@@ -80,6 +83,8 @@ namespace Horde.Storage.Implementation
                 await content.CopyToAsync(fs, cancellationToken);
                 await fs.FlushAsync(cancellationToken);
             }
+            if (filePath.Length == 0)
+                _logger.Warning("0 byte file written as {Id} {Namespace}", blobIdentifier, ns);
 
             UpdateLastWriteTime(filePath.FullName, DateTime.UnixEpoch);
             return blobIdentifier;
@@ -97,6 +102,9 @@ namespace Horde.Storage.Implementation
                 await fs.FlushAsync(cancellationToken);
             }
 
+            if (content.Length == 0)
+                _logger.Warning("0 byte file written as {Id} {Namespace}", blobIdentifier, ns);
+
             UpdateLastWriteTime(filePath.FullName, DateTime.UnixEpoch);
             return blobIdentifier;
         }
