Horde Storage - Added check if reference is finalized to Head check, not finalized is treated as not existing.

diff --git a/Engine/Source/Programs/Horde/HordeStorage/Horde.Storage/Controllers/ReferencesController.cs b/Engine/Source/Programs/Horde/HordeStorage/Horde.Storage/Controllers/ReferencesController.cs
--- a/Engine/Source/Programs/Horde/HordeStorage/Horde.Storage/Controllers/ReferencesController.cs
+++ b/Engine/Source/Programs/Horde/HordeStorage/Horde.Storage/Controllers/ReferencesController.cs
@@ -279,8 +279,13 @@ namespace Horde.Storage.Controllers
 
             try
             {
-                (ObjectRecord record, BlobContents _) = await _objectService.Get(ns, bucket, key, new string[] {"blobIdentifier"});
+                (ObjectRecord record, BlobContents _) = await _objectService.Get(ns, bucket, key, new string[] {"blobIdentifier", "IsFinalized"});
                 Response.Headers[CommonHeaders.HashHeaderName] = record.BlobIdentifier.ToString();
+
+                if (!record.IsFinalized)
+                {
+                    return NotFound(new ProblemDetails {Title = $"Object {bucket} {key} in namespace {ns} is not finalized."});
+                }
             }
             catch (NamespaceNotFoundException e)
             {
