Horde Storage - Fixed issue deleting temp file that was still in use

diff --git a/Engine/Source/Programs/Horde/HordeStorage/Horde.Storage/Implementation/TransactionLog/ReplicationLogSnapshotBuilder.cs b/Engine/Source/Programs/Horde/HordeStorage/Horde.Storage/Implementation/TransactionLog/ReplicationLogSnapshotBuilder.cs
--- a/Engine/Source/Programs/Horde/HordeStorage/Horde.Storage/Implementation/TransactionLog/ReplicationLogSnapshotBuilder.cs
+++ b/Engine/Source/Programs/Horde/HordeStorage/Horde.Storage/Implementation/TransactionLog/ReplicationLogSnapshotBuilder.cs
@@ -44,7 +44,8 @@ namespace Horde.Storage.Implementation.TransactionLog
                 if (cancellationToken.IsCancellationRequested)
                     throw new TaskCanceledException();
                 using FilesystemBufferedPayload snapshotPayload = await FilesystemBufferedPayload.Create(blobContents.Stream);
-                snapshot = ReplicationLogFactory.DeserializeSnapshotFromStream(snapshotPayload.GetStream());
+                await using Stream s = snapshotPayload.GetStream();
+                snapshot = ReplicationLogFactory.DeserializeSnapshotFromStream(s);
                 lastBucket = snapshot.LastBucket;
                 lastEvent = snapshot.LastEvent;
             }
