Horde: Add run mode settings

To allow configuring processes for different types of workloads. Primarily distinguishing between request serving vs background processing

diff --git a/Engine/Source/Programs/Horde/Horde.Build/ServerSettings.cs b/Engine/Source/Programs/Horde/Horde.Build/ServerSettings.cs
--- a/Engine/Source/Programs/Horde/Horde.Build/ServerSettings.cs
+++ b/Engine/Source/Programs/Horde/Horde.Build/ServerSettings.cs
@@ -139,6 +139,29 @@ namespace HordeServer
 		/// </summary>
 		public int MaxConcurrentLeasesPerAgent { get; set; } = 1;
 	}
+
+	/// <summary>
+	/// Type of run mode this process should use. Each carry different types of workloads. 
+	/// More than one mode can be active. But not all modes are not guaranteed to be compatible with each other and will
+	/// raise an error if combined in such a way.
+	/// </summary>
+	public enum RunMode
+	{
+		/// <summary>
+		/// Handle and respond to incoming external requests, such as HTTP REST and gRPC calls.
+		/// These requests are time-sensitive and short-lived, typically less than 5 secs.
+		/// If processes handling requests are unavailable, it will be very visible for users.
+		/// </summary>
+		RequestHandler,
+		
+		/// <summary>
+		/// Run non-request facing workloads. Such as background services, processing queues, running work
+		/// based on timers etc. Short periods of downtime or high CPU usage due to bursts are fine for this mode.
+		/// No user requests will be impacted directly. If auto-scaling is used, a much more aggressive policy can be
+		/// applied (tighter process packing, higher avg CPU usage).
+		/// </summary>
+		Worker
+	}
 	
 	/// <summary>
 	/// Feature flags to aid rollout of new features
@@ -173,6 +196,9 @@ namespace HordeServer
 	/// </summary>
 	public class ServerSettings
 	{
+		/// <inheritdoc cref="RunMode" />
+		public RunMode[] RunModes { get; set; } = { RunMode.RequestHandler, RunMode.Worker };
+		
 		/// <summary>
 		/// Main port for serving HTTP. Uses the default Kestrel port (5000) if not specified.
 		/// </summary>
@@ -438,12 +464,6 @@ namespace HordeServer
 		/// </summary>
 		public bool DisableSchedules { get; set; }
 
-		/// <summary>
-		/// Whether to run background services
-		///
-		/// </summary>
-		public bool EnableBackgroundServices { get; set; } = true;
-
 		/// <summary>
 		/// Timezone for evaluating schedules
 		/// </summary>
@@ -591,5 +611,15 @@ namespace HordeServer
 
 		/// <inheritdoc cref="FeatureFlags" />
 		public FeatureFlagSettings FeatureFlags { get; set; } = new ();
+
+		/// <summary>
+		/// Helper method to check if this process has activated the given mode
+		/// </summary>
+		/// <param name="Mode">Run mode</param>
+		/// <returns>True if mode is active</returns>
+		public bool IsRunModeActive(RunMode Mode)
+		{
+			return RunModes.Contains(Mode);
+		}
 	}
 }
diff --git a/Engine/Source/Programs/Horde/Horde.Build/Startup.cs b/Engine/Source/Programs/Horde/Horde.Build/Startup.cs
--- a/Engine/Source/Programs/Horde/Horde.Build/Startup.cs
+++ b/Engine/Source/Programs/Horde/Horde.Build/Startup.cs
@@ -559,7 +559,7 @@ namespace HordeServer
 						.Build();
 				});
 
-			if (Settings.EnableBackgroundServices)
+			if (Settings.IsRunModeActive(RunMode.Worker))
 			{
 				if (Settings.FeatureFlags.AutoscaleServiceV1Enabled)
 				{
