Horde Storage - Added temporary mime type (application/x-jupiter-inline) for inlining of value api requests until we have settled on a general format for inlining attachments with the requests.
Also updated the head and multiple exists checks to verify that blobs exists in the local region (similarly to what we did for the old api) - this step should not be nesscary once we support fetching blobs from other regions as the data does exist in Horde Storage, just not necessarily in the local blob store.

diff --git a/Engine/Source/Programs/Horde/HordeStorage/Horde.Storage/Controllers/ReferencesController.cs b/Engine/Source/Programs/Horde/HordeStorage/Horde.Storage/Controllers/ReferencesController.cs
--- a/Engine/Source/Programs/Horde/HordeStorage/Horde.Storage/Controllers/ReferencesController.cs
+++ b/Engine/Source/Programs/Horde/HordeStorage/Horde.Storage/Controllers/ReferencesController.cs
@@ -41,12 +41,13 @@ namespace Horde.Storage.Controllers
         private readonly IOptionsMonitor<JupiterSettings> _jupiterSettings;
         private readonly FormatResolver _formatResolver;
         private readonly BufferedPayloadFactory _bufferedPayloadFactory;
+        private readonly IReferenceResolver _referenceResolver;
 
         private readonly ILogger _logger = Log.ForContext<ReferencesController>();
         private readonly IObjectService _objectService;
         private readonly IBlobService _blobStore;
 
-        public ReferencesController(IObjectService objectService, IBlobService blobStore, IDiagnosticContext diagnosticContext, IAuthorizationService authorizationService, IOptionsMonitor<JupiterSettings> jupiterSettings, FormatResolver formatResolver, BufferedPayloadFactory bufferedPayloadFactory)
+        public ReferencesController(IObjectService objectService, IBlobService blobStore, IDiagnosticContext diagnosticContext, IAuthorizationService authorizationService, IOptionsMonitor<JupiterSettings> jupiterSettings, FormatResolver formatResolver, BufferedPayloadFactory bufferedPayloadFactory, IReferenceResolver referenceResolver)
         {
             _objectService = objectService;
             _blobStore = blobStore;
@@ -55,6 +56,7 @@ namespace Horde.Storage.Controllers
             _jupiterSettings = jupiterSettings;
             _formatResolver = formatResolver;
             _bufferedPayloadFactory = bufferedPayloadFactory;
+            _referenceResolver = referenceResolver;
         }
 
         /// <summary>
@@ -88,6 +90,7 @@ namespace Horde.Storage.Controllers
         /// <param name="fields">The fields to include in the response, omit this to include everything.</param>
         /// <param name="format">Optional specifier to set which output format is used json/raw/cb</param>
         [HttpGet("{ns}/{bucket}/{key}.{format?}", Order = 500)]
+        [Produces(MediaTypeNames.Application.Json, MediaTypeNames.Application.Octet, CustomMediaTypeNames.UnrealCompactBinary, CustomMediaTypeNames.JupiterInlinedPayload)]
         [Authorize("Object.read")]
         public async Task<IActionResult> Get(
             [FromRoute] [Required] NamespaceId ns,
@@ -183,6 +186,54 @@ namespace Horde.Storage.Controllers
                         break;
 
                     }
+                    case CustomMediaTypeNames.JupiterInlinedPayload:
+                    {
+                        byte[] blobMemory = await blob.Stream.ToByteArray();
+                        ReadOnlyMemory<byte> localMemory = new ReadOnlyMemory<byte>(blobMemory);
+                        CompactBinaryObject cb = CompactBinaryObject.Load(ref localMemory);
+                        List<CompactBinaryField> compactBinaryFields = cb.GetFields().ToList();
+                        int countOfBinaryAttachmentFields = compactBinaryFields.Count(field => field.IsBinaryAttachment());
+                        int countOfAttachmentFields = compactBinaryFields.Count(field => field.IsAttachment());
+                        // if the object consists of a single attachment field we return this attachment field instead
+                        if (countOfBinaryAttachmentFields == 1 && countOfAttachmentFields == 1)
+                        {
+                            // fetch the blob so we can resolve any content ids in it
+                            IAsyncEnumerable<BlobIdentifier> referencedBlobsEnumerable = _referenceResolver.ResolveReferences(ns, cb);
+                            List<BlobIdentifier> referencedBlobs = await referencedBlobsEnumerable.ToListAsync();
+                            if (referencedBlobs.Count == 1)
+                            {
+                                BlobContents referencedBlobContents = await _blobStore.GetObject(ns, referencedBlobs.First());
+                                await WriteBody(referencedBlobContents, CustomMediaTypeNames.JupiterInlinedPayload);
+                                return new EmptyResult();
+                            }
+                            else if (referencedBlobs.Count == 0)
+                            {
+                                return NotFound(new ProblemDetails
+                                {
+                                    Title =
+                                        $"Object {objectRecord.Bucket} {objectRecord.Name} did not resolve into any objects that we could find."
+                                });
+                            }
+
+                            return BadRequest(new ProblemDetails
+                            {
+                                Title =
+                                    $"Object {objectRecord.Bucket} {objectRecord.Name} contained a content id which resolved to more then 1 blob, unable to inline this object. Use compact object response instead."
+                            });
+                        }
+                        else if (countOfBinaryAttachmentFields == 0 && countOfAttachmentFields == 0)
+                        {
+                            // no attachments so we just return the compact object instead
+                            await WriteBody(new BlobContents(blobMemory), CustomMediaTypeNames.JupiterInlinedPayload);
+                            return new EmptyResult();
+                        }
+
+                        return BadRequest(new ProblemDetails
+                        {
+                            Title =
+                                $"Object {objectRecord.Bucket} {objectRecord.Name} had more then 1 binary attachment field, unable to inline this object. Use compact object response instead."
+                        });
+                    }
                     default:
                         throw new NotImplementedException($"Unknown expected response type {responseType}");
                 }
@@ -279,13 +330,31 @@ namespace Horde.Storage.Controllers
 
             try
             {
-                (ObjectRecord record, BlobContents _) = await _objectService.Get(ns, bucket, key, new string[] {"blobIdentifier", "IsFinalized"});
+                (ObjectRecord record, BlobContents? blob) = await _objectService.Get(ns, bucket, key, new string[] {"blobIdentifier", "IsFinalized"});
                 Response.Headers[CommonHeaders.HashHeaderName] = record.BlobIdentifier.ToString();
 
                 if (!record.IsFinalized)
                 {
                     return NotFound(new ProblemDetails {Title = $"Object {bucket} {key} in namespace {ns} is not finalized."});
                 }
+
+                blob ??= await _blobStore.GetObject(ns, record.BlobIdentifier);
+
+                // we have to verify the blobs are available locally, as the record of the key is replicated a head of the content
+                // TODO: Once we support inline replication this step is not needed as at least one region as this blob, just maybe not this current one
+                byte[] blobContents = await blob.Stream.ToByteArray();
+                CompactBinaryObject compactBinaryObject = CompactBinaryObject.Load(blobContents);
+                // the reference resolver will throw if any blob is missing, so no need to do anything other then process each reference
+                IAsyncEnumerable<BlobIdentifier> references = _referenceResolver.ResolveReferences(ns, compactBinaryObject);
+                List<BlobIdentifier>? _ = await references.ToListAsync();
+
+                // we have to verify the blobs are available locally, as the record of the key is replicated a head of the content
+                // TODO: Once we support inline replication this step is not needed as at least one region as this blob, just maybe not this current one
+                BlobIdentifier[] unknownBlobs = await _blobStore.FilterOutKnownBlobs(ns, new BlobIdentifier[] { record.BlobIdentifier });
+                if (unknownBlobs.Length != 0)
+                {
+                    return NotFound(new ProblemDetails {Title = $"Object {bucket} {key} in namespace {ns} had at least one missing blob."});
+                }
             }
             catch (NamespaceNotFoundException e)
             {
@@ -299,6 +368,14 @@ namespace Horde.Storage.Controllers
             {
                 return NotFound(new ProblemDetails { Title = $"Blobs {e.Blobs} from object {e.Bucket} {e.Key} in namespace {e.Namespace} did not exist" });
             }
+            catch (PartialReferenceResolveException)
+            {
+                return NotFound(new ProblemDetails {Title = $"Object {bucket} {key} in namespace {ns} was missing some content ids"});
+            }
+            catch (ReferenceIsMissingBlobsException)
+            {
+                return NotFound(new ProblemDetails {Title = $"Object {bucket} {key} in namespace {ns} was missing some blobs"});
+            }
 
             return Ok();
         }
@@ -324,12 +401,31 @@ namespace Horde.Storage.Controllers
             {
                 try
                 {
-                    await _objectService.Get(ns, bucket, name, new string[] {"blobIdentifier"});
+                    (ObjectRecord record, BlobContents? blob) =
+                        await _objectService.Get(ns, bucket, name, new string[] { "blobIdentifier" });
+
+                    blob ??= await _blobStore.GetObject(ns, record.BlobIdentifier);
+
+                    // we have to verify the blobs are available locally, as the record of the key is replicated a head of the content
+                    // TODO: Once we support inline replication this step is not needed as at least one region as this blob, just maybe not this current one
+                    byte[] blobContents = await blob.Stream.ToByteArray();
+                    CompactBinaryObject compactBinaryObject = CompactBinaryObject.Load(blobContents);
+                    // the reference resolver will throw if any blob is missing, so no need to do anything other then process each reference
+                    IAsyncEnumerable<BlobIdentifier> references = _referenceResolver.ResolveReferences(ns, compactBinaryObject);
+                    List<BlobIdentifier>? _ = await references.ToListAsync();
                 }
                 catch (ObjectNotFoundException)
                 {
                     missingObject.Add(name);
                 }
+                catch (PartialReferenceResolveException)
+                {
+                    missingObject.Add(name);
+                }
+                catch (ReferenceIsMissingBlobsException)
+                {
+                    missingObject.Add(name);
+                }
             });
             await Task.WhenAll(tasks);
 
diff --git a/Engine/Source/Programs/Horde/HordeStorage/Horde.Storage/Tests/Functional/Objects/ReferencesTests.cs b/Engine/Source/Programs/Horde/HordeStorage/Horde.Storage/Tests/Functional/Objects/ReferencesTests.cs
--- a/Engine/Source/Programs/Horde/HordeStorage/Horde.Storage/Tests/Functional/Objects/ReferencesTests.cs
+++ b/Engine/Source/Programs/Horde/HordeStorage/Horde.Storage/Tests/Functional/Objects/ReferencesTests.cs
@@ -323,6 +323,24 @@ namespace Horde.Storage.FunctionalTests.References
                 Assert.IsNotNull(value);
                 Assert.AreEqual(objectHash, new BlobIdentifier(value));
             }
+
+            {
+                // request the object as a jupiter inlined payload
+                var request = new HttpRequestMessage(HttpMethod.Get, $"api/v1/refs/{TestNamespace}/bucket/{key}");
+                request.Headers.Accept.Add(new MediaTypeWithQualityHeaderValue(CustomMediaTypeNames.JupiterInlinedPayload));
+                HttpResponseMessage getResponse = await _httpClient.SendAsync(request);
+                getResponse.EnsureSuccessStatusCode();
+                Assert.AreEqual(CustomMediaTypeNames.JupiterInlinedPayload, getResponse.Content.Headers.ContentType?.MediaType);
+
+                await using MemoryStream ms = new MemoryStream();
+                await getResponse.Content.CopyToAsync(ms);
+                byte[] roundTrippedBuffer = ms.ToArray();
+
+                string roundTrippedString = Encoding.ASCII.GetString(roundTrippedBuffer);
+
+                Assert.AreEqual(objectContents, roundTrippedString);
+                Assert.AreEqual(objectHash, BlobIdentifier.FromBlob(roundTrippedBuffer));
+            }
         }
 
         [TestMethod]
diff --git a/Engine/Source/Programs/Horde/HordeStorage/Jupiter.Common/CustomMediaTypeNames.cs b/Engine/Source/Programs/Horde/HordeStorage/Jupiter.Common/CustomMediaTypeNames.cs
--- a/Engine/Source/Programs/Horde/HordeStorage/Jupiter.Common/CustomMediaTypeNames.cs
+++ b/Engine/Source/Programs/Horde/HordeStorage/Jupiter.Common/CustomMediaTypeNames.cs
@@ -6,6 +6,8 @@ namespace Jupiter
     {
         public const string UnrealCompactBinary = "application/x-ue-cb";
         public const string UnrealCompressedBuffer = "application/x-ue-comp";
+
+        public const string JupiterInlinedPayload = "application/x-jupiter-inline";
     }
 
 }
