Horde Storage - Missing file to add support for x-jupiter-inline

diff --git a/Engine/Source/Programs/Horde/HordeStorage/Horde.Storage/Controllers/FormatResolver.cs b/Engine/Source/Programs/Horde/HordeStorage/Horde.Storage/Controllers/FormatResolver.cs
--- a/Engine/Source/Programs/Horde/HordeStorage/Horde.Storage/Controllers/FormatResolver.cs
+++ b/Engine/Source/Programs/Horde/HordeStorage/Horde.Storage/Controllers/FormatResolver.cs
@@ -16,7 +16,7 @@ namespace Horde.Storage.Controllers
         private readonly IOptionsMonitor<MvcOptions> _mvcOptions;
 
         private readonly string[] _validContentTypes = {
-            MediaTypeNames.Application.Octet, MediaTypeNames.Application.Json, CustomMediaTypeNames.UnrealCompactBinary
+            MediaTypeNames.Application.Octet, MediaTypeNames.Application.Json, CustomMediaTypeNames.UnrealCompactBinary, CustomMediaTypeNames.JupiterInlinedPayload
         };
 
         public FormatResolver(IOptionsMonitor<MvcOptions> mvcOptions)
