Horde: Remove unused config in ServerSettings

diff --git a/Engine/Source/Programs/Horde/Horde.Build.Tests/Fleet/AutoscaleServiceV2Test.cs b/Engine/Source/Programs/Horde/Horde.Build.Tests/Fleet/AutoscaleServiceV2Test.cs
--- a/Engine/Source/Programs/Horde/Horde.Build.Tests/Fleet/AutoscaleServiceV2Test.cs
+++ b/Engine/Source/Programs/Horde/Horde.Build.Tests/Fleet/AutoscaleServiceV2Test.cs
@@ -67,34 +67,6 @@ namespace Horde.Build.Tests.Fleet
 	{
 		FleetManagerSpy FleetManagerSpy = new();
 		
-		[TestMethod]
-		public async Task ShadowModeEnabled()
-		{
-			ServerSettings.FeatureFlags.AutoscaleServiceV2ShadowMode = false;
- 
-			var Service = GetAutoscaleService(FleetManagerSpy);
-			var PoolSizeDatas = await GetPoolSizeData();
-			PoolSizeDatas[0] = PoolSizeDatas[0].Copy(DesiredAgentCount: 8);
-
-			await Service.ResizePools(PoolSizeDatas);
-			Assert.AreEqual(1, FleetManagerSpy.ExpandPoolAsyncCallCount);
-			Assert.AreEqual(0, FleetManagerSpy.ShrinkPoolAsyncCallCount);
-		}
-		
-		[TestMethod]
-		public async Task ShadowModeDisabled()
-		{
-			ServerSettings.FeatureFlags.AutoscaleServiceV2ShadowMode = true;
- 
-			var Service = GetAutoscaleService(FleetManagerSpy);
-			var PoolSizeDatas = await GetPoolSizeData();
-			PoolSizeDatas[0] = PoolSizeDatas[0].Copy(DesiredAgentCount: 8);
-
-			await Service.ResizePools(PoolSizeDatas);
-			Assert.AreEqual(0, FleetManagerSpy.ExpandPoolAsyncCallCount);
-			Assert.AreEqual(0, FleetManagerSpy.ShrinkPoolAsyncCallCount);
-		}
-		
 		[TestMethod]
 		public async Task PerPoolStrategy()
 		{
diff --git a/Engine/Source/Programs/Horde/Horde.Build/Fleet/Autoscale/AutoscaleServiceV2.cs b/Engine/Source/Programs/Horde/Horde.Build/Fleet/Autoscale/AutoscaleServiceV2.cs
--- a/Engine/Source/Programs/Horde/Horde.Build/Fleet/Autoscale/AutoscaleServiceV2.cs
+++ b/Engine/Source/Programs/Horde/Horde.Build/Fleet/Autoscale/AutoscaleServiceV2.cs
@@ -118,11 +118,6 @@ namespace Horde.Build.Fleet.Autoscale
 				int DeltaAgentCount = DesiredAgentCount - CurrentAgentCount;
 
 				Logger.LogInformation("{PoolName,-48} Current={Current,4} Target={Target,4} Delta={Delta,4} Status={Status}", Pool.Name, CurrentAgentCount, DesiredAgentCount, DeltaAgentCount, PoolSizeData.StatusMessage);
-
-				if (Settings.FeatureFlags.AutoscaleServiceV2ShadowMode)
-				{
-					continue;
-				}
 				
 				try
 				{
diff --git a/Engine/Source/Programs/Horde/Horde.Build/ServerSettings.cs b/Engine/Source/Programs/Horde/Horde.Build/ServerSettings.cs
--- a/Engine/Source/Programs/Horde/Horde.Build/ServerSettings.cs
+++ b/Engine/Source/Programs/Horde/Horde.Build/ServerSettings.cs
@@ -82,64 +82,6 @@ namespace HordeServer
 		OpenIdConnect,
 	}
 
-	/// <summary>
-	/// Settings for a remote execution instance
-	///
-	/// Instances are a concept from Google's remote execution API that Horde implements.
-	/// They are used for separating different type of executions into groups.
-	/// </summary>
-	public class RemoteExecInstanceSettings
-	{
-		/// <summary>
-		/// gRPC URL for the content-addressable storage to be used for this instance
-		/// </summary>
-		[Required]
-		public Uri? CasUrl { get; set; }
-		
-		/// <summary>
-		/// gRPC URL for the action cache to be used for this instance
-		/// </summary>
-		[Required]
-		public Uri? ActionCacheUrl { get; set; }
-		
-		/// <summary>
-		/// Service account token for accessing the gRPC URLs
-		/// </summary>
-		[Required]
-		public string? ServiceAccountToken { get; set; }
-		
-		/// <summary>
-		/// List of pool IDs this instance can be scheduled in
-		/// </summary>
-		public List<string>? PoolFilter { get; set; }
-
-		/// <summary>
-		/// Max time to wait for an agent to become available for serving a remote execution request (in milliseconds)
-		/// </summary>
-		public int AgentQueueTimeout { get; set; } = 5000;
-
-		internal TimeSpan AgentQueueTimeoutAsTimeSpan()
-		{
-			return TimeSpan.FromMilliseconds(AgentQueueTimeout);
-		}
-	}
-	
-	/// <summary>
-	/// Settings for remote execution
-	/// </summary>
-	public class RemoteExecSettings
-	{
-		/// <summary>
-		/// Mapping instance names to instance settings
-		/// </summary>
-		public Dictionary<string, RemoteExecInstanceSettings> Instances { get; set; } = new Dictionary<string, RemoteExecInstanceSettings>();
-
-		/// <summary>
-		/// Max number of concurrent leases per agent
-		/// </summary>
-		public int MaxConcurrentLeasesPerAgent { get; set; } = 1;
-	}
-
 	/// <summary>
 	/// Type of run mode this process should use. Each carry different types of workloads. 
 	/// More than one mode can be active. But not all modes are not guaranteed to be compatible with each other and will
@@ -172,23 +114,6 @@ namespace HordeServer
 	/// </summary>
 	public class FeatureFlagSettings
 	{
-		/// <summary>
-		/// Registers the old v1 autoscale service as a background service
-		/// Added by carl.bystrom on Jan 26, 2022
-		/// </summary>
-		public bool AutoscaleServiceV1Enabled { get; set; } = true;
-		
-		/// <summary>
-		/// Registers the new autoscale service v2 as a background service (can run in parallel with v1 if shadow mode is on)
-		/// Added by carl.bystrom on Jan 26, 2022
-		/// </summary>
-		public bool AutoscaleServiceV2Enabled { get; set; } = false;
-		
-		/// <summary>
-		/// Shadow mode enabled on new AutoscaleServiceV2 will only log pool size changes but not act on them.
-		/// Added by carl.bystrom on Jan 26, 2022
-		/// </summary>
-		public bool AutoscaleServiceV2ShadowMode { get; set; } = true;
 	}
 
 	/// <summary>
@@ -577,11 +502,6 @@ namespace HordeServer
 		/// </summary>
 		public string ConfigPath { get; set; } = "Defaults/globals.json";
 
-		/// <summary>
-		/// Settings for remote execution
-		/// </summary>
-		public RemoteExecSettings RemoteExecSettings { get; set; } = new RemoteExecSettings();
-
 		/// <summary>
 		/// Settings for the storage service
 		/// </summary>
diff --git a/Engine/Source/Programs/Horde/Horde.Build/Startup.cs b/Engine/Source/Programs/Horde/Horde.Build/Startup.cs
--- a/Engine/Source/Programs/Horde/Horde.Build/Startup.cs
+++ b/Engine/Source/Programs/Horde/Horde.Build/Startup.cs
@@ -561,15 +561,7 @@ namespace HordeServer
 
 			if (Settings.IsRunModeActive(RunMode.Worker))
 			{
-				if (Settings.FeatureFlags.AutoscaleServiceV1Enabled)
-				{
-					Services.AddHostedService(Provider => Provider.GetRequiredService<AutoscaleService>());	
-				}
-				
-				if (Settings.FeatureFlags.AutoscaleServiceV2Enabled)
-				{
-					Services.AddHostedService(Provider => Provider.GetRequiredService<AutoscaleServiceV2>());
-				}
+				Services.AddHostedService(Provider => Provider.GetRequiredService<AutoscaleServiceV2>());
 				
 				Services.AddHostedService(Provider => Provider.GetRequiredService<AgentService>());
 				Services.AddHostedService(Provider => Provider.GetRequiredService<CommitService>());
