Horde Storage - Return not found for the x-jupiter-inline type when the content id resolve fails

diff --git a/Engine/Source/Programs/Horde/HordeStorage/Horde.Storage/Controllers/ReferencesController.cs b/Engine/Source/Programs/Horde/HordeStorage/Horde.Storage/Controllers/ReferencesController.cs
--- a/Engine/Source/Programs/Horde/HordeStorage/Horde.Storage/Controllers/ReferencesController.cs
+++ b/Engine/Source/Programs/Horde/HordeStorage/Horde.Storage/Controllers/ReferencesController.cs
@@ -198,8 +198,21 @@ namespace Horde.Storage.Controllers
                         if (countOfBinaryAttachmentFields == 1 && countOfAttachmentFields == 1)
                         {
                             // fetch the blob so we can resolve any content ids in it
-                            IAsyncEnumerable<BlobIdentifier> referencedBlobsEnumerable = _referenceResolver.ResolveReferences(ns, cb);
-                            List<BlobIdentifier> referencedBlobs = await referencedBlobsEnumerable.ToListAsync();
+                            List<BlobIdentifier> referencedBlobs;
+                            try
+                            {
+                                IAsyncEnumerable<BlobIdentifier> referencedBlobsEnumerable = _referenceResolver.ResolveReferences(ns, cb);
+                                referencedBlobs = await referencedBlobsEnumerable.ToListAsync();
+                            }
+                            catch (PartialReferenceResolveException)
+                            {
+                                return NotFound(new ProblemDetails {Title = $"Object {bucket} {key} in namespace {ns} was missing some content ids"});
+                            }
+                            catch (ReferenceIsMissingBlobsException)
+                            {
+                                return NotFound(new ProblemDetails {Title = $"Object {bucket} {key} in namespace {ns} was missing some blobs"});
+                            }
+
                             if (referencedBlobs.Count == 1)
                             {
                                 BlobContents referencedBlobContents = await _blobStore.GetObject(ns, referencedBlobs.First());
