Horde: Fix target handing when the target contains an '='
#rnx
#skipci

diff --git a/Engine/Source/Programs/Horde/HordeDashboard/src/backend/JobDetails.ts b/Engine/Source/Programs/Horde/HordeDashboard/src/backend/JobDetails.ts
--- a/Engine/Source/Programs/Horde/HordeDashboard/src/backend/JobDetails.ts
+++ b/Engine/Source/Programs/Horde/HordeDashboard/src/backend/JobDetails.ts
@@ -48,8 +48,14 @@ export class JobDetails {
 
             const targetArgs = detailArgs.map(arg => arg.trim()).filter(arg => arg.toLowerCase().startsWith("-target="));
 
-            targetArgs.forEach(t => {
-                const target = t.split("=")[1]?.trim();
+            targetArgs.forEach(t => { 
+                
+                const index = t.indexOf("=");
+                if (index === -1) {
+                    return;
+                }
+                const target = t.slice(index+1)?.trim();
+                
                 if (target) {
                     targets.push(target);
                 }
