Horde: Filter label notifications to only include initiator on preflights
#skipci
#rnx

diff --git a/Engine/Source/Programs/Horde/Horde.Build/Notifications/Impl/NotificationService.cs b/Engine/Source/Programs/Horde/Horde.Build/Notifications/Impl/NotificationService.cs
--- a/Engine/Source/Programs/Horde/Horde.Build/Notifications/Impl/NotificationService.cs
+++ b/Engine/Source/Programs/Horde/Horde.Build/Notifications/Impl/NotificationService.cs
@@ -695,6 +695,13 @@ namespace HordeServer.Notifications.Impl
 					bool bFireTrigger = NewLabel.State == LabelState.Complete;
 
 					List<IUser> UsersToNotify = await GetUsersToNotify(EventId, TriggerId, bFireTrigger);
+
+					// filter preflight label notifications to only include initiator
+					if (UsersToNotify.Count > 0 && Job.PreflightChange != 0 && Job.StartedByUserId != null)
+					{
+						UsersToNotify = UsersToNotify.Where(x => x.Id == Job.StartedByUserId).ToList();
+					}
+
 					if (UsersToNotify.Count > 0)
 					{
 						SendLabelUpdateNotifications(Job, Stream, Graph, StepForNode, Graph.Labels[LabelIdx], LabelIdx, NewLabel.Outcome, UsersToNotify);
