Horde Storage - Added http header returned when using x-jupiter-inline type to return the hash of the inlined payload

diff --git a/Engine/Source/Programs/Horde/HordeStorage/Horde.Storage/Controllers/ReferencesController.cs b/Engine/Source/Programs/Horde/HordeStorage/Horde.Storage/Controllers/ReferencesController.cs
--- a/Engine/Source/Programs/Horde/HordeStorage/Horde.Storage/Controllers/ReferencesController.cs
+++ b/Engine/Source/Programs/Horde/HordeStorage/Horde.Storage/Controllers/ReferencesController.cs
@@ -215,7 +215,10 @@ namespace Horde.Storage.Controllers
 
                             if (referencedBlobs.Count == 1)
                             {
-                                BlobContents referencedBlobContents = await _blobStore.GetObject(ns, referencedBlobs.First());
+                                BlobIdentifier attachmentToSend = referencedBlobs.First();
+                                BlobContents referencedBlobContents = await _blobStore.GetObject(ns, attachmentToSend);
+                                Response.Headers[CommonHeaders.InlinePayloadHash] = attachmentToSend.ToString();
+
                                 await WriteBody(referencedBlobContents, CustomMediaTypeNames.JupiterInlinedPayload);
                                 return new EmptyResult();
                             }
@@ -481,7 +484,7 @@ namespace Horde.Storage.Controllers
                 });
             }
 
-            CompactBinaryObject payloadObject;
+            CbObject payloadObject;
             BlobIdentifier blobHeader = headerHash;
             try
             {
@@ -494,13 +497,13 @@ namespace Horde.Storage.Controllers
                         blobHeader = await _blobStore.PutObject(ns, payload, headerHash);
 
                         // TODO: convert the json object into a compact binary instead
-                        CompactBinaryWriter compactBinaryWriter = new CompactBinaryWriter();
-                        compactBinaryWriter.BeginObject();
-                        compactBinaryWriter.AddBinaryAttachment(blobHeader);
-                        compactBinaryWriter.EndObject();
+                        CbWriter writer = new CbWriter();
+                        writer.BeginObject();
+                        writer.WriteBinaryAttachmentValue(blobHeader.AsIoHash());
+                        writer.EndObject();
 
-                        byte[] blob = compactBinaryWriter.Save();
-                        payloadObject = CompactBinaryObject.Load(blob);
+                        byte[] blob = writer.ToByteArray();
+                        payloadObject = new CbObject(blob);
                         blobHeader = BlobIdentifier.FromBlob(blob);
                         break;
                     }
@@ -509,20 +512,21 @@ namespace Horde.Storage.Controllers
                         MemoryStream ms = new MemoryStream();
                         await using Stream payloadStream = payload.GetStream();
                         await payloadStream.CopyToAsync(ms);
-                        payloadObject = CompactBinaryObject.Load(ms.ToArray());
+                        payloadObject = new CbObject(ms.ToArray());
                         break;
                     }
                     case MediaTypeNames.Application.Octet:
                     {
                         blobHeader = await _blobStore.PutObject(ns, payload, headerHash);
 
-                        CompactBinaryWriter compactBinaryWriter = new CompactBinaryWriter();
-                        compactBinaryWriter.BeginObject();
-                        compactBinaryWriter.AddBinaryAttachment(blobHeader);
-                        compactBinaryWriter.EndObject();
+                        CbWriter writer = new CbWriter();
+                        writer.BeginObject();
+                        writer.WriteBinaryAttachment("RawHash", blobHeader.AsIoHash());
+                        writer.WriteInteger("RawSize", payload.Length);
+                        writer.EndObject();
 
-                        byte[] blob = compactBinaryWriter.Save();
-                        payloadObject = CompactBinaryObject.Load(blob);
+                        byte[] blob = writer.ToByteArray();
+                        payloadObject = new CbObject(blob);
                         blobHeader = BlobIdentifier.FromBlob(blob);
                         break;
                     }
@@ -634,7 +638,7 @@ namespace Horde.Storage.Controllers
             }
 
 
-            return Ok(new { CountOfDeletedRecords = countOfDeletedRecords});
+            return Ok(new BucketDeletedResponse(countOfDeletedRecords));
         }
 
         /// <summary>
@@ -662,7 +666,7 @@ namespace Horde.Storage.Controllers
             try
             {
                 bool deleted = await _objectService.Delete(ns, bucket, key);
-                return Ok(new { DeletedCount = deleted ? 1: 0 } );
+                return Ok(new RefDeletedResponse(deleted ? 1: 0));
             }
             catch (NamespaceNotFoundException e)
             {
@@ -671,8 +675,46 @@ namespace Horde.Storage.Controllers
         }
     }
 
+    public class RefDeletedResponse
+    {
+        public RefDeletedResponse()
+        {
+
+        }
+
+        public RefDeletedResponse(int deletedCount)
+        {
+            DeletedCount = deletedCount;
+        }
+
+        [CbField("deletedCount")]
+        public int DeletedCount { get; set; }
+    }
+
+    public class BucketDeletedResponse
+    {
+        public BucketDeletedResponse()
+        {
+
+        }
+
+        public BucketDeletedResponse(long countOfDeletedRecords)
+        {
+            CountOfDeletedRecords = countOfDeletedRecords;
+        }
+
+        [CbField("countOfDeletedRecords")]
+        public long CountOfDeletedRecords { get; set; }
+    }
+
     public class RefMetadataResponse
     {
+        public RefMetadataResponse()
+        {
+            PayloadIdentifier = null!;
+            InlinePayload = null!;
+        }
+
         [JsonConstructor]
         public RefMetadataResponse(NamespaceId ns, BucketId bucket, IoHashKey name, BlobIdentifier payloadIdentifier, DateTime lastAccess, bool isFinalized, byte[]? inlinePayload)
         {
@@ -696,22 +738,56 @@ namespace Horde.Storage.Controllers
             InlinePayload = objectRecord.InlinePayload;
         }
 
+        [CbField("ns")]
         public NamespaceId Ns { get; set; }
+        
+        [CbField("bucket")]
         public BucketId Bucket { get; set; }
+
+        [CbField("name")]
         public IoHashKey Name { get; set; }
+
+        [CbField("payloadIdentifier")]
         public BlobIdentifier PayloadIdentifier { get; set; }
+
+        [CbField("lastAccess")]
         public DateTime LastAccess { get; set; }
+
+        [CbField("isFinalized")]
         public bool IsFinalized { get; set; }
+
+        [CbField("inlinePayload")]
         public byte[]? InlinePayload { get; set; }
     }
 
     public class PutObjectResponse
     {
+        public PutObjectResponse()
+        {
+            Needs = null!;
+        }
+
         public PutObjectResponse(ContentHash[] missingReferences)
         {
             Needs = missingReferences;
         }
 
+        [CbField("needs")]
         public ContentHash[] Needs { get; set; }
     }
+
+    public class ExistCheckMultipleRefsResponse
+    {
+        public ExistCheckMultipleRefsResponse(List<IoHashKey> missingNames)
+        {
+            Needs = missingNames;
+            Names = missingNames;
+        }
+
+        [CbField("needs")]
+        public List<IoHashKey> Needs { get; set; }
+
+        [CbField("names")]
+        public List<IoHashKey> Names { get; set; }
+    }
 }
diff --git a/Engine/Source/Programs/Horde/HordeStorage/Horde.Storage/Implementation/CommonHeaders.cs b/Engine/Source/Programs/Horde/HordeStorage/Horde.Storage/Implementation/CommonHeaders.cs
--- a/Engine/Source/Programs/Horde/HordeStorage/Horde.Storage/Implementation/CommonHeaders.cs
+++ b/Engine/Source/Programs/Horde/HordeStorage/Horde.Storage/Implementation/CommonHeaders.cs
@@ -8,5 +8,7 @@ namespace Horde.Storage.Implementation
         public const string HashHeaderName = "X-Jupiter-IoHash";
 
         public const string LastAccessHeaderName = "X-Jupiter-LastAccess";
+
+        public const string InlinePayloadHash = "X-Jupiter-InlinePayloadHash";
     }
 }
