Horde Storage - Switched to EpicGames.Serialization for Compact Binary

EpicGames.Serialization does not allow writing anonymous fields in objects so we have had to change how certain objects are generated within Horde Storage. This is good as the objects is then also consistently convertable to json.

diff --git a/Engine/Source/Programs/Horde/HordeStorage/Horde.Storage/Controllers/ObjectController.cs b/Engine/Source/Programs/Horde/HordeStorage/Horde.Storage/Controllers/ObjectController.cs
--- a/Engine/Source/Programs/Horde/HordeStorage/Horde.Storage/Controllers/ObjectController.cs
+++ b/Engine/Source/Programs/Horde/HordeStorage/Horde.Storage/Controllers/ObjectController.cs
@@ -11,6 +11,7 @@ using System.Threading.Tasks;
 using async_enumerable_dotnet;
 using Datadog.Trace;
 using EpicGames.Horde.Storage;
+using EpicGames.Serialization;
 using Horde.Storage.Implementation;
 using Jupiter;
 using Jupiter.Common.Implementation;
@@ -142,10 +143,7 @@ namespace Horde.Storage.Controllers
             using IBufferedPayload payload = await _bufferedPayloadFactory.CreateFromRequest(Request);
 
             BlobIdentifier identifier = await _storage.PutObject(ns, payload, id);
-            return Ok(new
-            {
-                Identifier = identifier
-            });
+            return Ok(new PutBlobResponse(identifier));
         }
 
         [HttpGet("{ns}/{id}/references")]
@@ -177,7 +175,7 @@ namespace Horde.Storage.Controllers
             {
                 _logger.Warning("0 byte object found for {Id} {Namespace}", id, ns);
             }
-            CompactBinaryObject compactBinaryObject = CompactBinaryObject.Load(blobContents);
+            CbObject compactBinaryObject = new CbObject(blobContents);
 
             try
             {
@@ -235,6 +233,22 @@ namespace Horde.Storage.Controllers
         }
     }
 
+    public class PutBlobResponse
+    {
+        public PutBlobResponse()
+        {
+            Identifier = null!;
+        }
+
+        public PutBlobResponse(BlobIdentifier identifier)
+        {
+            Identifier = identifier;
+        }
+
+        [CbField("identifier")]
+        public BlobIdentifier Identifier { get; set; }
+    }
+
     public class DeletedResponse
     {
         public int DeletedCount { get; set; }
diff --git a/Engine/Source/Programs/Horde/HordeStorage/Horde.Storage/Controllers/ReferencesController.cs b/Engine/Source/Programs/Horde/HordeStorage/Horde.Storage/Controllers/ReferencesController.cs
--- a/Engine/Source/Programs/Horde/HordeStorage/Horde.Storage/Controllers/ReferencesController.cs
+++ b/Engine/Source/Programs/Horde/HordeStorage/Horde.Storage/Controllers/ReferencesController.cs
@@ -12,7 +12,9 @@ using System.Text;
 using System.Threading.Tasks;
 using Dasync.Collections;
 using Datadog.Trace;
+using EpicGames.Core;
 using EpicGames.Horde.Storage;
+using EpicGames.Serialization;
 using Horde.Storage.Implementation;
 using Jupiter;
 using Jupiter.Common.Implementation;
@@ -25,6 +27,8 @@ using Microsoft.AspNetCore.Mvc;
 using Microsoft.Extensions.Options;
 using Newtonsoft.Json;
 using Serilog;
+using ContentHash = Jupiter.Implementation.ContentHash;
+using Log = Serilog.Log;
 
 namespace Horde.Storage.Controllers
 {
@@ -152,20 +156,19 @@ namespace Horde.Storage.Controllers
                     case MediaTypeNames.Application.Octet:
                     {
                         byte[] blobMemory = await blob.Stream.ToByteArray();
-                        ReadOnlyMemory<byte> localMemory = new ReadOnlyMemory<byte>(blobMemory);
-                        CompactBinaryObject cb = CompactBinaryObject.Load(ref localMemory);
-                        List<CompactBinaryField> compactBinaryFields = cb.GetFields().ToList();
-                        // detect if this object was uploaded as a raw object, in which case we deconstruct the generated compact binary
-                        if (compactBinaryFields.Count == 1)
+                        CbObject cb = new CbObject(blobMemory);
+                        int countOfAttachmentFields = cb.Count(field => field.IsAttachment());
+                        CbField? binaryAttachmentField = cb.FirstOrDefault(field => field.IsBinaryAttachment());
+                        if (countOfAttachmentFields == 1 && binaryAttachmentField != null)
                         {
-                            CompactBinaryField field = compactBinaryFields[0];
-                            if (field.IsBinaryAttachment() && string.IsNullOrEmpty(field.Name))
-                            {
-                                // this is a very simple object, we just fetch the blob it references and return that
-                                BlobContents referencedBlobContents = await _blobStore.GetObject(ns, field.AsBinaryAttachment()!);
-                                await WriteBody(referencedBlobContents, MediaTypeNames.Application.Octet);
-                                break;
-                            }
+                            // there is a single attachment field and that is of the binary attachment type, fetch that attachment and return it instead of the compact binary
+                            // this is so that we match the uploaded that happened as a octet-stream which generates a small cb object with a single attachment
+
+                            IoHash hash = binaryAttachmentField.AsBinaryAttachment();
+
+                            BlobContents referencedBlobContents = await _blobStore.GetObject(ns, BlobIdentifier.FromIoHash(hash));
+                            await WriteBody(referencedBlobContents, MediaTypeNames.Application.Octet);
+                            break;
                         }
 
                         // this doesn't look like the generated compact binary so we just return the payload
@@ -179,8 +182,7 @@ namespace Horde.Storage.Controllers
                             using IScope scope = Tracer.Instance.StartActive("json.readblob");
                             blobMemory = await blob.Stream.ToByteArray();
                         }
-                        ReadOnlyMemory<byte> localMemory = new ReadOnlyMemory<byte>(blobMemory);
-                        CompactBinaryObject cb = CompactBinaryObject.Load(ref localMemory);
+                        CbObject cb = new CbObject(blobMemory);
                         string s = cb.ToJson();
                         await WriteBody(new BlobContents(Encoding.UTF8.GetBytes(s)), MediaTypeNames.Application.Json);
                         break;
@@ -188,10 +190,10 @@ namespace Horde.Storage.Controllers
                     }
                     case CustomMediaTypeNames.JupiterInlinedPayload:
                     {
+
                         byte[] blobMemory = await blob.Stream.ToByteArray();
-                        ReadOnlyMemory<byte> localMemory = new ReadOnlyMemory<byte>(blobMemory);
-                        CompactBinaryObject cb = CompactBinaryObject.Load(ref localMemory);
-                        List<CompactBinaryField> compactBinaryFields = cb.GetFields().ToList();
+                        CbObject cb = new CbObject(blobMemory);
+                        List<CbField> compactBinaryFields = cb.ToList();
                         int countOfBinaryAttachmentFields = compactBinaryFields.Count(field => field.IsBinaryAttachment());
                         int countOfAttachmentFields = compactBinaryFields.Count(field => field.IsAttachment());
                         // if the object consists of a single attachment field we return this attachment field instead
@@ -359,7 +361,7 @@ namespace Horde.Storage.Controllers
                 // we have to verify the blobs are available locally, as the record of the key is replicated a head of the content
                 // TODO: Once we support inline replication this step is not needed as at least one region as this blob, just maybe not this current one
                 byte[] blobContents = await blob.Stream.ToByteArray();
-                CompactBinaryObject compactBinaryObject = CompactBinaryObject.Load(blobContents);
+                CbObject compactBinaryObject = new CbObject(blobContents);
                 // the reference resolver will throw if any blob is missing, so no need to do anything other then process each reference
                 IAsyncEnumerable<BlobIdentifier> references = _referenceResolver.ResolveReferences(ns, compactBinaryObject);
                 List<BlobIdentifier>? _ = await references.ToListAsync();
@@ -425,9 +427,9 @@ namespace Horde.Storage.Controllers
                     // we have to verify the blobs are available locally, as the record of the key is replicated a head of the content
                     // TODO: Once we support inline replication this step is not needed as at least one region as this blob, just maybe not this current one
                     byte[] blobContents = await blob.Stream.ToByteArray();
-                    CompactBinaryObject compactBinaryObject = CompactBinaryObject.Load(blobContents);
+                    CbObject cb = new CbObject(blobContents);
                     // the reference resolver will throw if any blob is missing, so no need to do anything other then process each reference
-                    IAsyncEnumerable<BlobIdentifier> references = _referenceResolver.ResolveReferences(ns, compactBinaryObject);
+                    IAsyncEnumerable<BlobIdentifier> references = _referenceResolver.ResolveReferences(ns, cb);
                     List<BlobIdentifier>? _ = await references.ToListAsync();
                 }
                 catch (ObjectNotFoundException)
@@ -445,7 +447,7 @@ namespace Horde.Storage.Controllers
             });
             await Task.WhenAll(tasks);
 
-            return Ok(new { Needs = missingObject.ToArray()});
+            return Ok(new ExistCheckMultipleRefsResponse(missingObject.ToList()));
         }
 
 
diff --git a/Engine/Source/Programs/Horde/HordeStorage/Horde.Storage/Controllers/ReplicationLogController.cs b/Engine/Source/Programs/Horde/HordeStorage/Horde.Storage/Controllers/ReplicationLogController.cs
--- a/Engine/Source/Programs/Horde/HordeStorage/Horde.Storage/Controllers/ReplicationLogController.cs
+++ b/Engine/Source/Programs/Horde/HordeStorage/Horde.Storage/Controllers/ReplicationLogController.cs
@@ -7,6 +7,7 @@ using System.Threading;
 using System.Threading.Tasks;
 using async_enumerable_dotnet;
 using EpicGames.Horde.Storage;
+using EpicGames.Serialization;
 using Horde.Storage.Implementation;
 using Horde.Storage.Implementation.TransactionLog;
 using Jupiter;
@@ -74,7 +75,7 @@ namespace Horde.Storage.Controllers
 
             ReplicationLogSnapshotBuilder builder = ActivatorUtilities.CreateInstance<ReplicationLogSnapshotBuilder>(_provider);
             BlobIdentifier snapshotBlob = await builder.BuildSnapshot(ns, _snapshotSettings.CurrentValue.SnapshotStorageNamespace, CancellationToken.None);
-            return Ok(new { SnapshotBlobId = snapshotBlob });
+            return Ok(new SnapshotCreatedResponse(snapshotBlob));
         }
 
         [HttpGet("incremental/{ns}")]
@@ -143,6 +144,22 @@ namespace Horde.Storage.Controllers
         }
     }
 
+    public class SnapshotCreatedResponse
+    {
+        public SnapshotCreatedResponse()
+        {
+            SnapshotBlobId = null!;
+        }
+
+        public SnapshotCreatedResponse(BlobIdentifier snapshotBlob)
+        {
+            SnapshotBlobId = snapshotBlob;
+        }
+
+        [CbField("snapshotBlobId")]
+        public BlobIdentifier SnapshotBlobId { get; set; }
+    }
+
     public class ReplicationLogSnapshots
     {
         public ReplicationLogSnapshots()
diff --git a/Engine/Source/Programs/Horde/HordeStorage/Horde.Storage/Controllers/StorageController.cs b/Engine/Source/Programs/Horde/HordeStorage/Horde.Storage/Controllers/StorageController.cs
--- a/Engine/Source/Programs/Horde/HordeStorage/Horde.Storage/Controllers/StorageController.cs
+++ b/Engine/Source/Programs/Horde/HordeStorage/Horde.Storage/Controllers/StorageController.cs
@@ -10,6 +10,7 @@ using System.Net.Mime;
 using System.Threading.Tasks;
 using Datadog.Trace;
 using EpicGames.Horde.Storage;
+using EpicGames.Serialization;
 using Horde.Storage.Implementation;
 using Jupiter;
 using Jupiter.Common.Implementation;
@@ -362,6 +363,7 @@ namespace Horde.Storage.Controllers
 
     public class HeadMultipleResponse
     {
+        [CbField("needs")]
         public BlobIdentifier[] Needs { get; set; } = null!;
     }
 }
diff --git a/Engine/Source/Programs/Horde/HordeStorage/Horde.Storage/Horde.Storage.csproj b/Engine/Source/Programs/Horde/HordeStorage/Horde.Storage/Horde.Storage.csproj
--- a/Engine/Source/Programs/Horde/HordeStorage/Horde.Storage/Horde.Storage.csproj
+++ b/Engine/Source/Programs/Horde/HordeStorage/Horde.Storage/Horde.Storage.csproj
@@ -99,6 +99,8 @@
   </ItemGroup>
 
   <ItemGroup>
+    <ProjectReference Include="..\..\..\Shared\EpicGames.AspNet\EpicGames.AspNet.csproj" />
+    <ProjectReference Include="..\..\..\Shared\EpicGames.Serialization\EpicGames.Serialization.csproj" />
     <ProjectReference Include="..\Jupiter.Common\Jupiter.Common.csproj" />
   </ItemGroup>
 
diff --git a/Engine/Source/Programs/Horde/HordeStorage/Horde.Storage/Implementation/DDCRefService.cs b/Engine/Source/Programs/Horde/HordeStorage/Horde.Storage/Implementation/DDCRefService.cs
--- a/Engine/Source/Programs/Horde/HordeStorage/Horde.Storage/Implementation/DDCRefService.cs
+++ b/Engine/Source/Programs/Horde/HordeStorage/Horde.Storage/Implementation/DDCRefService.cs
@@ -7,6 +7,7 @@ using System.Linq;
 using System.Threading.Tasks;
 using Datadog.Trace;
 using EpicGames.Horde.Storage;
+using EpicGames.Serialization;
 using Horde.Storage.Controllers;
 using Jupiter.Implementation;
 using Microsoft.Extensions.Options;
@@ -16,6 +17,15 @@ namespace Horde.Storage.Implementation
 {
     public class RefResponse
     {
+        public RefResponse()
+        {
+            Name = null!;
+            LastAccessTime = null!;
+            Metadata = null!;
+            ContentHash = null!;
+            BlobIdentifiers = null!;
+        }
+
         public RefResponse(string name, DateTime? lastAccessTime, ContentHash contentHash, BlobIdentifier[] blobIdentifiers, Dictionary<string, object>? metadata)
         {
             Name = name;
@@ -25,13 +35,22 @@ namespace Horde.Storage.Implementation
             BlobIdentifiers = blobIdentifiers;
         }
 
-        public string Name { get; }
-        public DateTime? LastAccessTime { get; }
-        public Dictionary<string, object>? Metadata { get; }
+        [CbField("name")]
+        public string Name { get; set; }
+
+        [CbField("lastAccessTime")]
+        public DateTime? LastAccessTime { get; set; }
+
+        [CbField("metadata")]
+        public Dictionary<string, object>? Metadata { get; set; }
+
+        [CbField("contentHash")]
+        public ContentHash ContentHash { get; set; }
 
-        public ContentHash ContentHash { get; }
-        public BlobIdentifier[] BlobIdentifiers { get; }
+        [CbField("blobIdentifiers")]
+        public BlobIdentifier[] BlobIdentifiers { get; set; }
 
+        [CbField("blob")]
         public byte[]? Blob { get; set; }
     }
 
diff --git a/Engine/Source/Programs/Horde/HordeStorage/Horde.Storage/Implementation/IReferenceResolver.cs b/Engine/Source/Programs/Horde/HordeStorage/Horde.Storage/Implementation/IReferenceResolver.cs
--- a/Engine/Source/Programs/Horde/HordeStorage/Horde.Storage/Implementation/IReferenceResolver.cs
+++ b/Engine/Source/Programs/Horde/HordeStorage/Horde.Storage/Implementation/IReferenceResolver.cs
@@ -1,12 +1,12 @@
 // Copyright Epic Games, Inc. All Rights Reserved.
 
 using System;
-using System.Collections.Concurrent;
 using System.Collections.Generic;
-using System.Reflection.Metadata;
 using System.Threading.Tasks;
 using Datadog.Trace;
+using EpicGames.Core;
 using EpicGames.Horde.Storage;
+using EpicGames.Serialization;
 using Jupiter.Implementation;
 using Jupiter.Utils;
 
@@ -14,7 +14,7 @@ namespace Horde.Storage.Implementation
 {
     public interface IReferenceResolver
     {
-        IAsyncEnumerable<BlobIdentifier> ResolveReferences(NamespaceId ns, CompactBinaryObject cb);
+        IAsyncEnumerable<BlobIdentifier> ResolveReferences(NamespaceId ns, CbObject cb);
     }
 
     public class ReferenceResolver : IReferenceResolver
@@ -28,34 +28,35 @@ namespace Horde.Storage.Implementation
             _contentIdStore = contentIdStore;
         }
 
-        public async IAsyncEnumerable<BlobIdentifier> ResolveReferences(NamespaceId ns, CompactBinaryObject cb)
+        public async IAsyncEnumerable<BlobIdentifier> ResolveReferences(NamespaceId ns, CbObject cb)
         {
             // TODO: This is cacheable and we should store the result somewhere
-            Queue<CompactBinaryObject> objectsToVisit = new Queue<CompactBinaryObject>();
+            Queue<CbObject> objectsToVisit = new Queue<CbObject>();
             objectsToVisit.Enqueue(cb);
             List<ContentId> unresolvedContentIdReferences = new List<ContentId>();
             List<BlobIdentifier> unresolvedBlobReferences = new List<BlobIdentifier>();
 
             List<Task<(ContentId, BlobIdentifier[]?)>> pendingContentIdResolves = new();
-            List<Task<CompactBinaryObject>> pendingCompactBinaryAttachments = new();
+            List<Task<CbObject>> pendingCompactBinaryAttachments = new();
 
             while(pendingCompactBinaryAttachments.Count != 0 || pendingContentIdResolves.Count != 0 || objectsToVisit.Count != 0)
             {
-                if (objectsToVisit.TryDequeue(out CompactBinaryObject? parent))
+                if (objectsToVisit.TryDequeue(out CbObject? parent))
                 {
                     // enumerate all fields in the compact binary and start to resolve their dependencies
-                    foreach (CompactBinaryField field in parent.GetAllFields())
+                    foreach (CbField field in parent)
                     {
                         if (!field.IsAttachment())
                             continue;
 
-                        BlobIdentifier? blobIdentifier = field.AsAttachment();
-                        if (blobIdentifier == null)
+                        IoHash attachmentHash = field.AsAttachment();
+                        if (field.HasError())
                             continue;
 
+                        BlobIdentifier blobIdentifier = BlobIdentifier.FromIoHash(attachmentHash);
                         if (field.IsBinaryAttachment())
                         {
-                            pendingContentIdResolves.Add(ResolveContentId(ns, ContentId.FromBlobIdentifier(blobIdentifier)));
+                            pendingContentIdResolves.Add(ResolveContentId(ns, ContentId.FromIoHash(attachmentHash)));
                         }
                         else
                         {
@@ -64,7 +65,7 @@ namespace Horde.Storage.Implementation
                         }
 
                         // if its a reference to another compact binary we need to fetch it and resolve it
-                        if (!field.IsCompactBinaryAttachment())
+                        if (!field.IsObjectAttachment())
                             continue;
 
                         pendingCompactBinaryAttachments.Add(ParseCompactBinaryAttachment(ns, blobIdentifier));
@@ -101,14 +102,14 @@ namespace Horde.Storage.Implementation
                 }
 
                 // check for any compact binary attachment fetches and add those to the objects we are handling
-                List<Task<CompactBinaryObject>> finishedCompactBinaryResolves = new();
-                foreach (Task<CompactBinaryObject> pendingCompactBinaryAttachment in pendingCompactBinaryAttachments)
+                List<Task<CbObject>> finishedCompactBinaryResolves = new();
+                foreach (Task<CbObject> pendingCompactBinaryAttachment in pendingCompactBinaryAttachments)
                 {
                     if (pendingCompactBinaryAttachment.IsCompleted)
                     {
                         try
                         {
-                            CompactBinaryObject childBinaryObject = await pendingCompactBinaryAttachment;
+                            CbObject childBinaryObject = await pendingCompactBinaryAttachment;
                             objectsToVisit.Enqueue(childBinaryObject);
                         }
                         catch (BlobNotFoundException e)
@@ -120,7 +121,7 @@ namespace Horde.Storage.Implementation
                 }
 
                 // cleanup finished tasks
-                foreach (Task<CompactBinaryObject> finishedTask in finishedCompactBinaryResolves)
+                foreach (Task<CbObject> finishedTask in finishedCompactBinaryResolves)
                 {
                     pendingCompactBinaryAttachments.Remove(finishedTask);
                 }
@@ -137,11 +138,11 @@ namespace Horde.Storage.Implementation
             }
         }
 
-        private async Task<CompactBinaryObject> ParseCompactBinaryAttachment(NamespaceId ns, BlobIdentifier blobIdentifier)
+        private async Task<CbObject> ParseCompactBinaryAttachment(NamespaceId ns, BlobIdentifier blobIdentifier)
         {
             BlobContents contents = await _blobStore.GetObject(ns, blobIdentifier);
             byte[] data = await contents.Stream.ToByteArray();
-            CompactBinaryObject childBinaryObject = CompactBinaryObject.Load(data);
+            CbObject childBinaryObject = new CbObject(data);
 
             return childBinaryObject;
         }
diff --git a/Engine/Source/Programs/Horde/HordeStorage/Horde.Storage/Implementation/ObjectService.cs b/Engine/Source/Programs/Horde/HordeStorage/Horde.Storage/Implementation/ObjectService.cs
--- a/Engine/Source/Programs/Horde/HordeStorage/Horde.Storage/Implementation/ObjectService.cs
+++ b/Engine/Source/Programs/Horde/HordeStorage/Horde.Storage/Implementation/ObjectService.cs
@@ -8,6 +8,7 @@ using System.Threading.Tasks;
 using async_enumerable_dotnet;
 using Datadog.Trace;
 using EpicGames.Horde.Storage;
+using EpicGames.Serialization;
 using Horde.Storage.Implementation.Blob;
 using Jupiter.Implementation;
 using Jupiter.Utils;
@@ -18,7 +19,7 @@ namespace Horde.Storage.Implementation
     public interface IObjectService
     {
         Task<(ObjectRecord, BlobContents?)> Get(NamespaceId ns, BucketId bucket, IoHashKey key, string[] fields);
-        Task<(ContentId[], BlobIdentifier[])> Put(NamespaceId ns, BucketId bucket, IoHashKey key, BlobIdentifier blobHash, CompactBinaryObject payload);
+        Task<(ContentId[], BlobIdentifier[])> Put(NamespaceId ns, BucketId bucket, IoHashKey key, BlobIdentifier blobHash, CbObject payload);
         Task<(ContentId[], BlobIdentifier[])> Finalize(NamespaceId ns, BucketId bucket, IoHashKey key, BlobIdentifier blobHash);
 
         IAsyncEnumerable<NamespaceId> GetNamespaces();
@@ -91,16 +92,16 @@ namespace Horde.Storage.Implementation
             return (o, blobContents);
         }
 
-        public async Task<(ContentId[], BlobIdentifier[])> Put(NamespaceId ns, BucketId bucket, IoHashKey key, BlobIdentifier blobHash, CompactBinaryObject payload)
+        public async Task<(ContentId[], BlobIdentifier[])> Put(NamespaceId ns, BucketId bucket, IoHashKey key, BlobIdentifier blobHash, CbObject payload)
         {
-            bool hasReferences = payload.GetAllFields().Any(field => field.IsAttachment());
+            bool hasReferences = payload.Any(field => field.IsAttachment());
 
             // if we have no references we are always finalized, e.g. there are no referenced blobs to upload
             bool isFinalized = !hasReferences;
 
-            Task objectStorePut = _referencesStore.Put(ns, bucket, key, blobHash, payload.Data, isFinalized);
+            Task objectStorePut = _referencesStore.Put(ns, bucket, key, blobHash, payload.GetView().ToArray(), isFinalized);
 
-            Task<BlobIdentifier> blobStorePut = _blobService.PutObject(ns, payload.Data, blobHash);
+            Task<BlobIdentifier> blobStorePut = _blobService.PutObject(ns, payload.GetView().ToArray(), blobHash);
             Task addRefToBodyTask = _blobIndex.AddRefToBlobs(ns, bucket, key, new [] {blobHash});
             ContentId[] missingReferences = Array.Empty<ContentId>();
             BlobIdentifier[] missingBlobs = Array.Empty<BlobIdentifier>();
@@ -157,12 +158,12 @@ namespace Horde.Storage.Implementation
             if (blob == null)
                 throw new InvalidOperationException("No blob when attempting to finalize");
             byte[] blobContents = await blob.Stream.ToByteArray();
-            CompactBinaryObject payload = CompactBinaryObject.Load(blobContents);
+            CbObject payload = new CbObject(blobContents);
 
             if (!o.BlobIdentifier.Equals(blobHash))
                 throw new ObjectHashMismatchException(ns, bucket, key, blobHash, o.BlobIdentifier);
 
-            bool hasReferences = payload.GetAllFields().Any(field => field.IsAttachment());
+            bool hasReferences = payload.Any(field => field.IsAttachment());
 
             ContentId[] missingReferences = Array.Empty<ContentId>();
             BlobIdentifier[] missingBlobs = Array.Empty<BlobIdentifier>();
diff --git a/Engine/Source/Programs/Horde/HordeStorage/Horde.Storage/Implementation/TransactionLog/ReplicationLogSnapshotBuilder.cs b/Engine/Source/Programs/Horde/HordeStorage/Horde.Storage/Implementation/TransactionLog/ReplicationLogSnapshotBuilder.cs
--- a/Engine/Source/Programs/Horde/HordeStorage/Horde.Storage/Implementation/TransactionLog/ReplicationLogSnapshotBuilder.cs
+++ b/Engine/Source/Programs/Horde/HordeStorage/Horde.Storage/Implementation/TransactionLog/ReplicationLogSnapshotBuilder.cs
@@ -7,6 +7,7 @@ using System.Linq;
 using System.Threading;
 using System.Threading.Tasks;
 using EpicGames.Horde.Storage;
+using EpicGames.Serialization;
 using Jupiter.Common.Implementation;
 using Jupiter.Implementation;
 
@@ -84,13 +85,13 @@ namespace Horde.Storage.Implementation.TransactionLog
                     blobIdentifier = await BlobIdentifier.FromStream(stream);
                 }
 
-                CompactBinaryWriter writer = new CompactBinaryWriter();
+                CbWriter writer = new CbWriter();
                 writer.BeginObject();
-                writer.AddBinaryAttachment(blobIdentifier, "snapshotBlob");
-                writer.AddDateTime(DateTime.Now, "timestamp");
+                writer.WriteBinaryAttachment("snapshotBlob", blobIdentifier.AsIoHash());
+                writer.WriteDateTime("timestamp", DateTime.Now);
                 writer.EndObject();
 
-                byte[] cbObjectBytes = writer.Save();
+                byte[] cbObjectBytes = writer.ToByteArray();
                 BlobIdentifier cbBlobId = BlobIdentifier.FromBlob(cbObjectBytes);
 
                 if (cancellationToken.IsCancellationRequested)
@@ -100,7 +101,7 @@ namespace Horde.Storage.Implementation.TransactionLog
                 await _blobService.PutObject(storeInNamespace, payload, blobIdentifier);
             
                 tempFile.Delete();
-                (ContentId[] missingContentIds, BlobIdentifier[] missingBlobs) = await _objectService.Put(storeInNamespace, new BucketId("snapshot"), new IoHashKey(blobIdentifier.ToString()), cbBlobId, CompactBinaryObject.Load(cbObjectBytes));
+                (ContentId[] missingContentIds, BlobIdentifier[] missingBlobs) = await _objectService.Put(storeInNamespace, new BucketId("snapshot"), new IoHashKey(blobIdentifier.ToString()), cbBlobId, new CbObject(cbObjectBytes));
                 List<ContentHash> missingHashes = new List<ContentHash>(missingContentIds);
                 missingHashes.AddRange(missingBlobs);
                 if (missingHashes.Count != 0)
diff --git a/Engine/Source/Programs/Horde/HordeStorage/Horde.Storage/Tests/Functional/Blob/BlobIndexTests.cs b/Engine/Source/Programs/Horde/HordeStorage/Horde.Storage/Tests/Functional/Blob/BlobIndexTests.cs
--- a/Engine/Source/Programs/Horde/HordeStorage/Horde.Storage/Tests/Functional/Blob/BlobIndexTests.cs
+++ b/Engine/Source/Programs/Horde/HordeStorage/Horde.Storage/Tests/Functional/Blob/BlobIndexTests.cs
@@ -10,6 +10,7 @@ using System.Text;
 using System.Threading.Tasks;
 using Cassandra;
 using EpicGames.Horde.Storage;
+using EpicGames.Serialization;
 using Horde.Storage.Implementation;
 using Horde.Storage.Implementation.Blob;
 using Jupiter;
@@ -92,12 +93,12 @@ namespace Horde.Storage.FunctionalTests.Storage
         [TestMethod]
         public async Task UploadRef()
         {
-            CompactBinaryWriter writer = new CompactBinaryWriter();
+            CbWriter writer = new CbWriter();
             writer.BeginObject();
-            writer.AddString("thisIsAField", "stringField");
+            writer.WriteString("stringField","thisIsAField");
             writer.EndObject();
 
-            byte[] objectData = writer.Save();
+            byte[] objectData = writer.ToByteArray();
             BlobIdentifier objectHash = BlobIdentifier.FromBlob(objectData);
 
             HttpContent requestContent = new ByteArrayContent(objectData);
diff --git a/Engine/Source/Programs/Horde/HordeStorage/Horde.Storage/Tests/Functional/GC/GCBlobTests.cs b/Engine/Source/Programs/Horde/HordeStorage/Horde.Storage/Tests/Functional/GC/GCBlobTests.cs
--- a/Engine/Source/Programs/Horde/HordeStorage/Horde.Storage/Tests/Functional/GC/GCBlobTests.cs
+++ b/Engine/Source/Programs/Horde/HordeStorage/Horde.Storage/Tests/Functional/GC/GCBlobTests.cs
@@ -22,6 +22,7 @@ using Serilog;
 using Serilog.Core;
 using Moq;
 using EpicGames.Horde.Storage;
+using EpicGames.Serialization;
 
 namespace Horde.Storage.FunctionalTests.GC
 {
@@ -283,16 +284,16 @@ namespace Horde.Storage.FunctionalTests.GC
             BucketId testBucket = new BucketId("test");
             IObjectService? objectService = server.Services.GetService<IObjectService>()!;
             Assert.IsNotNull(objectService);
-            (BlobIdentifier ob0_hash, CompactBinaryObject ob0_cb) = GetCBWithAttachment(object0id);
+            (BlobIdentifier ob0_hash, CbObject ob0_cb) = GetCBWithAttachment(object0id);
             await objectService.Put(TestNamespace, testBucket, IoHashKey.FromName("object0"), ob0_hash, ob0_cb);
            
-            (BlobIdentifier ob2_hash, CompactBinaryObject ob2_cb) = GetCBWithAttachment(object2id);
+            (BlobIdentifier ob2_hash, CbObject ob2_cb) = GetCBWithAttachment(object2id);
             await objectService.Put(TestNamespace, testBucket, IoHashKey.FromName("object2"), ob2_hash, ob2_cb);
 
-            (BlobIdentifier ob3_hash, CompactBinaryObject ob3_cb) = GetCBWithAttachment(object3id);
+            (BlobIdentifier ob3_hash, CbObject ob3_cb) = GetCBWithAttachment(object3id);
             await objectService.Put(TestNamespace, testBucket, IoHashKey.FromName("object3"), ob3_hash, ob3_cb);
 
-            (BlobIdentifier ob6_hash, CompactBinaryObject ob6_cb) = GetCBWithAttachment(object6id);
+            (BlobIdentifier ob6_hash, CbObject ob6_cb) = GetCBWithAttachment(object6id);
             await objectService.Put(TestNamespace, testBucket, IoHashKey.FromName("object6"), ob6_hash, ob6_cb);
 
             IReferencesStore referenceStore = server.Services.GetService<IReferencesStore>()!;
@@ -317,15 +318,15 @@ namespace Horde.Storage.FunctionalTests.GC
 
         protected abstract string GetImplementation();
 
-        private (BlobIdentifier, CompactBinaryObject) GetCBWithAttachment(BlobIdentifier blobIdentifier)
+        private (BlobIdentifier, CbObject) GetCBWithAttachment(BlobIdentifier blobIdentifier)
         {
-            CompactBinaryWriter compactBinaryWriter = new CompactBinaryWriter();
-            compactBinaryWriter.BeginObject();
-            compactBinaryWriter.AddBinaryAttachment(blobIdentifier);
-            compactBinaryWriter.EndObject();
+            CbWriter writer = new CbWriter();
+            writer.BeginObject();
+            writer.WriteBinaryAttachment("Attachment", blobIdentifier.AsIoHash());
+            writer.EndObject();
 
-            byte[] b = compactBinaryWriter.Save();
-            return (BlobIdentifier.FromBlob(b), CompactBinaryObject.Load(b));
+            byte[] b = writer.ToByteArray();
+            return (BlobIdentifier.FromBlob(b), new CbObject(b));
         }
 
         [TestMethod]
diff --git a/Engine/Source/Programs/Horde/HordeStorage/Horde.Storage/Tests/Functional/GC/GCRefTests.cs b/Engine/Source/Programs/Horde/HordeStorage/Horde.Storage/Tests/Functional/GC/GCRefTests.cs
--- a/Engine/Source/Programs/Horde/HordeStorage/Horde.Storage/Tests/Functional/GC/GCRefTests.cs
+++ b/Engine/Source/Programs/Horde/HordeStorage/Horde.Storage/Tests/Functional/GC/GCRefTests.cs
@@ -19,6 +19,7 @@ using Moq;
 using Serilog;
 using Serilog.Core;
 using EpicGames.Horde.Storage;
+using EpicGames.Serialization;
 
 namespace Horde.Storage.FunctionalTests.GC
 {
@@ -177,25 +178,25 @@ namespace Horde.Storage.FunctionalTests.GC
 
             IObjectService? objectService = server.Services.GetService<IObjectService>()!;
             Assert.IsNotNull(objectService);
-            (BlobIdentifier ob0_hash, CompactBinaryObject ob0_cb) = GetCBWithAttachment(object0id);
+            (BlobIdentifier ob0_hash, CbObject ob0_cb) = GetCBWithAttachment(object0id);
             await objectService.Put(TestNamespace, DefaultBucket, object0Name, ob0_hash, ob0_cb);
            
-            (BlobIdentifier ob1_hash, CompactBinaryObject ob1_cb) = GetCBWithAttachment(object1id);
+            (BlobIdentifier ob1_hash, CbObject ob1_cb) = GetCBWithAttachment(object1id);
             await objectService.Put(TestNamespace, DefaultBucket, object1Name, ob1_hash, ob1_cb);
 
-            (BlobIdentifier ob2_hash, CompactBinaryObject ob2_cb) = GetCBWithAttachment(object2id);
+            (BlobIdentifier ob2_hash, CbObject ob2_cb) = GetCBWithAttachment(object2id);
             await objectService.Put(TestNamespace, DefaultBucket, object2Name, ob2_hash, ob2_cb);
 
-            (BlobIdentifier ob3_hash, CompactBinaryObject ob3_cb) = GetCBWithAttachment(object3id);
+            (BlobIdentifier ob3_hash, CbObject ob3_cb) = GetCBWithAttachment(object3id);
             await objectService.Put(TestNamespace, DefaultBucket, object3Name, ob3_hash, ob3_cb);
 
-            (BlobIdentifier ob4_hash, CompactBinaryObject ob4_cb) = GetCBWithAttachment(object4id);
+            (BlobIdentifier ob4_hash, CbObject ob4_cb) = GetCBWithAttachment(object4id);
             await objectService.Put(TestNamespace, DefaultBucket, object4Name, ob4_hash, ob4_cb);
 
-            (BlobIdentifier ob5_hash, CompactBinaryObject ob5_cb) = GetCBWithAttachment(object5id);
+            (BlobIdentifier ob5_hash, CbObject ob5_cb) = GetCBWithAttachment(object5id);
             await objectService.Put(TestNamespace, DefaultBucket, object5Name, ob5_hash, ob5_cb);
 
-            (BlobIdentifier ob6_hash, CompactBinaryObject ob6_cb) = GetCBWithAttachment(object6id);
+            (BlobIdentifier ob6_hash, CbObject ob6_cb) = GetCBWithAttachment(object6id);
             await objectService.Put(TestNamespace, DefaultBucket, object6Name, ob6_hash, ob6_cb);
 
             IReferencesStore referenceStore = server.Services.GetService<IReferencesStore>()!;
@@ -236,15 +237,15 @@ namespace Horde.Storage.FunctionalTests.GC
             }
         }
 
-        private (BlobIdentifier, CompactBinaryObject) GetCBWithAttachment(BlobIdentifier blobIdentifier)
+        private (BlobIdentifier, CbObject) GetCBWithAttachment(BlobIdentifier blobIdentifier)
         {
-            CompactBinaryWriter compactBinaryWriter = new CompactBinaryWriter();
-            compactBinaryWriter.BeginObject();
-            compactBinaryWriter.AddBinaryAttachment(blobIdentifier);
-            compactBinaryWriter.EndObject();
+            CbWriter writer = new CbWriter();
+            writer.BeginObject();
+            writer.WriteBinaryAttachment("Attachment", blobIdentifier.AsIoHash());
+            writer.EndObject();
 
-            byte[] b = compactBinaryWriter.Save();
-            return (BlobIdentifier.FromBlob(b), CompactBinaryObject.Load(b));
+            byte[] b = writer.ToByteArray();
+            return (BlobIdentifier.FromBlob(b), new CbObject(b));
         }
     }
 }
diff --git a/Engine/Source/Programs/Horde/HordeStorage/Horde.Storage/Tests/Functional/Objects/ReferencesTests.cs b/Engine/Source/Programs/Horde/HordeStorage/Horde.Storage/Tests/Functional/Objects/ReferencesTests.cs
--- a/Engine/Source/Programs/Horde/HordeStorage/Horde.Storage/Tests/Functional/Objects/ReferencesTests.cs
+++ b/Engine/Source/Programs/Horde/HordeStorage/Horde.Storage/Tests/Functional/Objects/ReferencesTests.cs
@@ -12,6 +12,7 @@ using System.Text;
 using System.Threading.Tasks;
 using Cassandra;
 using Dasync.Collections;
+using EpicGames.Core;
 using Horde.Storage.Controllers;
 using Horde.Storage.Implementation;
 using Jupiter;
@@ -27,6 +28,7 @@ using Newtonsoft.Json.Linq;
 using Serilog;
 using Logger = Serilog.Core.Logger;
 using EpicGames.Horde.Storage;
+using EpicGames.Serialization;
 
 namespace Horde.Storage.FunctionalTests.References
 {
@@ -249,15 +251,14 @@ namespace Horde.Storage.FunctionalTests.References
                     await getResponse.Content.CopyToAsync(ms);
 
                     byte[] roundTrippedBuffer = ms.ToArray();
-                    ReadOnlyMemory<byte> localMemory = new ReadOnlyMemory<byte>(roundTrippedBuffer);
-                    CompactBinaryObject cb = CompactBinaryObject.Load(ref localMemory);
-                    List<CompactBinaryField> fields = cb.GetFields().ToList();
+                    CbObject cb = new CbObject(roundTrippedBuffer);
+                    List<CbField> fields = cb.ToList();
 
-                    Assert.AreEqual(1, fields.Count);
-                    CompactBinaryField payloadField = fields[0];
+                    Assert.AreEqual(2, fields.Count);
+                    CbField payloadField = fields[0];
                     Assert.IsNotNull(payloadField);
                     Assert.IsTrue(payloadField.IsBinaryAttachment());
-                    attachment = payloadField.AsBinaryAttachment()!;
+                    attachment = BlobIdentifier.FromIoHash(payloadField.AsBinaryAttachment());
                 }
 
                 {
@@ -282,14 +283,14 @@ namespace Horde.Storage.FunctionalTests.References
                 byte[] roundTrippedBuffer = ms.ToArray();
                 string s = Encoding.ASCII.GetString(roundTrippedBuffer);
                 JObject jObject = JObject.Parse(s);
-                Assert.AreEqual(1, jObject.Children().Count());
+                Assert.AreEqual(2, jObject.Children().Count());
                 JToken? childToken = jObject.Children().First();
                 Assert.IsNotNull(childToken);
                 Assert.AreEqual(JTokenType.Property, childToken.Type);
                 JProperty property = jObject.Children<JProperty>().First();
                 Assert.IsNotNull(property);
 
-                Assert.AreEqual("", property!.Name);
+                Assert.AreEqual("RawHash", property!.Name);
 
                 string value = property.Value.Value<string>()!;
                 Assert.IsNotNull(value);
@@ -310,14 +311,14 @@ namespace Horde.Storage.FunctionalTests.References
                 byte[] roundTrippedBuffer = ms.ToArray();
                 string s = Encoding.ASCII.GetString(roundTrippedBuffer);
                 JObject jObject = JObject.Parse(s);
-                Assert.AreEqual(1, jObject.Children().Count());
+                Assert.AreEqual(2, jObject.Children().Count());
                 JToken? childToken = jObject.Children().First();
                 Assert.IsNotNull(childToken);
                 Assert.AreEqual(JTokenType.Property, childToken.Type);
                 JProperty property = jObject.Children<JProperty>().First();
                 Assert.IsNotNull(property);
 
-                Assert.AreEqual("", property!.Name);
+                Assert.AreEqual("RawHash", property!.Name);
 
                 string? value = property.Value.Value<string>()!;
                 Assert.IsNotNull(value);
@@ -346,12 +347,12 @@ namespace Horde.Storage.FunctionalTests.References
         [TestMethod]
         public async Task PutGetCompactBinary()
         {
-            CompactBinaryWriter writer = new CompactBinaryWriter();
+            CbWriter writer = new CbWriter();
             writer.BeginObject();
-            writer.AddString("thisIsAField", "stringField");
+            writer.WriteString("stringField", "thisIsAField");
             writer.EndObject();
 
-            byte[] objectData = writer.Save();
+            byte[] objectData = writer.ToByteArray();
             BlobIdentifier objectHash = BlobIdentifier.FromBlob(objectData);
             IoHashKey key = IoHashKey.FromName("newReferenceObject");
 
@@ -368,11 +369,10 @@ namespace Horde.Storage.FunctionalTests.References
                 await using MemoryStream ms = new MemoryStream();
                 await result.Content.CopyToAsync(ms);
                 byte[] roundTrippedBuffer = ms.ToArray();
-                ReadOnlyMemory<byte> localMemory = new ReadOnlyMemory<byte>(roundTrippedBuffer);
-                CompactBinaryObject cb = CompactBinaryObject.Load(ref localMemory);
-                CompactBinaryField? needsField = cb["needs"];
-                Assert.IsNotNull(needsField);
-                List<BlobIdentifier?> missingBlobs = needsField!.AsArray().Select(field => field.AsHash()).ToList();
+                CbObject cb = new CbObject(roundTrippedBuffer);
+                CbField needsField = cb["needs"];
+                Assert.AreNotEqual(CbField.Empty, needsField);
+                List<BlobIdentifier> missingBlobs = needsField.AsArray().Select(field => BlobIdentifier.FromIoHash(field.AsHash())).ToList();
                 Assert.AreEqual(0, missingBlobs.Count);
             }
 
@@ -406,12 +406,11 @@ namespace Horde.Storage.FunctionalTests.References
                 await getResponse.Content.CopyToAsync(ms);
 
                 byte[] roundTrippedBuffer = ms.ToArray();
-                ReadOnlyMemory<byte> localMemory = new ReadOnlyMemory<byte>(roundTrippedBuffer);
-                CompactBinaryObject cb = CompactBinaryObject.Load(ref localMemory);
-                List<CompactBinaryField> fields = cb.GetFields().ToList();
+                CbObject cb = new CbObject(roundTrippedBuffer);
+                List<CbField> fields = cb.ToList();
 
                 Assert.AreEqual(1, fields.Count);
-                CompactBinaryField stringField = fields[0];
+                CbField stringField = fields[0];
                 Assert.AreEqual("stringField", stringField.Name);
                 Assert.AreEqual("thisIsAField", stringField.AsString());
             }
@@ -447,18 +446,18 @@ namespace Horde.Storage.FunctionalTests.References
         [TestMethod]
         public async Task PutGetCompactBinaryHierarchy()
         {
-            CompactBinaryWriter childObjectWriter = new CompactBinaryWriter();
+            CbWriter childObjectWriter = new CbWriter();
             childObjectWriter.BeginObject();
-            childObjectWriter.AddString("thisIsAField", "stringField");
+            childObjectWriter.WriteString("stringField", "thisIsAField");
             childObjectWriter.EndObject();
-            byte[] childObjectData = childObjectWriter.Save();
+            byte[] childObjectData = childObjectWriter.ToByteArray();
             BlobIdentifier childObjectHash = BlobIdentifier.FromBlob(childObjectData);
 
-            CompactBinaryWriter parentObjectWriter = new CompactBinaryWriter();
+            CbWriter parentObjectWriter = new CbWriter();
             parentObjectWriter.BeginObject();
-            parentObjectWriter.AddCompactBinaryAttachment(childObjectHash, "childObject");
+            parentObjectWriter.WriteObjectAttachment("childObject", childObjectHash.AsIoHash());
             parentObjectWriter.EndObject();
-            byte[] parentObjectData = parentObjectWriter.Save();
+            byte[] parentObjectData = parentObjectWriter.ToByteArray();
             BlobIdentifier parentObjectHash = BlobIdentifier.FromBlob(parentObjectData);
 
             IoHashKey key = IoHashKey.FromName("newHierarchyObject");
@@ -476,11 +475,10 @@ namespace Horde.Storage.FunctionalTests.References
                 await using MemoryStream ms = new MemoryStream();
                 await result.Content.CopyToAsync(ms);
                 byte[] roundTrippedBuffer = ms.ToArray();
-                ReadOnlyMemory<byte> localMemory = new ReadOnlyMemory<byte>(roundTrippedBuffer);
-                CompactBinaryObject cb = CompactBinaryObject.Load(ref localMemory);
-                CompactBinaryField? needsField = cb["needs"];
-                Assert.IsNotNull(needsField);
-                List<BlobIdentifier?> missingBlobs = needsField!.AsArray().Select(field => field.AsHash()).ToList();
+                CbObject cb = new CbObject(roundTrippedBuffer);
+                CbField needsField = cb["needs"];
+                Assert.AreNotEqual(CbField.Empty, needsField);
+                List<BlobIdentifier> missingBlobs = needsField.AsArray().Select(field => BlobIdentifier.FromIoHash(field.AsHash())).ToList();
                 Assert.AreEqual(1, missingBlobs.Count);
                 Assert.AreEqual(childObjectHash, missingBlobs[0]);
             }
@@ -503,11 +501,10 @@ namespace Horde.Storage.FunctionalTests.References
                 await using MemoryStream ms = new MemoryStream();
                 await result.Content.CopyToAsync(ms);
                 byte[] roundTrippedBuffer = ms.ToArray();
-                ReadOnlyMemory<byte> localMemory = new ReadOnlyMemory<byte>(roundTrippedBuffer);
-                CompactBinaryObject cb = CompactBinaryObject.Load(ref localMemory);
-                CompactBinaryField? value = cb["identifier"];
-                Assert.IsNotNull(value);
-                Assert.AreEqual(childObjectHash, value!.AsHash());
+                CbObject cb = new CbObject(roundTrippedBuffer);
+                CbField value = cb["identifier"];
+                Assert.AreNotEqual(CbField.Empty, value);
+                Assert.AreEqual(childObjectHash, BlobIdentifier.FromIoHash(value.AsHash()));
             }
 
             // since we have now uploaded the child object putting the object again should result in no missing references
@@ -524,11 +521,9 @@ namespace Horde.Storage.FunctionalTests.References
                 await using MemoryStream ms = new MemoryStream();
                 await result.Content.CopyToAsync(ms);
                 byte[] roundTrippedBuffer = ms.ToArray();
-                ReadOnlyMemory<byte> localMemory = new ReadOnlyMemory<byte>(roundTrippedBuffer);
-                CompactBinaryObject cb = CompactBinaryObject.Load(ref localMemory);
-                CompactBinaryField? needsField = cb["needs"];
-                Assert.IsNotNull(needsField);
-                List<BlobIdentifier?> missingBlobs = needsField!.AsArray().Select(field => field.AsHash()).ToList();
+                CbObject cb = new CbObject(roundTrippedBuffer);
+                CbField needsField = cb["needs"];
+                List<BlobIdentifier> missingBlobs = needsField.AsArray().Select(field => BlobIdentifier.FromIoHash(field.AsHash())).ToList();
                 Assert.AreEqual(0, missingBlobs.Count);
             }
 
@@ -562,14 +557,13 @@ namespace Horde.Storage.FunctionalTests.References
                 await getResponse.Content.CopyToAsync(ms);
 
                 byte[] roundTrippedBuffer = ms.ToArray();
-                ReadOnlyMemory<byte> localMemory = new ReadOnlyMemory<byte>(roundTrippedBuffer);
-                CompactBinaryObject cb = CompactBinaryObject.Load(ref localMemory);
-                List<CompactBinaryField> fields = cb.GetFields().ToList();
+                CbObject cb = new CbObject(roundTrippedBuffer);
+                List<CbField> fields = cb.ToList();
 
                 Assert.AreEqual(1, fields.Count);
-                CompactBinaryField childObjectField = fields[0];
+                CbField childObjectField = fields[0];
                 Assert.AreEqual("childObject", childObjectField.Name);
-                Assert.AreEqual(childObjectHash, childObjectField.AsHash());
+                Assert.AreEqual(childObjectHash, BlobIdentifier.FromIoHash(childObjectField.AsHash()));
             }
 
 
@@ -595,7 +589,7 @@ namespace Horde.Storage.FunctionalTests.References
 
                 string? value = property.Value.Value<string>();
                 Assert.IsNotNull(value);
-                Assert.AreEqual(childObjectHash.ToString(), value);
+                Assert.AreEqual(childObjectHash.ToString().ToLower(), value);
             }
 
             {
@@ -653,23 +647,23 @@ namespace Horde.Storage.FunctionalTests.References
             BlobIdentifier blobHashChild = BlobIdentifier.FromBlob(dataChild);
             await _blobService.PutObject(TestNamespace, dataChild, blobHashChild);
 
-            CompactBinaryWriter writerChild = new CompactBinaryWriter();
+            CbWriter writerChild = new CbWriter();
             writerChild.BeginObject();
-            writerChild.AddBinaryAttachment(blobHashChild, "blob");
+            writerChild.WriteBinaryAttachment("blob", blobHashChild.AsIoHash());
             writerChild.EndObject();
 
-            byte[] childDataObject = writerChild.Save();
+            byte[] childDataObject = writerChild.ToByteArray();
             BlobIdentifier childDataObjectHash = BlobIdentifier.FromBlob(childDataObject);
             await _blobService.PutObject(TestNamespace, childDataObject, childDataObjectHash);
 
-            CompactBinaryWriter writerParent = new CompactBinaryWriter();
+            CbWriter writerParent = new CbWriter();
             writerParent.BeginObject();
             
-            writerParent.AddBinaryAttachment(blobHash, "blobAttachment");
-            writerParent.AddCompactBinaryAttachment(childDataObjectHash, "objectAttachment");
+            writerParent.WriteBinaryAttachment("blobAttachment", blobHash.AsIoHash());
+            writerParent.WriteObjectAttachment("objectAttachment", childDataObjectHash.AsIoHash());
             writerParent.EndObject();
 
-            byte[] objectData = writerParent.Save();
+            byte[] objectData = writerParent.ToByteArray();
             BlobIdentifier objectHash = BlobIdentifier.FromBlob(objectData);
 
             IoHashKey key = IoHashKey.FromName("newHierarchyObject");
@@ -688,11 +682,9 @@ namespace Horde.Storage.FunctionalTests.References
                 await using MemoryStream ms = new MemoryStream();
                 await result.Content.CopyToAsync(ms);
                 byte[] roundTrippedBuffer = ms.ToArray();
-                ReadOnlyMemory<byte> localMemory = new ReadOnlyMemory<byte>(roundTrippedBuffer);
-                CompactBinaryObject cb = CompactBinaryObject.Load(ref localMemory);
-                CompactBinaryField? needsField = cb["needs"];
-                Assert.IsNotNull(needsField);
-                List<BlobIdentifier?> missingBlobs = needsField!.AsArray().Select(field => field.AsHash()).ToList();
+                CbObject cb = new CbObject(roundTrippedBuffer);
+                CbField needsField = cb["needs"];
+                List<BlobIdentifier> missingBlobs = needsField.AsArray().Select(field => BlobIdentifier.FromIoHash(field.AsHash())).ToList();
                 Assert.AreEqual(0, missingBlobs.Count);
             }
 
@@ -732,15 +724,15 @@ namespace Horde.Storage.FunctionalTests.References
 
                     byte[] roundTrippedBuffer = ms.ToArray();
                     ReadOnlyMemory<byte> localMemory = new ReadOnlyMemory<byte>(roundTrippedBuffer);
-                    CompactBinaryObject cb = CompactBinaryObject.Load(ref localMemory);
-                    Assert.AreEqual(2, cb.GetFields().Count());
-
-                    CompactBinaryField? blobAttachmentField = cb["blobAttachment"];
-                    Assert.IsNotNull(blobAttachmentField);
-                    blobAttachment = blobAttachmentField!.AsBinaryAttachment()!;
-                    CompactBinaryField? objectAttachmentField = cb["objectAttachment"];
-                    Assert.IsNotNull(objectAttachmentField);
-                    objectAttachment = objectAttachmentField!.AsCompactBinaryAttachment()!;
+                    CbObject cb = new CbObject(roundTrippedBuffer);
+                    Assert.AreEqual(2, cb.Count());
+
+                    CbField blobAttachmentField = cb["blobAttachment"];
+                    Assert.AreNotEqual(CbField.Empty, blobAttachmentField);
+                    blobAttachment = BlobIdentifier.FromIoHash(blobAttachmentField.AsBinaryAttachment());
+                    CbField objectAttachmentField = cb["objectAttachment"];
+                    Assert.AreNotEqual(CbField.Empty, objectAttachmentField);
+                    objectAttachment = BlobIdentifier.FromIoHash(objectAttachmentField.AsObjectAttachment().Hash);
                 }
 
                 {
@@ -763,13 +755,13 @@ namespace Horde.Storage.FunctionalTests.References
                     await getAttachment.Content.CopyToAsync(ms);
                     byte[] roundTrippedBuffer = ms.ToArray();
                     ReadOnlyMemory<byte> localMemory = new ReadOnlyMemory<byte>(roundTrippedBuffer);
-                    CompactBinaryObject cb = CompactBinaryObject.Load(ref localMemory);
-                    Assert.AreEqual(1, cb.GetFields().Count());
+                    CbObject cb = new CbObject(roundTrippedBuffer);
+                    Assert.AreEqual(1, cb.Count());
 
-                    CompactBinaryField? blobField = cb["blob"];
-                    Assert.IsNotNull(blobField);
+                    CbField blobField = cb["blob"];
+                    Assert.AreNotEqual(CbField.Empty, blobField);
 
-                    attachedBlobIdentifier = blobField!.AsBinaryAttachment()!;
+                    attachedBlobIdentifier = BlobIdentifier.FromIoHash(blobField!.AsBinaryAttachment());
                 }
 
                 {
@@ -822,23 +814,23 @@ namespace Horde.Storage.FunctionalTests.References
             byte[] dataChild = Encoding.ASCII.GetBytes(blobContentsChild);
             BlobIdentifier blobHashChild = BlobIdentifier.FromBlob(dataChild);
 
-            CompactBinaryWriter writerChild = new CompactBinaryWriter();
+            CbWriter writerChild = new CbWriter();
             writerChild.BeginObject();
-            writerChild.AddBinaryAttachment(blobHashChild, "blob");
+            writerChild.WriteBinaryAttachment("blob", blobHashChild.AsIoHash());
             writerChild.EndObject();
 
-            byte[] childDataObject = writerChild.Save();
+            byte[] childDataObject = writerChild.ToByteArray();
             BlobIdentifier childDataObjectHash = BlobIdentifier.FromBlob(childDataObject);
             await _blobService.PutObject(TestNamespace, childDataObject, childDataObjectHash);
 
-            CompactBinaryWriter writerParent = new CompactBinaryWriter();
+            CbWriter writerParent = new CbWriter();
             writerParent.BeginObject();
             
-            writerParent.AddBinaryAttachment(blobHash, "blobAttachment");
-            writerParent.AddCompactBinaryAttachment(childDataObjectHash, "objectAttachment");
+            writerParent.WriteBinaryAttachment("blobAttachment", blobHash.AsIoHash());
+            writerParent.WriteObjectAttachment("objectAttachment", childDataObjectHash.AsIoHash());
             writerParent.EndObject();
 
-            byte[] objectData = writerParent.Save();
+            byte[] objectData = writerParent.ToByteArray();
             BlobIdentifier objectHash = BlobIdentifier.FromBlob(objectData);
             
             IoHashKey key = IoHashKey.FromName("newHierarchyObject");
@@ -857,10 +849,9 @@ namespace Horde.Storage.FunctionalTests.References
                     await result.Content.CopyToAsync(ms);
                     byte[] roundTrippedBuffer = ms.ToArray();
                     ReadOnlyMemory<byte> localMemory = new ReadOnlyMemory<byte>(roundTrippedBuffer);
-                    CompactBinaryObject cb = CompactBinaryObject.Load(ref localMemory);
-                    CompactBinaryField? needsField = cb["needs"];
-                    Assert.IsNotNull(needsField);
-                    List<BlobIdentifier?> missingBlobs = needsField!.AsArray().Select(field => field.AsHash()).ToList();
+                    CbObject cb = new CbObject(roundTrippedBuffer);
+                    CbField needsField = cb["needs"];
+                    List<BlobIdentifier> missingBlobs = needsField.AsArray().Select(field => BlobIdentifier.FromIoHash(field.AsHash())).ToList();
                     Assert.AreEqual(2, missingBlobs.Count);
                     Assert.IsTrue(missingBlobs.Contains(blobHash));
                     Assert.IsTrue(missingBlobs.Contains(blobHashChild));
@@ -908,12 +899,12 @@ namespace Horde.Storage.FunctionalTests.References
 
             await contentIdStore.Put(TestNamespace, contentId, blobHash, blobData.Length);
 
-            CompactBinaryWriter writer = new CompactBinaryWriter();
+            CbWriter writer = new CbWriter();
             writer.BeginObject();
-            writer.AddBinaryAttachment(contentId.AsBlobIdentifier(), "blob");
+            writer.WriteBinaryAttachment("blob", contentId.AsIoHash());
             writer.EndObject();
 
-            byte[] objectData = writer.Save();
+            byte[] objectData = writer.ToByteArray();
             BlobIdentifier objectHash = BlobIdentifier.FromBlob(objectData);
             
             IoHashKey key = IoHashKey.FromName("putContentIdMissingBlob");
@@ -931,15 +922,14 @@ namespace Horde.Storage.FunctionalTests.References
                     await using MemoryStream ms = new MemoryStream();
                     await result.Content.CopyToAsync(ms);
                     byte[] roundTrippedBuffer = ms.ToArray();
-                    ReadOnlyMemory<byte> localMemory = new ReadOnlyMemory<byte>(roundTrippedBuffer);
-                    CompactBinaryObject cb = CompactBinaryObject.Load(ref localMemory);
-                    CompactBinaryField? needsField = cb["needs"];
-                    Assert.IsNotNull(needsField);
-                    List<BlobIdentifier?> missingBlobs = needsField!.AsArray().Select(field => field.AsHash()).ToList();
+                    CbObject cb = new CbObject(roundTrippedBuffer);
+                    CbField needsField = cb["needs"];
+                    Assert.AreNotEqual(CbField.Empty, needsField);
+                    List<IoHash> missingBlobs = needsField.AsArray().Select(field => field.AsHash()).ToList();
                     Assert.AreEqual(1, missingBlobs.Count);
 
                     Assert.AreNotEqual(blobHash, missingBlobs[0], "Refs should not be returning the mapped blob identifiers as this is unknown to the client attempting to put a new ref");
-                    Assert.AreEqual(contentId.AsBlobIdentifier(), missingBlobs[0]);
+                    Assert.AreEqual(contentId.AsBlobIdentifier(), BlobIdentifier.FromIoHash(missingBlobs[0]));
                 }
             }
 
@@ -983,23 +973,23 @@ namespace Horde.Storage.FunctionalTests.References
             byte[] dataChild = Encoding.ASCII.GetBytes(blobContentsChild);
             BlobIdentifier blobHashChild = BlobIdentifier.FromBlob(dataChild);
 
-            CompactBinaryWriter writerChild = new CompactBinaryWriter();
+            CbWriter writerChild = new CbWriter();
             writerChild.BeginObject();
-            writerChild.AddBinaryAttachment(blobHashChild, "blob");
+            writerChild.WriteBinaryAttachment("blob", blobHashChild.AsIoHash());
             writerChild.EndObject();
 
-            byte[] childDataObject = writerChild.Save();
+            byte[] childDataObject = writerChild.ToByteArray();
             BlobIdentifier childDataObjectHash = BlobIdentifier.FromBlob(childDataObject);
             await _blobService.PutObject(TestNamespace, childDataObject, childDataObjectHash);
 
-            CompactBinaryWriter writerParent = new CompactBinaryWriter();
+            CbWriter writerParent = new CbWriter();
             writerParent.BeginObject();
             
-            writerParent.AddBinaryAttachment(blobHash, "blobAttachment");
-            writerParent.AddCompactBinaryAttachment(childDataObjectHash, "objectAttachment");
+            writerParent.WriteBinaryAttachment("blobAttachment", blobHash.AsIoHash());
+            writerParent.WriteObjectAttachment("objectAttachment", childDataObjectHash.AsIoHash());
             writerParent.EndObject();
 
-            byte[] objectData = writerParent.Save();
+            byte[] objectData = writerParent.ToByteArray();
             BlobIdentifier objectHash = BlobIdentifier.FromBlob(objectData);
 
             {
@@ -1015,11 +1005,9 @@ namespace Horde.Storage.FunctionalTests.References
                     await using MemoryStream ms = new MemoryStream();
                     await result.Content.CopyToAsync(ms);
                     byte[] roundTrippedBuffer = ms.ToArray();
-                    ReadOnlyMemory<byte> localMemory = new ReadOnlyMemory<byte>(roundTrippedBuffer);
-                    CompactBinaryObject cb = CompactBinaryObject.Load(ref localMemory);
-                    CompactBinaryField? needsField = cb["needs"];
-                    Assert.IsNotNull(needsField);
-                    List<BlobIdentifier?> missingBlobs = needsField!.AsArray().Select(field => field.AsHash()).ToList();
+                    CbObject cb = new CbObject(roundTrippedBuffer);
+                    CbField needsField = cb["needs"];
+                    List<BlobIdentifier> missingBlobs = needsField.AsArray().Select(field => BlobIdentifier.FromIoHash(field.AsHash())).ToList();
                     Assert.AreEqual(2, missingBlobs.Count);
                     Assert.IsTrue(missingBlobs.Contains(blobHash));
                     Assert.IsTrue(missingBlobs.Contains(blobHashChild));
@@ -1055,10 +1043,9 @@ namespace Horde.Storage.FunctionalTests.References
                     await result.Content.CopyToAsync(ms);
                     byte[] roundTrippedBuffer = ms.ToArray();
                     ReadOnlyMemory<byte> localMemory = new ReadOnlyMemory<byte>(roundTrippedBuffer);
-                    CompactBinaryObject cb = CompactBinaryObject.Load(ref localMemory);
-                    CompactBinaryField? needsField = cb["needs"];
-                    Assert.IsNotNull(needsField);
-                    List<BlobIdentifier?> missingBlobs = needsField!.AsArray().Select(field => field.AsHash()).ToList();
+                    CbObject cb = new CbObject(roundTrippedBuffer);
+                    CbField? needsField = cb["needs"];
+                    List<BlobIdentifier> missingBlobs = needsField.AsArray().Select(field => BlobIdentifier.FromIoHash(field.AsHash())).ToList();
                     Assert.AreEqual(0, missingBlobs.Count);
                 }
             }
@@ -1094,11 +1081,9 @@ namespace Horde.Storage.FunctionalTests.References
                 await using MemoryStream ms = new MemoryStream();
                 await result.Content.CopyToAsync(ms);
                 byte[] roundTrippedBuffer = ms.ToArray();
-                ReadOnlyMemory<byte> localMemory = new ReadOnlyMemory<byte>(roundTrippedBuffer);
-                CompactBinaryObject cb = CompactBinaryObject.Load(ref localMemory);
-                CompactBinaryField? needsField = cb["needs"];
-                Assert.IsNotNull(needsField);
-                List<BlobIdentifier?> missingBlobs = needsField!.AsArray().Select(field => field.AsHash()).ToList();
+                CbObject cb = new CbObject(roundTrippedBuffer);
+                CbField needsField = cb["needs"];
+                List<BlobIdentifier> missingBlobs = needsField.AsArray().Select(field => BlobIdentifier.FromIoHash(field.AsHash())).ToList();
                 Assert.AreEqual(0, missingBlobs.Count);
             }
 
diff --git a/Engine/Source/Programs/Horde/HordeStorage/Horde.Storage/Tests/Functional/Objects/ReplicationServerTests.cs b/Engine/Source/Programs/Horde/HordeStorage/Horde.Storage/Tests/Functional/Objects/ReplicationServerTests.cs
--- a/Engine/Source/Programs/Horde/HordeStorage/Horde.Storage/Tests/Functional/Objects/ReplicationServerTests.cs
+++ b/Engine/Source/Programs/Horde/HordeStorage/Horde.Storage/Tests/Functional/Objects/ReplicationServerTests.cs
@@ -14,6 +14,7 @@ using Microsoft.VisualStudio.TestTools.UnitTesting;
 using Serilog;
 using Logger = Serilog.Core.Logger;
 using EpicGames.Horde.Storage;
+using EpicGames.Serialization;
 
 namespace Horde.Storage.FunctionalTests.References
 {
@@ -60,12 +61,12 @@ namespace Horde.Storage.FunctionalTests.References
         public async Task ReplicationPublicEndpoint()
         {
             // insert a random object
-            CompactBinaryWriter writer = new CompactBinaryWriter();
+            CbWriter writer = new CbWriter();
             writer.BeginObject();
-            writer.AddString("thisIsAField", "stringField");
+            writer.WriteString("stringField", "thisIsAField");
             writer.EndObject();
 
-            byte[] objectData = writer.Save();
+            byte[] objectData = writer.ToByteArray();
             BlobIdentifier objectHash = BlobIdentifier.FromBlob(objectData);
 
             HttpContent requestContent = new ByteArrayContent(objectData);
diff --git a/Engine/Source/Programs/Horde/HordeStorage/Horde.Storage/Tests/Functional/Objects/ReplicationTests.cs b/Engine/Source/Programs/Horde/HordeStorage/Horde.Storage/Tests/Functional/Objects/ReplicationTests.cs
--- a/Engine/Source/Programs/Horde/HordeStorage/Horde.Storage/Tests/Functional/Objects/ReplicationTests.cs
+++ b/Engine/Source/Programs/Horde/HordeStorage/Horde.Storage/Tests/Functional/Objects/ReplicationTests.cs
@@ -26,6 +26,7 @@ using Microsoft.VisualStudio.TestTools.UnitTesting;
 using Serilog;
 using Logger = Serilog.Core.Logger;
 using EpicGames.Horde.Storage;
+using EpicGames.Serialization;
 
 namespace Horde.Storage.FunctionalTests.References
 {
@@ -168,12 +169,12 @@ namespace Horde.Storage.FunctionalTests.References
         [TestMethod]
         public async Task ReplicationLogCreation()
         {
-            CompactBinaryWriter writer = new CompactBinaryWriter();
+            CbWriter writer = new CbWriter();
             writer.BeginObject();
-            writer.AddString("thisIsAField", "stringField");
+            writer.WriteString("stringField", "thisIsAField");
             writer.EndObject();
 
-            byte[] objectData = writer.Save();
+            byte[] objectData = writer.ToByteArray();
             BlobIdentifier objectHash = BlobIdentifier.FromBlob(objectData);
 
             HttpContent requestContent = new ByteArrayContent(objectData);
@@ -239,12 +240,12 @@ namespace Horde.Storage.FunctionalTests.References
         [TestMethod]
         public async Task ReplicationLogReading()
         {
-            CompactBinaryWriter writer = new CompactBinaryWriter();
+            CbWriter writer = new CbWriter();
             writer.BeginObject();
-            writer.AddString("thisIsAField", "stringField");
+            writer.WriteString("stringField", "thisIsAField");
             writer.EndObject();
 
-            byte[] objectData = writer.Save();
+            byte[] objectData = writer.ToByteArray();
             BlobIdentifier objectHash = BlobIdentifier.FromBlob(objectData);
 
             DateTime oldestTimestamp = DateTime.Now.AddDays(-1);
@@ -298,12 +299,12 @@ namespace Horde.Storage.FunctionalTests.References
         [TestMethod]
         public async Task ReplicationLogResume()
         {
-            CompactBinaryWriter writer = new CompactBinaryWriter();
+            CbWriter writer = new CbWriter();
             writer.BeginObject();
-            writer.AddString("thisIsAField", "stringField");
+            writer.WriteString("stringField", "thisIsAField");
             writer.EndObject();
 
-            byte[] objectData = writer.Save();
+            byte[] objectData = writer.ToByteArray();
             BlobIdentifier objectHash = BlobIdentifier.FromBlob(objectData);
 
             DateTime oldestTimestamp = DateTime.Now.AddDays(-3);
@@ -360,12 +361,12 @@ namespace Horde.Storage.FunctionalTests.References
         [TestMethod]
         public async Task ReplicationLogReadingLimit()
         {
-            CompactBinaryWriter writer = new CompactBinaryWriter();
+            CbWriter writer = new CbWriter();
             writer.BeginObject();
-            writer.AddString("thisIsAField", "stringField");
+            writer.WriteString("stringField", "thisIsAField");
             writer.EndObject();
 
-            byte[] objectData = writer.Save();
+            byte[] objectData = writer.ToByteArray();
             BlobIdentifier objectHash = BlobIdentifier.FromBlob(objectData);
 
             DateTime oldestTimestamp = DateTime.Now.AddDays(-1);
@@ -412,12 +413,12 @@ namespace Horde.Storage.FunctionalTests.References
         public async Task ReplicationLogInvalidBucket()
         {
             // the namespace exists but the bucket does not
-            CompactBinaryWriter writer = new CompactBinaryWriter();
+            CbWriter writer = new CbWriter();
             writer.BeginObject();
-            writer.AddString("thisIsAField", "stringField");
+            writer.WriteString("stringField", "thisIsAField");
             writer.EndObject();
 
-            byte[] objectData = writer.Save();
+            byte[] objectData = writer.ToByteArray();
             BlobIdentifier objectHash = BlobIdentifier.FromBlob(objectData);
 
             await _replicationLog.InsertAddEvent(TestNamespace, TestBucket, IoHashKey.FromName("firstObject"), objectHash, DateTime.Now.AddDays(-1));
@@ -442,12 +443,12 @@ namespace Horde.Storage.FunctionalTests.References
         public async Task ReplicationLogOldBucket()
         {
             // the namespace exists but the bucket id is from a old bucket that does not exist anymore
-            CompactBinaryWriter writer = new CompactBinaryWriter();
+            CbWriter writer = new CbWriter();
             writer.BeginObject();
-            writer.AddString("thisIsAField", "stringField");
+            writer.WriteString("stringField", "thisIsAField");
             writer.EndObject();
 
-            byte[] objectData = writer.Save();
+            byte[] objectData = writer.ToByteArray();
             BlobIdentifier objectHash = BlobIdentifier.FromBlob(objectData);
 
             await _replicationLog.InsertAddEvent(TestNamespace, TestBucket, IoHashKey.FromName("firstObject"), objectHash, DateTime.Now.AddDays(-1));
@@ -487,12 +488,12 @@ namespace Horde.Storage.FunctionalTests.References
         [TestMethod]
         public async Task ReplicationLogNoIncrementalLogAvailable()
         {
-            CompactBinaryWriter writer = new CompactBinaryWriter();
+            CbWriter writer = new CbWriter();
             writer.BeginObject();
-            writer.AddString("thisIsAField", "stringField");
+            writer.WriteString("stringField", "thisIsAField");
             writer.EndObject();
 
-            byte[] objectData = writer.Save();
+            byte[] objectData = writer.ToByteArray();
             BlobIdentifier objectHash = BlobIdentifier.FromBlob(objectData);
 
             DateTime oldestTimestamp = DateTime.Now.AddDays(-1);
@@ -526,12 +527,12 @@ namespace Horde.Storage.FunctionalTests.References
         [TestMethod]
         public async Task ReplicationLogSnapshotCreation()
         {
-            CompactBinaryWriter writer = new CompactBinaryWriter();
+            CbWriter writer = new CbWriter();
             writer.BeginObject();
-            writer.AddString("thisIsAField", "stringField");
+            writer.WriteString("stringField", "thisIsAField");
             writer.EndObject();
 
-            byte[] objectData = writer.Save();
+            byte[] objectData = writer.ToByteArray();
             BlobIdentifier objectHash = BlobIdentifier.FromBlob(objectData);
 
             DateTime oldestTimestamp = DateTime.Now.AddDays(-1);
@@ -609,12 +610,12 @@ namespace Horde.Storage.FunctionalTests.References
         [TestMethod]
         public async Task ReplicationLogSnapshotQuerying()
         {
-            CompactBinaryWriter writer = new CompactBinaryWriter();
+            CbWriter writer = new CbWriter();
             writer.BeginObject();
-            writer.AddString("thisIsAField", "stringField");
+            writer.WriteString("stringField", "thisIsAField");
             writer.EndObject();
 
-            byte[] objectData = writer.Save();
+            byte[] objectData = writer.ToByteArray();
             BlobIdentifier objectHash = BlobIdentifier.FromBlob(objectData);
 
             DateTime oldestTimestamp = DateTime.Now.AddDays(-1);
@@ -666,12 +667,12 @@ namespace Horde.Storage.FunctionalTests.References
         [TestMethod]
         public async Task ReplicationLogSnapshotResume()
         {
-            CompactBinaryWriter writer = new CompactBinaryWriter();
+            CbWriter writer = new CbWriter();
             writer.BeginObject();
-            writer.AddString("thisIsAField", "stringField");
+            writer.WriteString("stringField", "thisIsAField");
             writer.EndObject();
 
-            byte[] objectData = writer.Save();
+            byte[] objectData = writer.ToByteArray();
             BlobIdentifier objectHash = BlobIdentifier.FromBlob(objectData);
 
             DateTime oldestTimestamp = DateTime.Now.AddDays(-1);
@@ -739,12 +740,12 @@ namespace Horde.Storage.FunctionalTests.References
             const int maxCountOfSnapshots = 10;
             int countOfSnapshotsToCreate = maxCountOfSnapshots + 2;
 
-            CompactBinaryWriter writer = new CompactBinaryWriter();
+            CbWriter writer = new CbWriter();
             writer.BeginObject();
-            writer.AddString("thisIsAField", "stringField");
+            writer.WriteString("stringField", "thisIsAField");
             writer.EndObject();
 
-            byte[] objectData = writer.Save();
+            byte[] objectData = writer.ToByteArray();
             BlobIdentifier objectHash = BlobIdentifier.FromBlob(objectData);
 
             List<BlobIdentifier> createdSnapshots = new List<BlobIdentifier>();
diff --git a/Engine/Source/Programs/Horde/HordeStorage/Horde.Storage/Tests/Functional/RefTests.cs b/Engine/Source/Programs/Horde/HordeStorage/Horde.Storage/Tests/Functional/RefTests.cs
--- a/Engine/Source/Programs/Horde/HordeStorage/Horde.Storage/Tests/Functional/RefTests.cs
+++ b/Engine/Source/Programs/Horde/HordeStorage/Horde.Storage/Tests/Functional/RefTests.cs
@@ -29,6 +29,7 @@ using Newtonsoft.Json.Linq;
 using Serilog;
 using Serilog.Core;
 using EpicGames.Horde.Storage;
+using EpicGames.Serialization;
 
 namespace Horde.Storage.FunctionalTests.Ref
 {
@@ -340,20 +341,18 @@ namespace Horde.Storage.FunctionalTests.Ref
             response.EnsureSuccessStatusCode();
 
             byte[] buffer = await response.Content.ReadAsByteArrayAsync();
-            ReadOnlyMemory<byte> memory = new ReadOnlyMemory<byte>(buffer);
-            CompactBinaryObject o = CompactBinaryObject.Load(ref memory);
-            List<CompactBinaryField> fields = o.GetFields().ToList();
-            Assert.AreEqual(5, fields.Count);
+            CbObject o = new CbObject(buffer);
+            List<CbField> fields = o.ToList();
+            Assert.AreEqual(6, fields.Count);
             Assert.AreEqual($"{TestNamespace}.bucket.testObjectWithMetadata", o["name"]!.AsString());
-            // we do not support serializing random dictionaries to compact binary
-            Assert.IsNull(o["metadata"]);
+            Assert.IsNotNull(o["metadata"]);
 
             Assert.IsNotNull(o["contentHash"]);
 
             Assert.IsNotNull(o["lastAccessTime"]);
 
             Assert.IsNotNull(o["blob"]);
-            CollectionAssert.AreEqual(TestObjectData, o["blob"]!.AsBinary());
+            CollectionAssert.AreEqual(TestObjectData, o["blob"]!.AsBinary().ToArray());
         }
 
         [TestMethod]
diff --git a/Engine/Source/Programs/Horde/HordeStorage/Horde.Storage/Tests/Functional/StorageTests.cs b/Engine/Source/Programs/Horde/HordeStorage/Horde.Storage/Tests/Functional/StorageTests.cs
--- a/Engine/Source/Programs/Horde/HordeStorage/Horde.Storage/Tests/Functional/StorageTests.cs
+++ b/Engine/Source/Programs/Horde/HordeStorage/Horde.Storage/Tests/Functional/StorageTests.cs
@@ -16,7 +16,9 @@ using Amazon.S3;
 using Amazon.S3.Model;
 using Azure.Storage.Blobs;
 using Dasync.Collections;
+using EpicGames.Core;
 using EpicGames.Horde.Storage;
+using EpicGames.Serialization;
 using Horde.Storage.Controllers;
 using Horde.Storage.Implementation;
 using Jupiter;
@@ -33,6 +35,7 @@ using Newtonsoft.Json;
 using Newtonsoft.Json.Linq;
 using Serilog;
 using Serilog.Core;
+using ContentHash = Jupiter.Implementation.ContentHash;
 
 namespace Horde.Storage.FunctionalTests.Storage
 {
@@ -629,9 +632,12 @@ namespace Horde.Storage.FunctionalTests.Storage
         {
             BlobIdentifier newContent = BlobIdentifier.FromBlob(Encoding.ASCII.GetBytes("this content has never been submitted"));
             HttpRequestMessage request = new HttpRequestMessage(HttpMethod.Post, requestUri: $"api/v1/s/{TestNamespaceName}/exist");
-            CompactBinaryWriter writer = new CompactBinaryWriter();
-            writer.AsArray(new BlobIdentifier[] { _smallFileHash, newContent}, CompactBinaryFieldType.Hash);
-            byte[] buf = writer.Save();
+            CbWriter writer = new CbWriter();
+            writer.BeginUniformArray(CbFieldType.Hash);
+            writer.WriteHashValue(_smallFileHash.AsIoHash());
+            writer.WriteHashValue(newContent.AsIoHash());
+            writer.EndUniformArray();
+            byte[] buf = writer.ToByteArray();
             request.Content = new ByteArrayContent(buf);
             request.Content.Headers.ContentType = new MediaTypeHeaderValue(CustomMediaTypeNames.UnrealCompactBinary);
 
@@ -677,16 +683,15 @@ namespace Horde.Storage.FunctionalTests.Storage
             Assert.AreEqual(HttpStatusCode.OK, response.StatusCode);
             Assert.AreEqual(CustomMediaTypeNames.UnrealCompactBinary, response.Content.Headers.ContentType!.MediaType!);
             byte[] data = await response.Content.ReadAsByteArrayAsync();
-            ReadOnlyMemory<byte> memory = new ReadOnlyMemory<byte>(data);
-            CompactBinaryObject cb = CompactBinaryObject.Load(ref memory);
-            
-            Assert.AreEqual(1, cb.GetFields().Count());
-            CompactBinaryField? needs = cb["needs"];
+            CbObject cb = new CbObject(data);
+
+            Assert.AreEqual(1, cb.Count());
+            CbField? needs = cb["needs"];
             Assert.IsNotNull(needs);
-            BlobIdentifier?[] neededBlobs = needs!.AsArray().Select(field => field.AsHash()).ToArray();
+            IoHash[] neededBlobs = needs!.AsArray().Select(field => field.AsHash()).ToArray();
 
             Assert.AreEqual(1, neededBlobs.Length);
-            Assert.AreEqual(newContent, neededBlobs[0]);
+            Assert.AreEqual(newContent, BlobIdentifier.FromIoHash(neededBlobs[0]));
         }
 
         [TestMethod]
diff --git a/Engine/Source/Programs/Horde/HordeStorage/Horde.Storage/Tests/UnitTests/CompactBinaryTests.cs b/Engine/Source/Programs/Horde/HordeStorage/Horde.Storage/Tests/UnitTests/CompactBinaryTests.cs
--- a/Engine/Source/Programs/Horde/HordeStorage/Horde.Storage/Tests/UnitTests/CompactBinaryTests.cs
+++ b/Engine/Source/Programs/Horde/HordeStorage/Horde.Storage/Tests/UnitTests/CompactBinaryTests.cs
@@ -4,6 +4,8 @@ using System;
 using System.Collections.Generic;
 using System.IO;
 using System.Linq;
+using EpicGames.Core;
+using EpicGames.Serialization;
 using Jupiter.Implementation;
 using Microsoft.VisualStudio.TestTools.UnitTesting;
 
@@ -18,29 +20,28 @@ namespace EuropaUnit
         {
             byte[] bytes = File.ReadAllBytes("CompactBinaryObjects/build");
 
-            ReadOnlyMemory<byte> memory = new ReadOnlyMemory<byte>(bytes);
-            CompactBinaryObject o = CompactBinaryObject.Load(ref memory);
+            CbObject o = new CbObject(bytes);
 
-            Assert.AreEqual("BuildAction", o.Name);
-            List<CompactBinaryField> buildActionFields = o.GetFields().ToList();
+            Assert.AreEqual("BuildAction", o.AsField().Name);
+            List<CbField> buildActionFields = o.ToList();
             Assert.AreEqual(3, buildActionFields.Count);
             Assert.AreEqual("Function", buildActionFields[0].Name);
             Assert.AreEqual("Constants", buildActionFields[1].Name);
             Assert.AreEqual("Inputs", buildActionFields[2].Name);
 
-            List<CompactBinaryField>  constantsFields = buildActionFields[1].GetFields().ToList();
+            List<CbField>  constantsFields = buildActionFields[1].ToList();
             Assert.AreEqual(3, constantsFields.Count);
             Assert.AreEqual("TextureBuildSettings", constantsFields[0].Name);
             Assert.AreEqual("TextureOutputSettings", constantsFields[1].Name);
             Assert.AreEqual("TextureSource", constantsFields[2].Name);
 
-            List<CompactBinaryField>  inputsFields = buildActionFields[2].GetFields().ToList();
+            List<CbField>  inputsFields = buildActionFields[2].ToList();
             Assert.AreEqual(1, inputsFields.Count);
             Assert.AreEqual("7587B323422942733DDD048A91709FDE", inputsFields[0].Name);
             Assert.IsTrue(inputsFields[0].IsBinaryAttachment());
             Assert.IsTrue(inputsFields[0].IsAttachment());
-            Assert.IsFalse(inputsFields[0].IsCompactBinaryAttachment());
-            Assert.AreEqual(new BlobIdentifier("f855382171a0b1e5a1c653aa6c5121a05cbf4ba0"), inputsFields[0].AsHash());
+            Assert.IsFalse(inputsFields[0].IsObjectAttachment());
+            Assert.AreEqual(IoHash.Parse("f855382171a0b1e5a1c653aa6c5121a05cbf4ba0"), inputsFields[0].AsHash());
         }
 
         
@@ -49,19 +50,18 @@ namespace EuropaUnit
         {
             byte[] bytes = File.ReadAllBytes("CompactBinaryObjects/ReferenceOutput");
 
-            ReadOnlyMemory<byte> memory = new ReadOnlyMemory<byte>(bytes);
-            CompactBinaryObject o = CompactBinaryObject.Load(ref memory);
+            CbObject o = new CbObject(bytes);
 
-            Assert.AreEqual("BuildOutput", o.Name);
-            List<CompactBinaryField> buildActionFields = o.GetFields().ToList();
+            Assert.AreEqual("BuildOutput", o.AsField().Name);
+            List<CbField> buildActionFields = o.ToList();
             Assert.AreEqual(1, buildActionFields.Count);
-            CompactBinaryField payloads = buildActionFields[0];
-            List<CompactBinaryField> payloadFields = payloads.GetFields().ToList();
+            CbField payloads = buildActionFields[0];
+            List<CbField> payloadFields = payloads.ToList();
             Assert.AreEqual(3, payloadFields.Count);
 
-            Assert.AreEqual(new BlobIdentifier("5d8a6dc277c968f0d027c98f879c955c1905c293"), payloadFields[0]["RawHash"]!.AsHash());
-            Assert.AreEqual(new BlobIdentifier("313f0d0d334100d83aeb1ee2c42794fd087cb0ae"), payloadFields[1]["RawHash"]!.AsHash());
-            Assert.AreEqual(new BlobIdentifier("c7a03f83c08cdca882110ecf2b5654ee3b09b11e"), payloadFields[2]["RawHash"]!.AsHash());
+            Assert.AreEqual(IoHash.Parse("5d8a6dc277c968f0d027c98f879c955c1905c293"), payloadFields[0]["RawHash"]!.AsHash());
+            Assert.AreEqual(IoHash.Parse("313f0d0d334100d83aeb1ee2c42794fd087cb0ae"), payloadFields[1]["RawHash"]!.AsHash());
+            Assert.AreEqual(IoHash.Parse("c7a03f83c08cdca882110ecf2b5654ee3b09b11e"), payloadFields[2]["RawHash"]!.AsHash());
         }
 
         [TestMethod]
@@ -69,18 +69,15 @@ namespace EuropaUnit
         {
             byte[] bytes = File.ReadAllBytes("CompactBinaryObjects/compact_binary");
 
-            BlobIdentifier bi = BlobIdentifier.FromBlob(bytes);
-            ReadOnlyMemory<byte> memory = new ReadOnlyMemory<byte>(bytes);
-            CompactBinaryObject o = CompactBinaryObject.Load(ref memory);
-
-            Assert.AreEqual("", o.Name);
-            List<CompactBinaryField> buildActionFields = o.GetFields().ToList();
+            CbObject o = new CbObject(bytes);
+            Assert.AreEqual("", o.AsField().Name);
+            List<CbField> buildActionFields = o.ToList();
             Assert.AreEqual(3, buildActionFields.Count);
-            CompactBinaryField payloads = buildActionFields[0];
-            List<CompactBinaryField> payloadFields = payloads.GetFields().ToList();
+            CbField payloads = buildActionFields[0];
+            List<CbField> payloadFields = payloads.ToList();
             Assert.AreEqual(2, payloadFields.Count);
 
-            Assert.AreEqual("{\"Key\":{\"Bucket\":\"EditorDomainPackage\",\"Hash\":\"37DBAA409EF30BA67F18C8FC2FAAF606636CB915\"},\"Meta\":{\"FileSize\":24789},\"Attachments\":[{\"Id\":\"000000000000000000000001\",\"RawHash\":\"DA6FC57E4B9F91377C9509EA0AD567BACB3796C5\",\"RawSize\":24789}]}", o.ToJson());
+            Assert.AreEqual("{\"Key\":{\"Bucket\":\"EditorDomainPackage\",\"Hash\":\"37dbaa409ef30ba67f18c8fc2faaf606636cb915\"},\"Meta\":{\"FileSize\":24789},\"Attachments\":[{\"Id\":\"000000000000000000000001\",\"RawHash\":\"da6fc57e4b9f91377c9509ea0ad567bacb3796c5\",\"RawSize\":24789}]}", o.ToJson());
         }
 
 
@@ -88,53 +85,54 @@ namespace EuropaUnit
         [TestMethod]
         public void WriteArray()
         {
-            var hash1 = new BlobIdentifier("5d8a6dc277c968f0d027c98f879c955c1905c293");
-            var hash2 = new BlobIdentifier("313f0d0d334100d83aeb1ee2c42794fd087cb0ae");
+            IoHash hash1 = IoHash.Parse("5d8a6dc277c968f0d027c98f879c955c1905c293");
+            IoHash hash2 = IoHash.Parse("313f0d0d334100d83aeb1ee2c42794fd087cb0ae");
 
-            CompactBinaryWriter writer = new CompactBinaryWriter();
+            CbWriter writer = new CbWriter();
             writer.BeginObject();
-            writer.AddUniformArray(new BlobIdentifier[] { hash1, hash2}, CompactBinaryFieldType.Hash, "needs");
+            writer.BeginUniformArray("needs", CbFieldType.Hash);
+            writer.WriteHashValue(hash1);
+            writer.WriteHashValue(hash2);
+            writer.EndUniformArray();
             writer.EndObject();
 
-            byte[] objectData = writer.Save();
-            ReadOnlyMemory<byte> memory = new ReadOnlyMemory<byte>(objectData);
-            CompactBinaryObject o = CompactBinaryObject.Load(ref memory);
+            byte[] objectData = writer.ToByteArray();
+            CbObject o = new CbObject(objectData);
 
             // the top object has no name
-            Assert.AreEqual("", o.Name);
-            List<CompactBinaryField> fields = o.GetFields().ToList();
+            Assert.AreEqual("", o.AsField().Name);
+            List<CbField> fields = o.ToList();
             Assert.AreEqual(1, fields.Count);
-            CompactBinaryField? needs = o["needs"];
-            List<CompactBinaryField> blobList = needs!.AsArray().ToList();
-            BlobIdentifier?[] blobs = blobList.Select(field => field!.AsHash()).ToArray();
-            CollectionAssert.AreEqual(new BlobIdentifier[] {hash1, hash2}, blobs);
+            CbField? needs = o["needs"];
+            List<CbField> blobList = needs!.AsArray().ToList();
+            IoHash[] blobs = blobList.Select(field => field!.AsHash()).ToArray();
+            CollectionAssert.AreEqual(new IoHash[] {hash1, hash2}, blobs);
         }
 
 
         [TestMethod]
         public void WriteObject()
         {
-            var hash1 = new BlobIdentifier("5d8a6dc277c968f0d027c98f879c955c1905c293");
+            IoHash hash1 = IoHash.Parse("5d8a6dc277c968f0d027c98f879c955c1905c293");
 
-            CompactBinaryWriter writer = new CompactBinaryWriter();
+            CbWriter writer = new CbWriter();
             writer.BeginObject();
-            writer.AddString("test", "string");
-            writer.AddBinaryAttachment(hash1, "hash");
+            writer.WriteString("string", "test");
+            writer.WriteBinaryAttachment("hash", hash1);
             writer.EndObject();
 
-            byte[] objectData = writer.Save();
-            ReadOnlyMemory<byte> memory = new ReadOnlyMemory<byte>(objectData);
-            CompactBinaryObject o = CompactBinaryObject.Load(ref memory);
+            byte[] objectData = writer.ToByteArray();
+            CbObject o = new CbObject(objectData);
 
             // the object has no name and 2 fields
-            Assert.AreEqual("", o.Name);
-            List<CompactBinaryField> fields = o.GetFields().ToList();
+            Assert.AreEqual("", o.AsField().Name);
+            List<CbField> fields = o.ToList();
             Assert.AreEqual(2, fields.Count);
 
-            CompactBinaryField? stringField = o["string"];
+            CbField? stringField = o["string"];
             Assert.AreEqual("test", stringField!.AsString());
 
-            CompactBinaryField? hashField = o["hash"];
+            CbField? hashField = o["hash"];
             Assert.AreEqual(hash1, hashField!.AsAttachment());
         }
     }
diff --git a/Engine/Source/Programs/Horde/HordeStorage/Jupiter.Common/BaseStartup.cs b/Engine/Source/Programs/Horde/HordeStorage/Jupiter.Common/BaseStartup.cs
--- a/Engine/Source/Programs/Horde/HordeStorage/Jupiter.Common/BaseStartup.cs
+++ b/Engine/Source/Programs/Horde/HordeStorage/Jupiter.Common/BaseStartup.cs
@@ -10,6 +10,7 @@ using System.Text;
 using System.Threading.Tasks;
 using Amazon;
 using Datadog.Trace;
+using EpicGames.AspNet;
 using Jupiter.Common;
 using Microsoft.AspNetCore.Authentication;
 using Microsoft.AspNetCore.Authorization;
@@ -81,8 +82,8 @@ namespace Jupiter
                 .AddNewtonsoftJson()
                 .AddMvcOptions(options =>
                 {
-                    options.InputFormatters.Add(new CompactBinaryInputFormatter());
-                    options.OutputFormatters.Add(new CompactBinaryOutputFormatter());
+                    options.InputFormatters.Add(new CbInputFormatter());
+                    options.OutputFormatters.Add(new CbOutputFormatter());
                     options.OutputFormatters.Add(new RawOutputFormatter());
 
                     options.FormatterMappings.SetMediaTypeMappingForFormat("raw", MediaTypeNames.Application.Octet);
diff --git a/Engine/Source/Programs/Horde/HordeStorage/Jupiter.Common/CompactBinaryInputFormatter.cs b/Engine/Source/Programs/Horde/HordeStorage/Jupiter.Common/CompactBinaryInputFormatter.cs
deleted file mode 100644
diff --git a/Engine/Source/Programs/Horde/HordeStorage/Jupiter.Common/CompactBinaryOutputFormatter.cs b/Engine/Source/Programs/Horde/HordeStorage/Jupiter.Common/CompactBinaryOutputFormatter.cs
deleted file mode 100644
diff --git a/Engine/Source/Programs/Horde/HordeStorage/Jupiter.Common/Implementation/BlobIdentifier.cs b/Engine/Source/Programs/Horde/HordeStorage/Jupiter.Common/Implementation/BlobIdentifier.cs
--- a/Engine/Source/Programs/Horde/HordeStorage/Jupiter.Common/Implementation/BlobIdentifier.cs
+++ b/Engine/Source/Programs/Horde/HordeStorage/Jupiter.Common/Implementation/BlobIdentifier.cs
@@ -6,12 +6,17 @@ using System.Globalization;
 using System.IO;
 using System.Threading.Tasks;
 using Blake3;
+using EpicGames.Core;
+using EpicGames.Serialization;
 using Newtonsoft.Json;
+using SharpYaml.Tokens;
+using JsonWriter = Newtonsoft.Json.JsonWriter;
 
 namespace Jupiter.Implementation
 {
     [JsonConverter(typeof(BlobIdentifierConverter))]
     [TypeConverter(typeof(BlobIdentifierTypeConverter))]
+    [CbConverter(typeof(BlobIdentifierCbConverter))]
     public class BlobIdentifier : ContentHash,  IEquatable<BlobIdentifier>
     {
         // multi thread the hashing for blobs larger then this size
@@ -135,6 +140,16 @@ namespace Jupiter.Implementation
             byte[] hash = blake3Hash.AsSpan().Slice(0, 20).ToArray();
             return new BlobIdentifier(hash);
         }
+
+        public static BlobIdentifier FromIoHash(IoHash blobIdentifier)
+        {
+            return new BlobIdentifier(blobIdentifier.ToByteArray());
+        }
+
+        public IoHash AsIoHash()
+        {
+            return new IoHash(HashData);
+        }
     }
 
     public class BlobIdentifierTypeConverter : TypeConverter
@@ -177,4 +192,15 @@ namespace Jupiter.Implementation
             return new BlobIdentifier(s!);
         }
     }
+
+    public class BlobIdentifierCbConverter : CbConverterBase<BlobIdentifier>
+    {
+        public override BlobIdentifier Read(CbField field) => new BlobIdentifier(field.AsHash().ToByteArray());
+
+        /// <inheritdoc/>
+        public override void Write(CbWriter writer, BlobIdentifier value) => writer.WriteBinaryAttachmentValue(new IoHash(value.HashData));
+
+        /// <inheritdoc/>
+        public override void WriteNamed(CbWriter writer, Utf8String name, BlobIdentifier value) => writer.WriteBinaryAttachment(name, new IoHash(value.HashData));
+    }
 }
diff --git a/Engine/Source/Programs/Horde/HordeStorage/Jupiter.Common/Implementation/CompactBinaryObject.cs b/Engine/Source/Programs/Horde/HordeStorage/Jupiter.Common/Implementation/CompactBinaryObject.cs
deleted file mode 100644
diff --git a/Engine/Source/Programs/Horde/HordeStorage/Jupiter.Common/Implementation/CompactBinaryWriter.cs b/Engine/Source/Programs/Horde/HordeStorage/Jupiter.Common/Implementation/CompactBinaryWriter.cs
deleted file mode 100644
diff --git a/Engine/Source/Programs/Horde/HordeStorage/Jupiter.Common/Implementation/ContentHash.cs b/Engine/Source/Programs/Horde/HordeStorage/Jupiter.Common/Implementation/ContentHash.cs
--- a/Engine/Source/Programs/Horde/HordeStorage/Jupiter.Common/Implementation/ContentHash.cs
+++ b/Engine/Source/Programs/Horde/HordeStorage/Jupiter.Common/Implementation/ContentHash.cs
@@ -5,11 +5,15 @@ using System.IO;
 using System.Text;
 using System.Threading.Tasks;
 using Blake3;
+using EpicGames.Core;
+using EpicGames.Serialization;
 using Newtonsoft.Json;
+using JsonWriter = Newtonsoft.Json.JsonWriter;
 
 namespace Jupiter.Implementation
 {
     [JsonConverter(typeof(ContentHashConverter))]
+    [CbConverter(typeof(ContentHashCbConverter))]
     public class ContentHash : IEquatable<ContentHash>, IEquatable<byte[]>
     {
         protected readonly ByteArrayComparer Comparer = new ByteArrayComparer();
@@ -112,4 +116,15 @@ namespace Jupiter.Implementation
             return new ContentHash(s!);
         }
     }
+
+    public class ContentHashCbConverter : CbConverterBase<ContentHash>
+    {
+        public override ContentHash Read(CbField field) => new ContentHash(field.AsHash().ToByteArray());
+
+        /// <inheritdoc/>
+        public override void Write(CbWriter writer, ContentHash value) => writer.WriteHashValue(new IoHash(value.HashData));
+
+        /// <inheritdoc/>
+        public override void WriteNamed(CbWriter writer, Utf8String name, ContentHash value) => writer.WriteHash(name, new IoHash(value.HashData));
+    }
 }
diff --git a/Engine/Source/Programs/Horde/HordeStorage/Jupiter.Common/Implementation/ContentId.cs b/Engine/Source/Programs/Horde/HordeStorage/Jupiter.Common/Implementation/ContentId.cs
--- a/Engine/Source/Programs/Horde/HordeStorage/Jupiter.Common/Implementation/ContentId.cs
+++ b/Engine/Source/Programs/Horde/HordeStorage/Jupiter.Common/Implementation/ContentId.cs
@@ -3,7 +3,9 @@
 using System;
 using System.ComponentModel;
 using System.Globalization;
+using EpicGames.Core;
 using Newtonsoft.Json;
+using JsonWriter = Newtonsoft.Json.JsonWriter;
 
 namespace Jupiter.Implementation
 {
@@ -68,6 +70,16 @@ namespace Jupiter.Implementation
         {
             return new BlobIdentifier(HashData);
         }
+
+        public static ContentId FromIoHash(IoHash ioHash)
+        {
+            return new ContentId(ioHash.ToByteArray());
+        }
+
+        public IoHash AsIoHash()
+        {
+            return new IoHash(HashData);
+        }
     }
 
     public class ContentIdTypeConverter : TypeConverter
diff --git a/Engine/Source/Programs/Horde/HordeStorage/Jupiter.Common/Jupiter.Common.csproj b/Engine/Source/Programs/Horde/HordeStorage/Jupiter.Common/Jupiter.Common.csproj
--- a/Engine/Source/Programs/Horde/HordeStorage/Jupiter.Common/Jupiter.Common.csproj
+++ b/Engine/Source/Programs/Horde/HordeStorage/Jupiter.Common/Jupiter.Common.csproj
@@ -53,6 +53,7 @@
 
   <ItemGroup>
     <ProjectReference Include="..\..\..\Shared\EpicGames.Horde\EpicGames.Horde.csproj" />
+    <ProjectReference Include="..\..\..\Shared\EpicGames.AspNet\EpicGames.AspNet.csproj" />
   </ItemGroup>
 
 </Project>
