Horde: Remove default value for run modes for easier config

Setting array-based config variables in ASP.NET is tricky since they are only additive and cannot remove already set values.

diff --git a/Engine/Source/Programs/Horde/Horde.Build/ServerSettings.cs b/Engine/Source/Programs/Horde/Horde.Build/ServerSettings.cs
--- a/Engine/Source/Programs/Horde/Horde.Build/ServerSettings.cs
+++ b/Engine/Source/Programs/Horde/Horde.Build/ServerSettings.cs
@@ -122,7 +122,7 @@ namespace HordeServer
 	public class ServerSettings
 	{
 		/// <inheritdoc cref="RunMode" />
-		public RunMode[] RunModes { get; set; } = { RunMode.RequestHandler, RunMode.Worker };
+		public RunMode[]? RunModes { get; set; }
 		
 		/// <summary>
 		/// Main port for serving HTTP. Uses the default Kestrel port (5000) if not specified.
@@ -539,6 +539,7 @@ namespace HordeServer
 		/// <returns>True if mode is active</returns>
 		public bool IsRunModeActive(RunMode Mode)
 		{
+			if (RunModes == null) return true;
 			return RunModes.Contains(Mode);
 		}
 	}
