Horde: Fix device service permission issue
#rnx
#skipci

diff --git a/Engine/Source/Programs/Horde/Horde.Build/Devices/Impl/DeviceService.cs b/Engine/Source/Programs/Horde/Horde.Build/Devices/Impl/DeviceService.cs
--- a/Engine/Source/Programs/Horde/Horde.Build/Devices/Impl/DeviceService.cs
+++ b/Engine/Source/Programs/Horde/Horde.Build/Devices/Impl/DeviceService.cs
@@ -454,15 +454,27 @@ namespace HordeServer.Services
 		/// <returns></returns>
 		[SuppressMessage("Usage", "CA1801:Review unused parameters")]
 		[SuppressMessage("Performance", "CA1822: Can be static ")]
-		public Task<bool> AuthorizeAsync(AclAction Action, ClaimsPrincipal User)
+		async public Task<bool> AuthorizeAsync(AclAction Action, ClaimsPrincipal User)
 		{
-			// Setup ACL's for pool and platform access, https://jira.it.epicgames.com/browse/UE-117224/			
-			// allow reads, though restrict writing to internal employees
+			// Tihs is deprecated and device auth should be going through GetUserPoolAuthorizationsAsync
+
 			if (Action == AclAction.DeviceRead)
 			{
-				return Task.FromResult(true);
+				return true;
+			}
+
+			if (User.IsInRole("Internal-Employees"))
+			{
+				return true;
 			}
-			return Task.FromResult(User.IsInRole("Internal-Employees"));			
+
+			if (await AclService.AuthorizeAsync(AclAction.AdminWrite, User))
+			{
+				return true;
+			}
+
+			return false;
+
 		}
 
 		/// <summary>
@@ -498,14 +510,22 @@ namespace HordeServer.Services
 				DeviceWrite.Add(ProjectId, await ProjectService.AuthorizeAsync(ProjectId, AclAction.DeviceWrite, User, PermissionsCache));
 			}
 
-			bool InternalEmployee = User.IsInRole("Internal-Employees");
+			// for global pools which aren't associated with a project
+			bool GlobalPoolAccess = User.IsInRole("Internal-Employees");
 
-			foreach (IDevicePool Pool in AllPools)
+			if (!GlobalPoolAccess)
 			{
-				// for global pools which aren't associated with a project
+				if (await AclService.AuthorizeAsync(AclAction.AdminWrite, User))
+				{
+					GlobalPoolAccess = true;
+				}
+			}
+
+			foreach (IDevicePool Pool in AllPools)
+			{				
 				if (Pool.ProjectIds == null || Pool.ProjectIds.Count == 0)
 				{
-					if (Pool.PoolType == DevicePoolType.Shared && !InternalEmployee)
+					if (Pool.PoolType == DevicePoolType.Shared && !GlobalPoolAccess)
 					{
 						AuthPools.Add(new DevicePoolAuthorization(Pool, false, false));
 						continue;
