Horde: Rename run mode for serving requests

To better match other config and Horde Storage.

diff --git a/Engine/Source/Programs/Horde/Horde.Build/ServerSettings.cs b/Engine/Source/Programs/Horde/Horde.Build/ServerSettings.cs
--- a/Engine/Source/Programs/Horde/Horde.Build/ServerSettings.cs
+++ b/Engine/Source/Programs/Horde/Horde.Build/ServerSettings.cs
@@ -94,7 +94,7 @@ namespace HordeServer
 		/// These requests are time-sensitive and short-lived, typically less than 5 secs.
 		/// If processes handling requests are unavailable, it will be very visible for users.
 		/// </summary>
-		RequestHandler,
+		Server,
 		
 		/// <summary>
 		/// Run non-request facing workloads. Such as background services, processing queues, running work
