Horde: Validate run modes in server settings more strictly

diff --git a/Engine/Source/Programs/Horde/Horde.Build/ServerSettings.cs b/Engine/Source/Programs/Horde/Horde.Build/ServerSettings.cs
--- a/Engine/Source/Programs/Horde/Horde.Build/ServerSettings.cs
+++ b/Engine/Source/Programs/Horde/Horde.Build/ServerSettings.cs
@@ -89,6 +89,11 @@ namespace HordeServer
 	/// </summary>
 	public enum RunMode
 	{
+		/// <summary>
+		/// Default no-op value (ASP.NET config will default to this for enums that cannot be parsed)
+		/// </summary> 
+		None,
+		
 		/// <summary>
 		/// Handle and respond to incoming external requests, such as HTTP REST and gRPC calls.
 		/// These requests are time-sensitive and short-lived, typically less than 5 secs.
@@ -542,5 +547,17 @@ namespace HordeServer
 			if (RunModes == null) return true;
 			return RunModes.Contains(Mode);
 		}
+
+		/// <summary>
+		/// Validate the settings object does not contain any invalid fields
+		/// </summary>
+		/// <exception cref="ArgumentException"></exception>
+		public void Validate()
+		{
+			if (IsRunModeActive(RunMode.None))
+			{
+				throw new ArgumentException($"Settings key '{nameof(RunModes)}' contains one or more invalid entries");
+			}
+		}
 	}
 }
diff --git a/Engine/Source/Programs/Horde/Horde.Build/Startup.cs b/Engine/Source/Programs/Horde/Horde.Build/Startup.cs
--- a/Engine/Source/Programs/Horde/Horde.Build/Startup.cs
+++ b/Engine/Source/Programs/Horde/Horde.Build/Startup.cs
@@ -258,6 +258,8 @@ namespace HordeServer
 			ServerSettings Settings = new ServerSettings();
 			ConfigSection.Bind(Settings);
 
+			Settings.Validate();
+			
 			Services.Configure<CommitServiceOptions>(ConfigSection.GetSection("Replication"));
 
 			if (Settings.GlobalThreadPoolMinSize != null)
