Horde: Fix bug in validation of empty run modes

diff --git a/Engine/Source/Programs/Horde/Horde.Build/ServerSettings.cs b/Engine/Source/Programs/Horde/Horde.Build/ServerSettings.cs
--- a/Engine/Source/Programs/Horde/Horde.Build/ServerSettings.cs
+++ b/Engine/Source/Programs/Horde/Horde.Build/ServerSettings.cs
@@ -127,7 +127,7 @@ namespace HordeServer
 	public class ServerSettings
 	{
 		/// <inheritdoc cref="RunMode" />
-		public RunMode[]? RunModes { get; set; }
+		public RunMode[]? RunModes { get; set; } = null;
 		
 		/// <summary>
 		/// Main port for serving HTTP. Uses the default Kestrel port (5000) if not specified.
@@ -554,7 +554,7 @@ namespace HordeServer
 		/// <exception cref="ArgumentException"></exception>
 		public void Validate()
 		{
-			if (IsRunModeActive(RunMode.None))
+			if (RunModes != null && IsRunModeActive(RunMode.None))
 			{
 				throw new ArgumentException($"Settings key '{nameof(RunModes)}' contains one or more invalid entries");
 			}
