Horde: Add extra tracing to FindLeases endpoint

At times, slowdowns are seen for this one. Likely database related.

diff --git a/Engine/Source/Programs/Horde/Horde.Build/Agents/Impl/AgentService.cs b/Engine/Source/Programs/Horde/Horde.Build/Agents/Impl/AgentService.cs
--- a/Engine/Source/Programs/Horde/Horde.Build/Agents/Impl/AgentService.cs
+++ b/Engine/Source/Programs/Horde/Horde.Build/Agents/Impl/AgentService.cs
@@ -22,6 +22,8 @@ using System.Linq;
 using System.Security.Claims;
 using System.Threading;
 using System.Threading.Tasks;
+using OpenTracing;
+using OpenTracing.Util;
 
 namespace HordeServer.Services
 {
@@ -748,6 +750,13 @@ namespace HordeServer.Services
 		/// <returns>List of leases matching the given criteria</returns>
 		public Task<List<ILease>> FindLeasesAsync(AgentId? AgentId, SessionId? SessionId, DateTime? StartTime, DateTime? FinishTime, int Index, int Count)
 		{
+			using IScope Scope = GlobalTracer.Instance.BuildSpan($"{nameof(AgentService)}.{nameof(FindLeasesAsync)}").StartActive();
+			Scope.Span.SetTag("AgentId", AgentId?.ToString());
+			Scope.Span.SetTag("SessionId", SessionId?.ToString());
+			Scope.Span.SetTag("StartTime", StartTime?.ToString());
+			Scope.Span.SetTag("FinishTime", FinishTime?.ToString());
+			Scope.Span.SetTag("Index", Index);
+			Scope.Span.SetTag("Count", Count);
 			return Leases.FindLeasesAsync(AgentId, SessionId, StartTime, FinishTime, Index, Count);
 		}
 
diff --git a/Engine/Source/Programs/Horde/Horde.Build/Controllers/LeasesController.cs b/Engine/Source/Programs/Horde/Horde.Build/Controllers/LeasesController.cs
--- a/Engine/Source/Programs/Horde/Horde.Build/Controllers/LeasesController.cs
+++ b/Engine/Source/Programs/Horde/Horde.Build/Controllers/LeasesController.cs
@@ -21,6 +21,8 @@ using System.Security.Claims;
 using System.Text.Json;
 using System.Threading.Tasks;
 using Microsoft.Extensions.Logging;
+using OpenTracing;
+using OpenTracing.Util;
 
 namespace HordeServer.Controllers
 {
@@ -108,18 +110,21 @@ namespace HordeServer.Controllers
 
 			List<object> Responses = new List<object>();
 
-			Dictionary<AgentId, double?> CachedAgentRates = new Dictionary<AgentId, double?>();
-			foreach (ILease Lease in Leases)
+			using (IScope _ = GlobalTracer.Instance.BuildSpan($"GenerateResponses").StartActive())
 			{
-				double? AgentRate = null;
-				if (IncludeCosts && !CachedAgentRates.TryGetValue(Lease.AgentId, out AgentRate))
+				Dictionary<AgentId, double?> CachedAgentRates = new Dictionary<AgentId, double?>();
+				foreach (ILease Lease in Leases)
 				{
-					AgentRate = await AgentService.GetRateAsync(Lease.AgentId);
-					CachedAgentRates.Add(Lease.AgentId, AgentRate);
+					double? AgentRate = null;
+					if (IncludeCosts && !CachedAgentRates.TryGetValue(Lease.AgentId, out AgentRate))
+					{
+						AgentRate = await AgentService.GetRateAsync(Lease.AgentId);
+						CachedAgentRates.Add(Lease.AgentId, AgentRate);
+					}
+
+					Dictionary<string, string>? Details = AgentService.GetPayloadDetails(Lease.Payload);
+					Responses.Add(PropertyFilter.Apply(new GetAgentLeaseResponse(Lease, Details, AgentRate), Filter));
 				}
-
-				Dictionary<string, string>? Details = AgentService.GetPayloadDetails(Lease.Payload);
-				Responses.Add(PropertyFilter.Apply(new GetAgentLeaseResponse(Lease, Details, AgentRate), Filter));
 			}
 
 			return Responses;
