Horde: Add optimized path for querying leases by finish time

Query performance for leases have proven to be problematic. This exposes an optimized path all the way to the controller.

diff --git a/Engine/Source/Programs/Horde/Horde.Build.Tests/LeasesControllerDbTest.cs b/Engine/Source/Programs/Horde/Horde.Build.Tests/LeasesControllerDbTest.cs
new file mode 100644
--- /dev/null
+++ b/Engine/Source/Programs/Horde/Horde.Build.Tests/LeasesControllerDbTest.cs
@@ -0,0 +1,47 @@
+// Copyright Epic Games, Inc. All Rights Reserved.
+
+using System;
+using System.Collections.Generic;
+using System.Threading.Tasks;
+using HordeCommon;
+using HordeServer.Api;
+using HordeServer.Models;
+using HordeServer.Utilities;
+using HordeServerTests;
+using Microsoft.AspNetCore.Mvc;
+using Microsoft.VisualStudio.TestTools.UnitTesting;
+
+namespace Horde.Build.Tests
+{
+	/// <summary>
+	/// Database-only integration test for LeasesController
+	/// </summary>
+	[TestClass]
+	public class LeasesControllerDbTest : TestSetup
+	{
+		[TestMethod]
+		public async Task FindLeases()
+		{
+			DateTimeOffset MinTime = Clock.UtcNow - TimeSpan.FromMinutes(5);
+			DateTimeOffset MaxTime = Clock.UtcNow;
+
+			ILease Lease1 = await CreateLease(Clock.UtcNow - TimeSpan.FromMinutes(10), TimeSpan.FromMinutes(6));
+			ILease Lease2 = await CreateLease(Clock.UtcNow - TimeSpan.FromMinutes(7), TimeSpan.FromMinutes(3));
+			ILease OutOfTimeWindow = await CreateLease(Clock.UtcNow - TimeSpan.FromMinutes(7), TimeSpan.FromMinutes(25));
+			
+			ActionResult<List<object>> Res = await LeasesController.FindLeasesAsync(null, null, null, null, MinTime, MaxTime);
+			Assert.AreEqual(2, Res.Value!.Count);
+			Assert.AreEqual(Lease2.Id.ToString(), (Res.Value[0] as GetAgentLeaseResponse)!.Id);
+			Assert.AreEqual(Lease1.Id.ToString(), (Res.Value[1] as GetAgentLeaseResponse)!.Id);
+		}
+
+		private async Task<ILease> CreateLease(DateTime StartTime, TimeSpan Duration)
+		{
+			ObjectId<ILease> Id = ObjectId<ILease>.GenerateNewId();
+			ObjectId<ISession> SessionId = ObjectId<ISession>.GenerateNewId();
+			ILease Lease = await LeaseCollection.AddAsync(Id, "myLease", new AgentId("agent-1"), SessionId, null, null, null, StartTime, new byte[] { });
+			await LeaseCollection.TrySetOutcomeAsync(Id, StartTime + Duration, LeaseOutcome.Success, null);
+			return Lease;
+		}
+	}
+}
\ No newline at end of file
diff --git a/Engine/Source/Programs/Horde/Horde.Build.Tests/TestSetup.cs b/Engine/Source/Programs/Horde/Horde.Build.Tests/TestSetup.cs
--- a/Engine/Source/Programs/Horde/Horde.Build.Tests/TestSetup.cs
+++ b/Engine/Source/Programs/Horde/Horde.Build.Tests/TestSetup.cs
@@ -99,6 +99,7 @@ namespace HordeServerTests
 		public JobsController JobsController => GetJobsController();
 		public AgentsController AgentsController => GetAgentsController();
 		public PoolsController PoolsController => GetPoolsController();
+		public LeasesController LeasesController => GetLeasesController();
 
 		private static bool IsDatadogWriterPatched;
 
@@ -239,6 +240,14 @@ namespace HordeServerTests
 			return Controller;
 		}
 		
+		private LeasesController GetLeasesController()
+		{
+			ILogger<LeasesController> Logger = ServiceProvider.GetRequiredService<ILogger<LeasesController>>();
+			LeasesController Controller = new LeasesController(AclService, AgentService, Logger);
+			Controller.ControllerContext = GetControllerContext();
+			return Controller;
+		}
+		
 		private ControllerContext GetControllerContext()
 		{
 			ControllerContext ControllerContext = new ControllerContext();
diff --git a/Engine/Source/Programs/Horde/Horde.Build/Agents/ILeaseCollection.cs b/Engine/Source/Programs/Horde/Horde.Build/Agents/ILeaseCollection.cs
--- a/Engine/Source/Programs/Horde/Horde.Build/Agents/ILeaseCollection.cs
+++ b/Engine/Source/Programs/Horde/Horde.Build/Agents/ILeaseCollection.cs
@@ -68,6 +68,18 @@ namespace HordeServer.Collections
 		/// <returns>List of leases matching the given criteria</returns>
 		Task<List<ILease>> FindLeasesAsync(AgentId? AgentId = null, SessionId? SessionId = null, DateTime? MinTime = null, DateTime? MaxTime = null, int? Index = null, int? Count = null, string? IndexHint = null, bool ConsistentRead = true);
 
+		/// <summary>
+		/// Finds all leases by finish time
+		/// </summary>
+		/// <param name="MinFinishTime">Start of finish time window</param>
+		/// <param name="MaxFinishTime">End of finish time window</param>
+		/// <param name="Index">Index of the first result to return</param>
+		/// <param name="Count">Number of results to return</param>
+		/// <param name="IndexHint">Name of index to be specified as a hint to the database query planner</param>
+		/// <param name="ConsistentRead">If the database read should be made to the replica server</param>
+		/// <returns>List of leases matching the given criteria</returns>
+		Task<List<ILease>> FindLeasesByFinishTimeAsync(DateTime? MinFinishTime, DateTime? MaxFinishTime, int? Index, int? Count, string? IndexHint = null, bool ConsistentRead = true);
+
 		/// <summary>
 		/// Finds all leases between min and/or max time
 		/// Preferred over the more generic FindLeasesAsync as this one use a better index and relaxed read consistency
diff --git a/Engine/Source/Programs/Horde/Horde.Build/Agents/Impl/AgentService.cs b/Engine/Source/Programs/Horde/Horde.Build/Agents/Impl/AgentService.cs
--- a/Engine/Source/Programs/Horde/Horde.Build/Agents/Impl/AgentService.cs
+++ b/Engine/Source/Programs/Horde/Horde.Build/Agents/Impl/AgentService.cs
@@ -760,6 +760,19 @@ namespace HordeServer.Services
 			return Leases.FindLeasesAsync(AgentId, SessionId, StartTime, FinishTime, Index, Count);
 		}
 
+		/// <summary>
+		/// Finds all leases by finish time
+		/// </summary>
+		/// <param name="MinFinishTime">Start of the search window to return results for</param>
+		/// <param name="MaxFinishTime">End of the search window to return results for</param>
+		/// <param name="Index">Index of the first result to return</param>
+		/// <param name="Count">Number of results to return</param>
+		/// <returns>List of leases matching the given criteria</returns>
+		public Task<List<ILease>> FindLeasesByFinishTimeAsync(DateTime? MinFinishTime, DateTime? MaxFinishTime, int? Index, int? Count)
+		{
+			return Leases.FindLeasesByFinishTimeAsync(MinFinishTime, MaxFinishTime, Index, Count, null, false);
+		}
+
 		/// <summary>
 		/// Gets a specific lease
 		/// </summary>
diff --git a/Engine/Source/Programs/Horde/Horde.Build/Agents/Impl/LeaseCollection.cs b/Engine/Source/Programs/Horde/Horde.Build/Agents/Impl/LeaseCollection.cs
--- a/Engine/Source/Programs/Horde/Horde.Build/Agents/Impl/LeaseCollection.cs
+++ b/Engine/Source/Programs/Horde/Horde.Build/Agents/Impl/LeaseCollection.cs
@@ -159,6 +159,27 @@ namespace HordeServer.Collections.Impl
 			List<LeaseDocument> Results = await Collection.Find(Filter, FindOptions).SortByDescending(x => x.StartTime).Range(Index, Count).ToListAsync();
 			return Results.ConvertAll<ILease>(x => x);
 		}
+		
+		/// <inheritdoc/>
+		public async Task<List<ILease>> FindLeasesByFinishTimeAsync(DateTime? MinFinishTime, DateTime? MaxFinishTime, int? Index, int? Count, string? IndexHint, bool ConsistentRead)
+		{
+			IMongoCollection<LeaseDocument> Collection = ConsistentRead ? Leases : Leases.WithReadPreference(ReadPreference.SecondaryPreferred);
+			FilterDefinitionBuilder<LeaseDocument> FilterBuilder = Builders<LeaseDocument>.Filter;
+			FilterDefinition<LeaseDocument> Filter = FilterDefinition<LeaseDocument>.Empty;
+
+			if (MinFinishTime == null && MaxFinishTime == null)
+			{
+				throw new ArgumentException($"Both {nameof(MinFinishTime)} and {nameof(MaxFinishTime)} cannot be null");
+			}
+
+			if (MinFinishTime != null) Filter &= FilterBuilder.Gt(x => x.FinishTime, MinFinishTime.Value);
+			if (MaxFinishTime != null) Filter &= FilterBuilder.Lt(x => x.FinishTime, MaxFinishTime.Value);
+
+			FindOptions? FindOptions = IndexHint == null ? null : new FindOptions { Hint = new BsonString(IndexHint) };
+			List<LeaseDocument> Results = await Collection.Find(Filter, FindOptions).SortByDescending(x => x.FinishTime).Range(Index, Count).ToListAsync();
+
+			return Results.ConvertAll<ILease>(x => x);
+		}
 
 		/// <inheritdoc/>
 		public async Task<List<ILease>> FindLeasesAsync(DateTime? MinTime, DateTime? MaxTime)
diff --git a/Engine/Source/Programs/Horde/Horde.Build/Controllers/LeasesController.cs b/Engine/Source/Programs/Horde/Horde.Build/Controllers/LeasesController.cs
--- a/Engine/Source/Programs/Horde/Horde.Build/Controllers/LeasesController.cs
+++ b/Engine/Source/Programs/Horde/Horde.Build/Controllers/LeasesController.cs
@@ -74,6 +74,8 @@ namespace HordeServer.Controllers
 		/// <param name="SessionId">The session to query</param>
 		/// <param name="StartTime">Start of the time window to consider</param>
 		/// <param name="FinishTime">End of the time window to consider</param>
+		/// <param name="MinFinishTime">Start of the time window to consider when querying by finish time (if set, other criteria are ignored)</param>
+		/// <param name="MaxFinishTime">End of the time window to consider when querying by finish time (if set, other criteria are ignored)</param>
 		/// <param name="Index">Index of the first result to return</param>
 		/// <param name="Count">Number of results to return</param>
 		/// <param name="Filter">Filter to apply to the results</param>
@@ -81,7 +83,16 @@ namespace HordeServer.Controllers
 		[HttpGet]
 		[Route("/api/v1/leases")]
 		[ProducesResponseType(200, Type = typeof(List<GetAgentLeaseResponse>))]
-		public async Task<ActionResult<List<object>>> FindLeasesAsync([FromQuery] AgentId? AgentId, [FromQuery] SessionId? SessionId, [FromQuery] DateTimeOffset? StartTime, [FromQuery] DateTimeOffset? FinishTime, [FromQuery] int Index = 0, [FromQuery] int Count = 1000, [FromQuery] PropertyFilter? Filter = null)
+		public async Task<ActionResult<List<object>>> FindLeasesAsync(
+			[FromQuery] AgentId? AgentId,
+			[FromQuery] SessionId? SessionId,
+			[FromQuery] DateTimeOffset? StartTime,
+			[FromQuery] DateTimeOffset? FinishTime,
+			[FromQuery] DateTimeOffset? MinFinishTime,
+			[FromQuery] DateTimeOffset? MaxFinishTime,
+			[FromQuery] int Index = 0,
+			[FromQuery] int Count = 1000,
+			[FromQuery] PropertyFilter? Filter = null)
 		{
 			GlobalPermissionsCache PermissionsCache = new GlobalPermissionsCache();
 			if (AgentId == null)
@@ -106,7 +117,16 @@ namespace HordeServer.Controllers
 
 			bool IncludeCosts = await AclService.AuthorizeAsync(AclAction.ViewCosts, User, PermissionsCache);
 
-			List<ILease> Leases = await AgentService.FindLeasesAsync(AgentId, SessionId, StartTime?.UtcDateTime, FinishTime?.UtcDateTime, Index, Count);
+			List<ILease> Leases;
+			if (MinFinishTime == null && MaxFinishTime == null)
+			{
+				Leases = await AgentService.FindLeasesAsync(AgentId, SessionId, StartTime?.UtcDateTime, FinishTime?.UtcDateTime, Index, Count);
+			}
+			else
+			{
+				// Optimized path for queries made by finish time
+				Leases = await AgentService.FindLeasesByFinishTimeAsync(MinFinishTime?.UtcDateTime, MaxFinishTime?.UtcDateTime, Index, Count);
+			}
 
 			List<object> Responses = new List<object>();
 
