Horde: Use the CommandHost for dispatching different commands within the server.

diff --git a/Engine/Source/Programs/Horde/Horde.Build/Program.cs b/Engine/Source/Programs/Horde/Horde.Build/Program.cs
--- a/Engine/Source/Programs/Horde/Horde.Build/Program.cs
+++ b/Engine/Source/Programs/Horde/Horde.Build/Program.cs
@@ -43,6 +43,8 @@ using Serilog.Filters;
 using Serilog.Formatting.Json;
 using Serilog.Sinks.SystemConsole.Themes;
 using OpenTracing.Propagation;
+using Microsoft.Extensions.DependencyInjection;
+using Horde.Build.Commands;
 
 namespace HordeServer
 {
@@ -145,7 +147,7 @@ namespace HordeServer
 			return SchemaTypes.ToArray();
 		}
 
-		public static async Task Main(string[] Args)
+		public static async Task<int> Main(string[] Args)
 		{
 			CommandLineArguments Arguments = new CommandLineArguments(Args);
 
@@ -157,7 +159,6 @@ namespace HordeServer
 				.AddJsonFile("appsettings.User.json", optional: true)
 				.AddJsonFile(UserConfigFile.FullName, optional: true, reloadOnChange: true)
 				.AddEnvironmentVariables()
-				.AddCommandLine(Args)
 				.Build();
 
 			ServerSettings HordeSettings = new ServerSettings();
@@ -190,125 +191,15 @@ namespace HordeServer
 				Serilog.Log.Logger.Information("Enabling datadog tracing (OpenTrace)");
 			}
 
-			if (Arguments.HasOption("-UpdateSchemas"))
-			{
-				DirectoryReference? SchemaDir = Arguments.GetDirectoryReferenceOrDefault("-SchemaDir=", null);
-				if (SchemaDir == null)
-				{
-					SchemaDir = new DirectoryReference("Schemas");
-				}
-
-				Arguments.CheckAllArgumentsUsed();
-
-				DirectoryReference.CreateDirectory(SchemaDir);
-				foreach (Type SchemaType in ConfigSchemas)
-				{
-					FileReference OutputFile = FileReference.Combine(SchemaDir, $"{SchemaType.Name}.json");
-					Schemas.CreateSchema(SchemaType).Write(OutputFile);
-				}
-				return;
-			}
-
-			try
-			{
-				using (X509Certificate2? GrpcCertificate = ReadGrpcCertificate(HordeSettings))
-				{
-					List<IHost> Hosts = new List<IHost>();
-					Hosts.Add(CreateHostBuilderWithCert(Args, Config, HordeSettings, GrpcCertificate).Build());
-#if WITH_HORDE_STORAGE
-					IHostBuilder StorageHostBuilder = Horde.Storage.Program.CreateHostBuilder(Args);
-					StorageHostBuilder.ConfigureWebHostDefaults(Builder =>
-					{
-						Builder.ConfigureKestrel(Options =>
-						{
-							Options.ListenAnyIP(57000);
-							Options.ListenAnyIP(57001, Configure => Configure.UseHttps());
-						});
-					});
-					Hosts.Add(StorageHostBuilder.Build());
-#endif
-					await Task.WhenAll(Hosts.Select(x => x.RunAsync()));
-				}
-			}
-#pragma warning disable CA1031 // Do not catch general exception types
-			catch (Exception Ex)
-#pragma warning restore CA1031 // Do not catch general exception types
-			{
-				Serilog.Log.Logger.Error(Ex, "Unhandled exception");
-			}
-		}
-
-		// Used by WebApplicationFactory in controller tests. Uses reflection to call this exact function signature.
-		public static IHostBuilder CreateHostBuilder(string[] Args)
-		{
-			ServerSettings HordeSettings = new ServerSettings();
-			return CreateHostBuilderWithCert(Args, new ConfigurationBuilder().Build(), HordeSettings, null);
-		}
+			IServiceCollection Services = new ServiceCollection();
+			Services.AddCommandsFromAssembly(Assembly.GetExecutingAssembly());
+			Services.AddLogging(Builder => Builder.AddSerilog());
+			Services.AddSingleton<IConfiguration>(Config);
 
-		public static IHostBuilder CreateHostBuilderWithCert(string[] Args, IConfiguration Config, ServerSettings ServerSettings, X509Certificate2? SslCert)
-		{
-			return Host.CreateDefaultBuilder(Args)
-				.UseSerilog()
-				.ConfigureAppConfiguration(Builder => Builder.AddConfiguration(Config))
-				.ConfigureWebHostDefaults(WebBuilder => 
-				{
-					WebBuilder.ConfigureKestrel(Options =>
-					{
-						Options.Limits.MaxRequestBodySize = 100 * 1024 * 1024;
-
-						if (ServerSettings.HttpPort != 0)
-						{
-							Options.ListenAnyIP(ServerSettings.HttpPort, Configure => { Configure.Protocols = HttpProtocols.Http1AndHttp2; });
-						}
-
-						if (ServerSettings.HttpsPort != 0)
-						{
-							Options.ListenAnyIP(ServerSettings.HttpsPort, Configure => { if (SslCert != null) { Configure.UseHttps(SslCert); } });
-						}
-
-						// To serve HTTP/2 with gRPC *without* TLS enabled, a separate port for HTTP/2 must be used.
-						// This is useful when having a load balancer in front that terminates TLS.
-						if (ServerSettings.Http2Port != 0)
-						{
-							Options.ListenAnyIP(ServerSettings.Http2Port, Configure => { Configure.Protocols = HttpProtocols.Http2; });
-						}
-					});
-					WebBuilder.UseStartup<Startup>(); 
-				});
-		}
-
-		/// <summary>
-		/// Gets the certificate to use for Grpc endpoints
-		/// </summary>
-		/// <returns>Custom certificate to use for Grpc endpoints, or null for the default.</returns>
-		public static X509Certificate2? ReadGrpcCertificate(ServerSettings HordeSettings)
-		{
-			string Base64Prefix = "base64:";
-			
-			if (HordeSettings.ServerPrivateCert == null)
-			{
-				return null;
-			}
-			else if (HordeSettings.ServerPrivateCert.StartsWith(Base64Prefix, StringComparison.Ordinal))
-			{
-				byte[] CertData = Convert.FromBase64String(HordeSettings.ServerPrivateCert.Replace(Base64Prefix, "", StringComparison.Ordinal));
-				return new X509Certificate2(CertData);
-			}
-			else
-			{
-				FileReference? ServerPrivateCert = null;
-
-				if (!Path.IsPathRooted(HordeSettings.ServerPrivateCert))
-				{
-					ServerPrivateCert = FileReference.Combine(AppDir, HordeSettings.ServerPrivateCert);
-				}
-				else
-				{
-					ServerPrivateCert = new FileReference(HordeSettings.ServerPrivateCert);
-				}
-				
-				return new X509Certificate2(FileReference.ReadAllBytes(ServerPrivateCert));
-			}
+#pragma warning disable ASP0000 // Do not call 'IServiceCollection.BuildServiceProvider' in 'ConfigureServices'
+			IServiceProvider ServiceProvider = Services.BuildServiceProvider();
+			return await CommandHost.RunAsync(Arguments, ServiceProvider, typeof(ServerCommand));
+#pragma warning restore ASP0000 // Do not call 'IServiceCollection.BuildServiceProvider' in 'ConfigureServices'
 		}
 
 		/// <summary>
