Horde Storage - Introduced a batch api for refs, this allows you to GET/PUT and check for existence (HEAD) in batch. This requires compact binary response types as we want to return the ref payload inline in the response.
Also removed the old Compact binary serialization code as its not used anymore.
Lastly I fix the old exist batch api (that uses query strings) as it wasnt working and had no tests (this shouldnt really need to be used though as the new api replaces it but a simpler api might be useful for someone)

diff --git a/Engine/Source/Programs/Horde/HordeStorage/Horde.Storage/Controllers/ReferencesController.cs b/Engine/Source/Programs/Horde/HordeStorage/Horde.Storage/Controllers/ReferencesController.cs
--- a/Engine/Source/Programs/Horde/HordeStorage/Horde.Storage/Controllers/ReferencesController.cs
+++ b/Engine/Source/Programs/Horde/HordeStorage/Horde.Storage/Controllers/ReferencesController.cs
@@ -7,8 +7,10 @@ using System.ComponentModel.DataAnnotations;
 using System.Globalization;
 using System.IO;
 using System.Linq;
+using System.Net;
 using System.Net.Mime;
 using System.Text;
+using System.Threading;
 using System.Threading.Tasks;
 using Dasync.Collections;
 using Datadog.Trace;
@@ -24,7 +26,6 @@ using Microsoft.AspNetCore.Authorization;
 using Microsoft.AspNetCore.Http;
 using Microsoft.AspNetCore.Http.Extensions;
 using Microsoft.AspNetCore.Mvc;
-using Microsoft.Extensions.Options;
 using Newtonsoft.Json;
 using Serilog;
 using ContentHash = Jupiter.Implementation.ContentHash;
@@ -42,7 +43,6 @@ namespace Horde.Storage.Controllers
     {
         private readonly IDiagnosticContext _diagnosticContext;
         private readonly IAuthorizationService _authorizationService;
-        private readonly IOptionsMonitor<JupiterSettings> _jupiterSettings;
         private readonly FormatResolver _formatResolver;
         private readonly BufferedPayloadFactory _bufferedPayloadFactory;
         private readonly IReferenceResolver _referenceResolver;
@@ -51,13 +51,12 @@ namespace Horde.Storage.Controllers
         private readonly IObjectService _objectService;
         private readonly IBlobService _blobStore;
 
-        public ReferencesController(IObjectService objectService, IBlobService blobStore, IDiagnosticContext diagnosticContext, IAuthorizationService authorizationService, IOptionsMonitor<JupiterSettings> jupiterSettings, FormatResolver formatResolver, BufferedPayloadFactory bufferedPayloadFactory, IReferenceResolver referenceResolver)
+        public ReferencesController(IObjectService objectService, IBlobService blobStore, IDiagnosticContext diagnosticContext, IAuthorizationService authorizationService, FormatResolver formatResolver, BufferedPayloadFactory bufferedPayloadFactory, IReferenceResolver referenceResolver)
         {
             _objectService = objectService;
             _blobStore = blobStore;
             _diagnosticContext = diagnosticContext;
             _authorizationService = authorizationService;
-            _jupiterSettings = jupiterSettings;
             _formatResolver = formatResolver;
             _bufferedPayloadFactory = bufferedPayloadFactory;
             _referenceResolver = referenceResolver;
@@ -398,13 +397,12 @@ namespace Horde.Storage.Controllers
             return Ok();
         }
 
-        [HttpPost("{ns}/exists")]
+        [HttpGet("{ns}/exists")]
         [Authorize("Object.read")]
         [ProducesDefaultResponseType]
         public async Task<IActionResult> ExistsMultiple(
             [FromRoute] [Required] NamespaceId ns,
-            [FromRoute] [Required] BucketId bucket,
-            [FromQuery] [Required] List<IoHashKey> names)
+            [FromQuery] [Required] List<string> names)
         {
             AuthorizationResult authorizationResult = await _authorizationService.AuthorizeAsync(User, ns, NamespaceAccessRequirement.Name);
 
@@ -413,14 +411,26 @@ namespace Horde.Storage.Controllers
                 return Forbid();
             }
 
-            ConcurrentBag<IoHashKey> missingObject = new ConcurrentBag<IoHashKey>();
+            ConcurrentBag<(BucketId, IoHashKey)> missingObject = new ();
 
-            IEnumerable<Task> tasks = names.Select(async name =>
+            List<(BucketId, IoHashKey)> requestedNames = new List<(BucketId, IoHashKey)>();
+            foreach (string name in names)
             {
+                int separatorIndex = name.IndexOf(".", StringComparison.Ordinal);
+                if (separatorIndex == -1)
+                    return BadRequest(new ProblemDetails() { Title = $"Key {name} did not contain a '.' separator" });
+                BucketId bucket = new BucketId(name.Substring(0, separatorIndex));
+                IoHashKey key = new IoHashKey(name.Substring(separatorIndex + 1));
+                requestedNames.Add((bucket, key));
+            }
+
+            IEnumerable<Task> tasks = requestedNames.Select(async pair =>
+            {
+                (BucketId bucket, IoHashKey key) = pair;
                 try
                 {
                     (ObjectRecord record, BlobContents? blob) =
-                        await _objectService.Get(ns, bucket, name, new string[] { "blobIdentifier" });
+                        await _objectService.Get(ns, bucket, key, new string[] { "blobIdentifier" });
 
                     blob ??= await _blobStore.GetObject(ns, record.BlobIdentifier);
 
@@ -434,15 +444,15 @@ namespace Horde.Storage.Controllers
                 }
                 catch (ObjectNotFoundException)
                 {
-                    missingObject.Add(name);
+                    missingObject.Add((bucket, key));
                 }
                 catch (PartialReferenceResolveException)
                 {
-                    missingObject.Add(name);
+                    missingObject.Add((bucket, key));
                 }
                 catch (ReferenceIsMissingBlobsException)
                 {
-                    missingObject.Add(name);
+                    missingObject.Add((bucket, key));
                 }
             });
             await Task.WhenAll(tasks);
@@ -578,7 +588,167 @@ namespace Horde.Storage.Controllers
         }
 
 
-        
+        [HttpPost("{ns}")]
+        [Consumes(CustomMediaTypeNames.UnrealCompactBinary)]
+        [Produces(CustomMediaTypeNames.UnrealCompactBinary)]
+        public async Task<IActionResult> Batch(
+            [FromRoute] [Required] NamespaceId ns,
+            [FromBody] [Required] BatchOps ops)
+        {
+            {
+                using IScope _ = Tracer.Instance.StartActive("authorize");
+                AuthorizationResult authorizationResult =
+                    await _authorizationService.AuthorizeAsync(User, ns, NamespaceAccessRequirement.Name);
+
+                if (!authorizationResult.Succeeded)
+                {
+                    return Forbid();
+                }
+            }
+
+            ConcurrentDictionary<uint, (CbObject, HttpStatusCode)> results = new();
+
+            async Task<(CbObject, HttpStatusCode)> BatchGetOp(BatchOps.BatchOp op)
+            {
+                try
+                {
+                    (ObjectRecord objectRecord, BlobContents? blob) = await _objectService.Get(ns, op.Bucket, op.Key, Array.Empty<string>());
+
+                    if (!objectRecord.IsFinalized)
+                    {
+                        throw new Exception("Object is not finalized");
+                    }
+
+                    if (blob == null)
+                        throw new Exception();
+
+                    CbObject cb = new CbObject(await blob.Stream.ToByteArray());
+                    return (cb, HttpStatusCode.OK);
+                }
+                catch (ObjectNotFoundException e)
+                {
+                    return ToErrorResult(e, HttpStatusCode.NotFound);
+                }
+                catch (Exception e)
+                {
+                    return ToErrorResult(e);
+                }
+            }
+
+            async Task<(CbObject, HttpStatusCode)> BatchHeadOp(BatchOps.BatchOp op)
+            {
+                try
+                {
+                    (ObjectRecord record, BlobContents? blob) = await _objectService.Get(ns, op.Bucket, op.Key, new string[] { "blobIdentifier" });
+
+                    if (!record.IsFinalized)
+                    {
+                        return (CbObject.Build(writer => writer.WriteBool("exists", false)), HttpStatusCode.NotFound);
+                    }
+
+                    blob ??= await _blobStore.GetObject(ns, record.BlobIdentifier);
+
+                    // we have to verify the blobs are available locally, as the record of the key is replicated a head of the content
+                    // TODO: Once we support inline replication this step is not needed as at least one region as this blob, just maybe not this current one
+                    byte[] blobContents = await blob.Stream.ToByteArray();
+                    CbObject cb = new CbObject(blobContents);
+                    // the reference resolver will throw if any blob is missing, so no need to do anything other then process each reference
+                    IAsyncEnumerable<BlobIdentifier> references = _referenceResolver.ResolveReferences(ns, cb);
+                    List<BlobIdentifier>? _ = await references.ToListAsync();
+
+                    if (blob == null)
+                        throw new Exception();
+
+                    return (CbObject.Build(writer => writer.WriteBool("exists", true)), HttpStatusCode.OK);
+                }
+                catch (Exception ex) when( ex is ObjectNotFoundException || ex is PartialReferenceResolveException || ex is ReferenceIsMissingBlobsException)
+                {
+                    return (CbObject.Build(writer => writer.WriteBool("exists", false)), HttpStatusCode.NotFound);
+                }
+                catch (Exception e)
+                {
+                    return ToErrorResult(e);
+                }
+            }
+
+            async Task<(CbObject, HttpStatusCode)> BatchPutOp(BatchOps.BatchOp op)
+            {
+                try
+                {
+                    if (op.Payload == null || op.Payload.Equals(CbObject.Empty))
+                    {
+                        throw new Exception($"Missing payload for operation: {op.OpId}");
+                    }
+
+                    if (op.PayloadHash == null)
+                    {
+                        throw new Exception($"Missing payload hash for operation: {op.OpId}");
+                    }
+                    BlobIdentifier headerHash = BlobIdentifier.FromContentHash(op.PayloadHash);
+                    BlobIdentifier objectHash = BlobIdentifier.FromBlob(op.Payload.GetView().ToArray());
+
+                    if (!headerHash.Equals(objectHash))
+                    {
+                        throw new HashMismatchException(headerHash, objectHash);
+                    }
+
+                    (ContentId[] missingReferences, BlobIdentifier[] missingBlobs) = await _objectService.Put(ns, op.Bucket, op.Key, objectHash, op.Payload);
+                    List<ContentHash> missingHashes = new List<ContentHash>(missingReferences);
+
+                    return (CbSerializer.Serialize(new PutObjectResponse(missingHashes.ToArray())), HttpStatusCode.OK);
+                }
+                catch (Exception e)
+                {
+                    return ToErrorResult(e);
+                }
+            }
+
+            await Parallel.ForEachAsync(ops.Ops, CancellationToken.None, async (op, token) =>
+            {
+                switch (op.Op)
+                {
+                    case BatchOps.BatchOp.Operation.GET:
+                        results.TryAdd(op.OpId, await BatchGetOp(op));
+                        break;
+                    case BatchOps.BatchOp.Operation.PUT:
+                        results.TryAdd(op.OpId, await BatchPutOp(op));
+                        break;
+                    case BatchOps.BatchOp.Operation.HEAD:
+                        results.TryAdd(op.OpId, await BatchHeadOp(op));
+                        break;
+                    case BatchOps.BatchOp.Operation.INVALID:
+                    default:
+                        throw new ArgumentOutOfRangeException();
+                }
+
+                await Task.CompletedTask;
+            });
+
+            return Ok(new BatchOpsResponse()
+            {
+                Results = results.Select(result =>
+                {
+                    return new BatchOpsResponse.OpResponses()
+                    {
+                        OpId = result.Key,
+                        Response = result.Value.Item1,
+                        StatusCode = (int)result.Value.Item2
+                    };
+                }).ToList()
+            });
+        }
+
+        private (CbObject, HttpStatusCode) ToErrorResult(Exception exception, HttpStatusCode statusCode = HttpStatusCode.InternalServerError)
+        {
+            Exception e = exception;
+            CbWriter writer = new CbWriter();
+            writer.BeginObject();
+            writer.WriteString("title", e.Message);
+            writer.WriteInteger("status", (int)statusCode);
+            writer.EndObject();
+            return (writer.ToObject(), statusCode);
+        }
+
         /// <summary>
         /// Drop all refs records in the namespace
         /// </summary>
@@ -709,6 +879,91 @@ namespace Horde.Storage.Controllers
         public long CountOfDeletedRecords { get; set; }
     }
 
+    public class BatchOps
+    {
+        public BatchOps()
+        {
+            Ops = Array.Empty<BatchOp>();
+        }
+
+        public class BatchOp
+        {
+            public BatchOp()
+            {
+                Payload = null;
+                PayloadHash = null;
+            }
+
+            public enum Operation
+            {
+                INVALID,
+                GET,
+                PUT,
+                HEAD,
+            }
+
+            [Required]
+            [CbField("opId")]
+            public uint OpId { get; set; }
+
+
+            [CbField("op")]
+            [JsonIgnore]
+            public string OpString
+            {
+                get { return Op.ToString(); }
+                set { Op = Enum.Parse<Operation>(value); }
+            }
+
+            [Required]
+            public Operation Op { get; set; } = Operation.INVALID;
+
+            [Required]
+            [CbField("bucket")]
+            public BucketId Bucket { get; set; }
+
+            [Required]
+            [CbField("key")]
+            public IoHashKey Key { get; set; }
+
+            [CbField("payload")]
+            public CbObject? Payload { get; set; } = null;
+
+            [CbField("payloadHash")] 
+            public ContentHash? PayloadHash { get; set; } = null;
+        }
+
+        [CbField("ops")]
+        public BatchOp[] Ops { get; set; }
+    }
+
+    public class BatchOpsResponse
+    {
+        public BatchOpsResponse()
+        {
+
+        }
+
+        public class OpResponses
+        {
+            public OpResponses()
+            {
+
+            }
+            [CbField("opId")]
+            public uint OpId { get; set; }
+
+            [CbField("response")]
+            public CbObject Response { get; set; } = null!;
+
+            [CbField("statusCode")]
+            public int StatusCode { get; set; }
+        }
+
+        [CbField("results")]
+        public List<OpResponses> Results { get; set; } = new List<OpResponses>();
+    }
+
     public class RefMetadataResponse
     {
         public RefMetadataResponse()
@@ -780,16 +1035,35 @@ namespace Horde.Storage.Controllers
 
     public class ExistCheckMultipleRefsResponse
     {
-        public ExistCheckMultipleRefsResponse(List<IoHashKey> missingNames)
+        public ExistCheckMultipleRefsResponse(List<(BucketId,IoHashKey)> missingNames)
         {
-            Needs = missingNames;
-            Names = missingNames;
+            Missing = missingNames.Select(pair =>
+            {
+                (BucketId bucketId, IoHashKey ioHashKey) = pair;
+                return new MissingReference()
+                {
+                    Bucket = bucketId,
+                    Key = ioHashKey,
+                };
+            }).ToList();
         }
 
-        [CbField("needs")]
-        public List<IoHashKey> Needs { get; set; }
+        [JsonConstructor]
+        public ExistCheckMultipleRefsResponse(List<MissingReference> missingNames)
+        {
+            Missing = missingNames;
+        }
 
-        [CbField("names")]
-        public List<IoHashKey> Names { get; set; }
+        [CbField("missing")]
+        public List<MissingReference> Missing { get; set; }
+
+        public class MissingReference
+        {
+            [CbField("bucket")]
+            public BucketId Bucket { get; set; }
+
+            [CbField("key")]
+            public IoHashKey Key { get; set; }
+        }
     }
 }
diff --git a/Engine/Source/Programs/Horde/HordeStorage/Horde.Storage/Tests/Functional/Objects/ReferencesTests.cs b/Engine/Source/Programs/Horde/HordeStorage/Horde.Storage/Tests/Functional/Objects/ReferencesTests.cs
--- a/Engine/Source/Programs/Horde/HordeStorage/Horde.Storage/Tests/Functional/Objects/ReferencesTests.cs
+++ b/Engine/Source/Programs/Horde/HordeStorage/Horde.Storage/Tests/Functional/Objects/ReferencesTests.cs
@@ -29,6 +29,7 @@ using Serilog;
 using Logger = Serilog.Core.Logger;
 using EpicGames.Horde.Storage;
 using EpicGames.Serialization;
+using ContentHash = Jupiter.Implementation.ContentHash;
 
 namespace Horde.Storage.FunctionalTests.References
 {
@@ -633,6 +634,40 @@ namespace Horde.Storage.FunctionalTests.References
             }
         }
 
+        [TestMethod]
+        public async Task ExistsChecksMultiple()
+        {
+            BucketId bucket = new BucketId("bucket");
+            IoHashKey existingObject = IoHashKey.FromName("existingObject");
+
+            const string objectContents = "This is treated as a opaque blob";
+            byte[] data = Encoding.ASCII.GetBytes(objectContents);
+            BlobIdentifier objectHash = BlobIdentifier.FromBlob(data);
+            HttpContent requestContent = new ByteArrayContent(data);
+            requestContent.Headers.ContentType = new MediaTypeHeaderValue(MediaTypeNames.Application.Octet);
+            requestContent.Headers.Add(CommonHeaders.HashHeaderName, objectHash.ToString());
+
+            {
+                HttpResponseMessage result = await _httpClient!.PutAsync(requestUri: $"api/v1/refs/{TestNamespace}/{bucket}/{existingObject}", requestContent);
+                result.EnsureSuccessStatusCode();
+            }
+
+            IoHashKey missingObject = IoHashKey.FromName("missingObject");
+
+            string queryString = $"?names={bucket}.{existingObject}&names={bucket}.{missingObject}";
+
+            {
+                HttpResponseMessage result = await _httpClient!.GetAsync(requestUri: $"api/v1/refs/{TestNamespace}/exists" + queryString);
+                result.EnsureSuccessStatusCode();
+                ExistCheckMultipleRefsResponse response = await result.Content.ReadAsAsync<ExistCheckMultipleRefsResponse>();
+
+                Assert.AreEqual(1, response.Missing.Count);
+                Assert.AreEqual(bucket, response.Missing[0].Bucket);
+                Assert.AreEqual(missingObject, response.Missing[0].Key);
+            }
+        }
+
+
 
         [TestMethod]
         public async Task PutGetObjectHierarchy()
@@ -1346,5 +1381,521 @@ namespace Horde.Storage.FunctionalTests.References
             Assert.AreEqual(key, oldRecord.Value.Item2);
             Assert.AreEqual("bucket", oldRecord.Value.Item1.ToString());
         }
+
+        
+        [TestMethod]
+        public async Task BatchJsonRequest()
+        {
+            // verifies that json request against the batch endpoint will fail
+            {
+                HttpContent requestContent = new StringContent("{}");
+                requestContent.Headers.ContentType = new MediaTypeHeaderValue(MediaTypeNames.Application.Json);
+
+                HttpResponseMessage result = await _httpClient!.PostAsync(requestUri: $"api/v1/refs/{TestNamespace}", requestContent);
+
+                Assert.AreEqual(HttpStatusCode.UnsupportedMediaType, result.StatusCode);
+            }
+        }
+
+        [TestMethod]
+        public async Task BatchErrorOperations()
+        {
+            // seed some data
+            BucketId bucket = new BucketId("bucket");
+            IoHashKey newBlobObjectKey = IoHashKey.FromName("thisObjectDoesNotExist");
+
+            IoHashKey putObjectKey = IoHashKey.FromName("putObjectFail");
+            CbObject ref1 = CbObject.Empty;
+
+            CbWriter getObjectOp = new CbWriter();
+            getObjectOp.BeginObject();
+            getObjectOp.WriteInteger( "opId",0);
+            getObjectOp.WriteString("op", BatchOps.BatchOp.Operation.GET.ToString());
+            getObjectOp.WriteString("bucket", bucket.ToString());
+            getObjectOp.WriteString("key", newBlobObjectKey.ToString());
+            getObjectOp.EndObject();
+
+            CbWriter putObjectOp = new CbWriter();
+            putObjectOp.BeginObject();
+            putObjectOp.WriteInteger("opId", 1);
+            putObjectOp.WriteString("op", BatchOps.BatchOp.Operation.PUT.ToString());
+            putObjectOp.WriteString("bucket", bucket.ToString());
+            putObjectOp.WriteString("key", putObjectKey.ToString());
+            putObjectOp.WriteObject("payload", ref1);
+            // omit the hash for a error object1.WriteHash("payloadHash", IoHash.Compute(ref1.GetView().Span));
+            putObjectOp.EndObject();
+
+            CbObject[] ops = new[]
+            {
+                getObjectOp.ToObject(),
+                putObjectOp.ToObject(),
+            };
+
+            CbWriter batchRequestWriter = new CbWriter();
+            batchRequestWriter.BeginObject();
+            batchRequestWriter.BeginUniformArray("ops", CbFieldType.Object);
+            foreach (CbObject o in ops)
+            {
+                batchRequestWriter.WriteObject(o);
+            }
+            batchRequestWriter.EndUniformArray();
+            batchRequestWriter.EndObject();
+            byte[] batchRequestData = batchRequestWriter.ToByteArray();
+
+            {
+                HttpContent requestContent = new ByteArrayContent(batchRequestData);
+                requestContent.Headers.ContentType = new MediaTypeHeaderValue(CustomMediaTypeNames.UnrealCompactBinary);
+
+                HttpResponseMessage result = await _httpClient!.PostAsync(requestUri: $"api/v1/refs/{TestNamespace}", requestContent);
+                result.EnsureSuccessStatusCode();
+
+                await using MemoryStream ms = new MemoryStream();
+                await result.Content.CopyToAsync(ms);
+                byte[] roundTrippedBuffer = ms.ToArray();
+
+                BatchOpsResponse response = CbSerializer.Deserialize<BatchOpsResponse>(roundTrippedBuffer);
+                Assert.AreEqual(2, response.Results.Count);
+                
+                BatchOpsResponse.OpResponses op0 = response.Results.First(r => r.OpId == 0);
+                Assert.IsNotNull(op0.Response);
+                Assert.AreEqual(404, op0.StatusCode);
+                Assert.IsTrue(!op0.Response["title"].Equals(CbField.Empty));
+                Assert.AreEqual("Object not found 29911b58b3c970ba39d9690f2dca66839dd6f5d9 in bucket bucket namespace test-namespace", op0.Response["title"].AsString());
+
+                BatchOpsResponse.OpResponses op1 = response.Results.First(r => r.OpId == 1);
+                Assert.IsNotNull(op1.Response);
+                Assert.AreEqual(500, op1.StatusCode);
+                Assert.IsTrue(!op1.Response["title"].Equals(CbField.Empty));
+                Assert.AreEqual("Missing payload for operation: 1", op1.Response["title"].AsString());
+            }
+        }
+
+
+        [TestMethod]
+        public async Task BatchGetOperations()
+        {
+            // seed some data
+            BucketId bucket = new BucketId("bucket");
+            IoHashKey newBlobObjectKey = IoHashKey.FromName("newBlobObject");
+            CbObject newBlobObject = CbObject.Build(writer => writer.WriteString("String", "this-has-contents"));
+
+            {
+                byte[] cbObjectBytes = newBlobObject.GetView().ToArray();
+                BlobIdentifier blobHash = BlobIdentifier.FromBlob(cbObjectBytes);
+
+                HttpContent requestContent = new ByteArrayContent(cbObjectBytes);
+                requestContent.Headers.ContentType = new MediaTypeHeaderValue(CustomMediaTypeNames.UnrealCompactBinary);
+                requestContent.Headers.Add(CommonHeaders.HashHeaderName, blobHash.ToString());
+
+                HttpResponseMessage result = await _httpClient!.PutAsync(requestUri: $"api/v1/refs/{TestNamespace}/{bucket}/{newBlobObjectKey}.uecb", requestContent);
+                result.EnsureSuccessStatusCode();
+                byte[] content = await result.Content.ReadAsByteArrayAsync();
+                PutObjectResponse? response = CbSerializer.Deserialize<PutObjectResponse?>(content);
+                Assert.IsNotNull(response);
+                Assert.IsTrue(response.Needs.Length == 0);
+            }
+
+            byte[] blobContents = Encoding.ASCII.GetBytes("This is a attached blob");
+            CbObject newReferenceObject = CbObject.Build(writer => writer.WriteBinaryAttachment("Attachment", IoHash.Compute(blobContents)));
+            IoHashKey newReferenceObjectKey = IoHashKey.FromName("newReferenceObject");
+
+            {
+                BlobIdentifier blobHash = BlobIdentifier.FromBlob(blobContents);
+
+                HttpContent requestContent = new ByteArrayContent(blobContents);
+                requestContent.Headers.ContentType = new MediaTypeHeaderValue(MediaTypeNames.Application.Octet);
+                requestContent.Headers.Add(CommonHeaders.HashHeaderName, blobHash.ToString());
+
+                HttpResponseMessage result = await _httpClient!.PutAsync(requestUri: $"api/v1/blobs/{TestNamespace}/{blobHash}", requestContent);
+                result.EnsureSuccessStatusCode();
+            }
+
+            {
+                byte[] cbObjectBytes = newReferenceObject.GetView().ToArray();
+                BlobIdentifier blobHash = BlobIdentifier.FromBlob(cbObjectBytes);
+
+                HttpContent requestContent = new ByteArrayContent(cbObjectBytes);
+                requestContent.Headers.ContentType = new MediaTypeHeaderValue(CustomMediaTypeNames.UnrealCompactBinary);
+                requestContent.Headers.Add(CommonHeaders.HashHeaderName, blobHash.ToString());
+
+                HttpResponseMessage result = await _httpClient!.PutAsync(requestUri: $"api/v1/refs/{TestNamespace}/{bucket}/{newReferenceObjectKey}.uecb", requestContent);
+                result.EnsureSuccessStatusCode();
+
+                result.EnsureSuccessStatusCode();
+                byte[] content = await result.Content.ReadAsByteArrayAsync();
+                PutObjectResponse response = CbSerializer.Deserialize<PutObjectResponse>(content);
+                Assert.IsTrue(response.Needs.Length == 0);
+            }
+            CbWriter getObjectOp = new CbWriter();
+            getObjectOp.BeginObject();
+            getObjectOp.WriteInteger( "opId",0);
+            getObjectOp.WriteString("op", BatchOps.BatchOp.Operation.GET.ToString());
+            getObjectOp.WriteString("bucket", bucket.ToString());
+            getObjectOp.WriteString("key", newBlobObjectKey.ToString());
+            getObjectOp.EndObject();
+
+            CbWriter getObjectOp2 = new CbWriter();
+            getObjectOp2.BeginObject();
+            getObjectOp2.WriteInteger("opId", 1);
+            getObjectOp2.WriteString("op", BatchOps.BatchOp.Operation.GET.ToString());
+            getObjectOp2.WriteString("bucket", bucket.ToString());
+            getObjectOp2.WriteString("key", newReferenceObjectKey.ToString());
+            getObjectOp2.EndObject();
+
+            CbObject[] ops = new[]
+            {
+                getObjectOp.ToObject(),
+                getObjectOp2.ToObject(),
+            };
+
+            CbWriter batchRequestWriter = new CbWriter();
+            batchRequestWriter.BeginObject();
+            batchRequestWriter.BeginUniformArray("ops", CbFieldType.Object);
+            foreach (CbObject o in ops)
+            {
+                batchRequestWriter.WriteObject(o);
+            }
+            batchRequestWriter.EndUniformArray();
+            batchRequestWriter.EndObject();
+            byte[] batchRequestData = batchRequestWriter.ToByteArray();
+
+            {
+                HttpContent requestContent = new ByteArrayContent(batchRequestData);
+                requestContent.Headers.ContentType = new MediaTypeHeaderValue(CustomMediaTypeNames.UnrealCompactBinary);
+
+                HttpResponseMessage result = await _httpClient!.PostAsync(requestUri: $"api/v1/refs/{TestNamespace}", requestContent);
+                result.EnsureSuccessStatusCode();
+
+                await using MemoryStream ms = new MemoryStream();
+                await result.Content.CopyToAsync(ms);
+                byte[] roundTrippedBuffer = ms.ToArray();
+
+                BatchOpsResponse response = CbSerializer.Deserialize<BatchOpsResponse>(roundTrippedBuffer);
+                Assert.AreEqual(2, response.Results.Count);
+                
+                BatchOpsResponse.OpResponses op0 = response.Results.First(r => r.OpId == 0);
+                Assert.IsNotNull(op0.Response);
+                Assert.AreEqual(200, op0.StatusCode);
+                Assert.AreEqual(newBlobObject, op0.Response);
+
+                BatchOpsResponse.OpResponses op1 = response.Results.First(r => r.OpId == 1);
+                Assert.IsNotNull(op1.Response);
+                Assert.AreEqual(200, op1.StatusCode);
+                Assert.AreEqual(newReferenceObject, op1.Response);
+            }
+        }
+
+        [TestMethod]
+        public async Task BatchHeadOperations()
+        {
+            // seed some data
+            BucketId bucket = new BucketId("bucket");
+            IoHashKey newBlobObjectKey = IoHashKey.FromName("newBlobObject");
+            CbObject newBlobObject = CbObject.Build(writer => writer.WriteString("String", "this-has-contents"));
+
+            {
+                byte[] cbObjectBytes = newBlobObject.GetView().ToArray();
+                BlobIdentifier blobHash = BlobIdentifier.FromBlob(cbObjectBytes);
+
+                HttpContent requestContent = new ByteArrayContent(cbObjectBytes);
+                requestContent.Headers.ContentType = new MediaTypeHeaderValue(CustomMediaTypeNames.UnrealCompactBinary);
+                requestContent.Headers.Add(CommonHeaders.HashHeaderName, blobHash.ToString());
+
+                HttpResponseMessage result = await _httpClient!.PutAsync(requestUri: $"api/v1/refs/{TestNamespace}/{bucket}/{newBlobObjectKey}.uecb", requestContent);
+                result.EnsureSuccessStatusCode();
+                byte[] content = await result.Content.ReadAsByteArrayAsync();
+                PutObjectResponse? response = CbSerializer.Deserialize<PutObjectResponse?>(content);
+                Assert.IsNotNull(response);
+                Assert.IsTrue(response.Needs.Length == 0);
+            }
+
+            byte[] blobContents = Encoding.ASCII.GetBytes("This is a attached blob");
+            CbObject newReferenceObject = CbObject.Build(writer => writer.WriteBinaryAttachment("Attachment", IoHash.Compute(blobContents)));
+            IoHashKey newReferenceObjectKey = IoHashKey.FromName("newReferenceObject");
+
+            {
+                BlobIdentifier blobHash = BlobIdentifier.FromBlob(blobContents);
+
+                HttpContent requestContent = new ByteArrayContent(blobContents);
+                requestContent.Headers.ContentType = new MediaTypeHeaderValue(MediaTypeNames.Application.Octet);
+                requestContent.Headers.Add(CommonHeaders.HashHeaderName, blobHash.ToString());
+
+                HttpResponseMessage result = await _httpClient!.PutAsync(requestUri: $"api/v1/blobs/{TestNamespace}/{blobHash}", requestContent);
+                result.EnsureSuccessStatusCode();
+            }
+
+            {
+                byte[] cbObjectBytes = newReferenceObject.GetView().ToArray();
+                BlobIdentifier blobHash = BlobIdentifier.FromBlob(cbObjectBytes);
+
+                HttpContent requestContent = new ByteArrayContent(cbObjectBytes);
+                requestContent.Headers.ContentType = new MediaTypeHeaderValue(CustomMediaTypeNames.UnrealCompactBinary);
+                requestContent.Headers.Add(CommonHeaders.HashHeaderName, blobHash.ToString());
+
+                HttpResponseMessage result = await _httpClient!.PutAsync(requestUri: $"api/v1/refs/{TestNamespace}/{bucket}/{newReferenceObjectKey}.uecb", requestContent);
+                result.EnsureSuccessStatusCode();
+
+                result.EnsureSuccessStatusCode();
+                byte[] content = await result.Content.ReadAsByteArrayAsync();
+                PutObjectResponse response = CbSerializer.Deserialize<PutObjectResponse>(content);
+                Assert.IsTrue(response.Needs.Length == 0);
+            }
+            CbWriter getObjectOp = new CbWriter();
+            getObjectOp.BeginObject();
+            getObjectOp.WriteInteger( "opId",0);
+            getObjectOp.WriteString("op", BatchOps.BatchOp.Operation.HEAD.ToString());
+            getObjectOp.WriteString("bucket", bucket.ToString());
+            getObjectOp.WriteString("key", newBlobObjectKey.ToString());
+            getObjectOp.EndObject();
+
+            CbWriter getObjectOp2 = new CbWriter();
+            getObjectOp2.BeginObject();
+            getObjectOp2.WriteInteger("opId", 1);
+            getObjectOp2.WriteString("op", BatchOps.BatchOp.Operation.HEAD.ToString());
+            getObjectOp2.WriteString("bucket", bucket.ToString());
+            getObjectOp2.WriteString("key", newReferenceObjectKey.ToString());
+            getObjectOp2.EndObject();
+
+            CbObject[] ops = new[]
+            {
+                getObjectOp.ToObject(),
+                getObjectOp2.ToObject(),
+            };
+
+            CbWriter batchRequestWriter = new CbWriter();
+            batchRequestWriter.BeginObject();
+            batchRequestWriter.BeginUniformArray("ops", CbFieldType.Object);
+            foreach (CbObject o in ops)
+            {
+                batchRequestWriter.WriteObject(o);
+            }
+            batchRequestWriter.EndUniformArray();
+            batchRequestWriter.EndObject();
+            byte[] batchRequestData = batchRequestWriter.ToByteArray();
+            
+            {
+                HttpContent requestContent = new ByteArrayContent(batchRequestData);
+                requestContent.Headers.ContentType = new MediaTypeHeaderValue(CustomMediaTypeNames.UnrealCompactBinary);
+
+                HttpResponseMessage result = await _httpClient!.PostAsync(requestUri: $"api/v1/refs/{TestNamespace}", requestContent);
+                result.EnsureSuccessStatusCode();
+
+                await using MemoryStream ms = new MemoryStream();
+                await result.Content.CopyToAsync(ms);
+                byte[] roundTrippedBuffer = ms.ToArray();
+
+                BatchOpsResponse response = CbSerializer.Deserialize<BatchOpsResponse>(roundTrippedBuffer);
+                Assert.AreEqual(2, response.Results.Count);
+                
+                BatchOpsResponse.OpResponses op0 = response.Results.First(r => r.OpId == 0);
+                Assert.IsNotNull(op0.Response);
+                Assert.AreEqual(200, op0.StatusCode);
+                Assert.IsTrue(!op0.Response["exists"].Equals(CbField.Empty));
+                Assert.IsTrue(op0.Response["exists"].AsBool());
+
+                BatchOpsResponse.OpResponses op1 = response.Results.First(r => r.OpId == 1);
+                Assert.IsNotNull(op1.Response);
+                Assert.AreEqual(200, op1.StatusCode);
+                Assert.IsTrue(!op1.Response["exists"].Equals(CbField.Empty));
+                Assert.IsTrue(op1.Response["exists"].AsBool());
+            }
+        }
+
+        [TestMethod]
+        public async Task BatchPutOperations()
+        {
+            BucketId bucket = new BucketId("bucket");
+            IoHashKey ref0name = IoHashKey.FromName("putRef0");
+            CbObject ref0 = CbObject.Build(writer => writer.WriteString("foo", "bar"));
+
+            IoHashKey ref1name = IoHashKey.FromName("putRef1");
+            CbObject ref1 = CbObject.Build(writer => writer.WriteInteger("baz", 1337));
+
+            CbWriter object0 = new CbWriter();
+            object0.BeginObject();
+            object0.WriteInteger( "opId",0);
+            object0.WriteString("op", BatchOps.BatchOp.Operation.PUT.ToString());
+            object0.WriteString("bucket", bucket.ToString());
+            object0.WriteString("key", ref0name.ToString());
+            object0.WriteObject("payload", ref0);
+            object0.WriteHash("payloadHash", IoHash.Compute(ref0.GetView().Span));
+            object0.EndObject();
+
+            CbWriter object1 = new CbWriter();
+            object1.BeginObject();
+            object1.WriteInteger("opId", 1);
+            object1.WriteString("op", BatchOps.BatchOp.Operation.PUT.ToString());
+            object1.WriteString("bucket", bucket.ToString());
+            object1.WriteString("key", ref1name.ToString());
+            object1.WriteObject("payload", ref1);
+            object1.WriteHash("payloadHash", IoHash.Compute(ref1.GetView().Span));
+            object1.EndObject();
+
+            CbObject[] ops = new[]
+            {
+                object0.ToObject(),
+                object1.ToObject(),
+            };
+
+            CbWriter batchRequestWriter = new CbWriter();
+            batchRequestWriter.BeginObject();
+            batchRequestWriter.BeginUniformArray("ops", CbFieldType.Object);
+            foreach (CbObject o in ops)
+            {
+                batchRequestWriter.WriteObject(o);
+            }
+            batchRequestWriter.EndUniformArray();
+            batchRequestWriter.EndObject();
+            byte[] batchRequestData = batchRequestWriter.ToByteArray();
+
+            {
+                HttpContent requestContent = new ByteArrayContent(batchRequestData);
+                requestContent.Headers.ContentType = new MediaTypeHeaderValue(CustomMediaTypeNames.UnrealCompactBinary);
+
+                HttpResponseMessage result = await _httpClient!.PostAsync(requestUri: $"api/v1/refs/{TestNamespace}", requestContent);
+                result.EnsureSuccessStatusCode();
+
+                await using MemoryStream ms = new MemoryStream();
+                await result.Content.CopyToAsync(ms);
+                byte[] roundTrippedBuffer = ms.ToArray();
+
+                BatchOpsResponse response = CbSerializer.Deserialize<BatchOpsResponse>(roundTrippedBuffer);
+                Assert.AreEqual(2, response.Results.Count);
+                
+                BatchOpsResponse.OpResponses op0 = response.Results.First(r => r.OpId == 0);
+                Assert.IsNotNull(op0.Response);
+                Assert.AreEqual(200, op0.StatusCode);
+                Assert.IsTrue(!op0.Response["needs"].Equals(CbField.Empty));
+                CollectionAssert.AreEqual(Array.Empty<IoHash>(), op0.Response["needs"].ToArray());
+
+                BatchOpsResponse.OpResponses op1 = response.Results.First(r => r.OpId == 1);
+                Assert.IsNotNull(op1.Response);
+                Assert.AreEqual(200, op1.StatusCode);
+                Assert.IsTrue(!op1.Response["needs"].Equals(CbField.Empty));
+                CollectionAssert.AreEqual(Array.Empty<IoHash>(), op1.Response["needs"].ToArray());
+            }
+        }
+
+        [TestMethod]
+        public async Task BatchMixedOperations()
+        {
+            // seed some data
+            BucketId bucket = new BucketId("bucket");
+            IoHashKey getObjectKey = IoHashKey.FromName("getBlobObject");
+            CbObject getObject = CbObject.Build(writer => writer.WriteString("String", "this-has-contents"));
+
+            IoHashKey putObjectKey = IoHashKey.FromName("putBlobObject");
+            CbObject putObject = CbObject.Build(writer => writer.WriteString("String", "this-has-contents"));
+
+            IoHashKey missingObjectKey = IoHashKey.FromName("thisKeyDoesNotExist");
+
+            {
+                byte[] cbObjectBytes = getObject.GetView().ToArray();
+                BlobIdentifier blobHash = BlobIdentifier.FromBlob(cbObjectBytes);
+
+                HttpContent requestContent = new ByteArrayContent(cbObjectBytes);
+                requestContent.Headers.ContentType = new MediaTypeHeaderValue(CustomMediaTypeNames.UnrealCompactBinary);
+                requestContent.Headers.Add(CommonHeaders.HashHeaderName, blobHash.ToString());
+
+                HttpResponseMessage result = await _httpClient!.PutAsync(requestUri: $"api/v1/refs/{TestNamespace}/{bucket}/{getObjectKey}.uecb", requestContent);
+                result.EnsureSuccessStatusCode();
+                byte[] content = await result.Content.ReadAsByteArrayAsync();
+                PutObjectResponse? response = CbSerializer.Deserialize<PutObjectResponse?>(content);
+                Assert.IsNotNull(response);
+                Assert.IsTrue(response.Needs.Length == 0);
+            }
+
+            CbWriter getObjectOp = new CbWriter();
+            getObjectOp.BeginObject();
+            getObjectOp.WriteInteger( "opId",0);
+            getObjectOp.WriteString("op", BatchOps.BatchOp.Operation.GET.ToString());
+            getObjectOp.WriteString("bucket", bucket.ToString());
+            getObjectOp.WriteString("key", getObjectKey.ToString());
+            getObjectOp.EndObject();
+
+            CbWriter putObjectOp = new CbWriter();
+            putObjectOp.BeginObject();
+            putObjectOp.WriteInteger("opId", 1);
+            putObjectOp.WriteString("op", BatchOps.BatchOp.Operation.PUT.ToString());
+            putObjectOp.WriteString("bucket", bucket.ToString());
+            putObjectOp.WriteString("key", putObjectKey.ToString());
+            putObjectOp.WriteObject("payload", putObject);
+            putObjectOp.WriteHash("payloadHash", IoHash.Compute(putObject.GetView().Span));
+            putObjectOp.EndObject();
+
+            CbWriter errorObjectOp = new CbWriter();
+            errorObjectOp.BeginObject();
+            errorObjectOp.WriteInteger( "opId",2);
+            errorObjectOp.WriteString("op", BatchOps.BatchOp.Operation.GET.ToString());
+            errorObjectOp.WriteString("bucket", bucket.ToString());
+            errorObjectOp.WriteString("key", missingObjectKey.ToString());
+            errorObjectOp.EndObject();
+
+            CbWriter headObjectOp = new CbWriter();
+            headObjectOp.BeginObject();
+            headObjectOp.WriteInteger("opId", 3);
+            headObjectOp.WriteString("op", BatchOps.BatchOp.Operation.HEAD.ToString());
+            headObjectOp.WriteString("bucket", bucket.ToString());
+            headObjectOp.WriteString("key", getObjectKey.ToString());
+            headObjectOp.EndObject();
+
+            CbObject[] ops = new[]
+            {
+                getObjectOp.ToObject(),
+                putObjectOp.ToObject(),
+                errorObjectOp.ToObject(),
+                headObjectOp.ToObject(),
+            };
+
+            CbWriter batchRequestWriter = new CbWriter();
+            batchRequestWriter.BeginObject();
+            batchRequestWriter.BeginUniformArray("ops", CbFieldType.Object);
+            foreach (CbObject o in ops)
+            {
+                batchRequestWriter.WriteObject(o);
+            }
+            batchRequestWriter.EndUniformArray();
+            batchRequestWriter.EndObject();
+            byte[] batchRequestData = batchRequestWriter.ToByteArray();
+            
+            {
+                HttpContent requestContent = new ByteArrayContent(batchRequestData);
+                requestContent.Headers.ContentType = new MediaTypeHeaderValue(CustomMediaTypeNames.UnrealCompactBinary);
+
+                HttpResponseMessage result = await _httpClient!.PostAsync(requestUri: $"api/v1/refs/{TestNamespace}", requestContent);
+                result.EnsureSuccessStatusCode();
+
+                await using MemoryStream ms = new MemoryStream();
+                await result.Content.CopyToAsync(ms);
+                byte[] roundTrippedBuffer = ms.ToArray();
+
+                BatchOpsResponse response = CbSerializer.Deserialize<BatchOpsResponse>(roundTrippedBuffer);
+                Assert.AreEqual(4, response.Results.Count);
+                
+                BatchOpsResponse.OpResponses getOp = response.Results.First(r => r.OpId == 0);
+                Assert.IsNotNull(getOp.Response);
+                Assert.AreEqual(200, getOp.StatusCode);
+                Assert.AreEqual(getObject, getOp.Response);
+
+                BatchOpsResponse.OpResponses putOp = response.Results.First(r => r.OpId == 1);
+                Assert.IsNotNull(putOp.Response);
+                Assert.AreEqual(200, putOp.StatusCode);
+                Assert.IsTrue(!putOp.Response["needs"].Equals(CbField.Empty));
+                CollectionAssert.AreEqual(Array.Empty<IoHash>(), putOp.Response["needs"].ToArray());
+
+                BatchOpsResponse.OpResponses errorOp = response.Results.First(r => r.OpId == 2);
+                Assert.IsNotNull(errorOp.Response);
+                Assert.AreEqual(404, errorOp.StatusCode);
+                Assert.IsTrue(!errorOp.Response["title"].Equals(CbField.Empty));
+                Assert.AreEqual("Object not found cbc3db15b9c8253f6106158962325c8fd848daef in bucket bucket namespace test-namespace", errorOp.Response["title"].AsString());
+                Assert.AreEqual(404, errorOp.Response["status"].AsInt32());
+
+                BatchOpsResponse.OpResponses headOp = response.Results.First(r => r.OpId == 3);
+                Assert.IsNotNull(headOp.Response);
+                Assert.AreEqual(200, headOp.StatusCode);
+                Assert.IsTrue(!headOp.Response["exists"].Equals(CbField.Empty));
+                Assert.IsTrue(headOp.Response["exists"].AsBool());
+            }
+        }
     }
 }
diff --git a/Engine/Source/Programs/Horde/HordeStorage/Jupiter.Common/Implementation/IoHashKey.cs b/Engine/Source/Programs/Horde/HordeStorage/Jupiter.Common/Implementation/IoHashKey.cs
--- a/Engine/Source/Programs/Horde/HordeStorage/Jupiter.Common/Implementation/IoHashKey.cs
+++ b/Engine/Source/Programs/Horde/HordeStorage/Jupiter.Common/Implementation/IoHashKey.cs
@@ -4,12 +4,16 @@ using System;
 using System.ComponentModel;
 using System.Globalization;
 using System.Text;
+using EpicGames.Core;
+using EpicGames.Serialization;
 using Newtonsoft.Json;
+using JsonWriter = Newtonsoft.Json.JsonWriter;
 
 namespace Jupiter.Implementation
 {
     [JsonConverter(typeof(IoHashKeyJsonConverter))]
     [TypeConverter(typeof(IoHashKeyTypeConverter))]
+    [CbConverter(typeof(IoHashKeyCbConverter))]
 
     public readonly struct IoHashKey: IEquatable<IoHashKey>
     {
@@ -117,5 +121,16 @@ namespace Jupiter.Implementation
             return base.ConvertFrom(context, culture, value);  
         }
     }
+
+    sealed class IoHashKeyCbConverter : CbConverterBase<IoHashKey>
+    {
+        public override IoHashKey Read(CbField field) => new IoHashKey(field.AsString());
+
+        /// <inheritdoc/>
+        public override void Write(CbWriter writer, IoHashKey value) => writer.WriteStringValue(value.ToString());
+
+        /// <inheritdoc/>
+        public override void WriteNamed(CbWriter writer, Utf8String name, IoHashKey value) => writer.WriteString(name, value.ToString());
+    }
 }
 
diff --git a/Engine/Source/Programs/Shared/EpicGames.Horde/Storage/BucketId.cs b/Engine/Source/Programs/Shared/EpicGames.Horde/Storage/BucketId.cs
--- a/Engine/Source/Programs/Shared/EpicGames.Horde/Storage/BucketId.cs
+++ b/Engine/Source/Programs/Shared/EpicGames.Horde/Storage/BucketId.cs
@@ -8,6 +8,8 @@ using System.Globalization;
 using System.Text;
 using System.Text.Json;
 using System.Text.Json.Serialization;
+using EpicGames.Core;
+using EpicGames.Serialization;
 
 namespace EpicGames.Horde.Storage
 {
@@ -16,6 +18,7 @@ namespace EpicGames.Horde.Storage
 	/// </summary>
 	[JsonConverter(typeof(BucketIdJsonConverter))]
 	[TypeConverter(typeof(BucketIdTypeConverter))]
+	[CbConverter(typeof(BucketIdCbConverter))]
 	public struct BucketId : IEquatable<BucketId>
 	{
 		/// <summary>
@@ -74,4 +77,15 @@ namespace EpicGames.Horde.Storage
 		/// <inheritdoc/>
 		public override object ConvertFrom(ITypeDescriptorContext Context, CultureInfo Culture, object Value) => new BucketId((string)Value);
 	}
+
+	sealed class BucketIdCbConverter : CbConverterBase<BucketId>
+	{
+		public override BucketId Read(CbField field) => new BucketId(field.AsString());
+
+		/// <inheritdoc/>
+		public override void Write(CbWriter writer, BucketId value) => writer.WriteStringValue(value.ToString());
+
+		/// <inheritdoc/>
+		public override void WriteNamed(CbWriter writer, Utf8String name, BucketId value) => writer.WriteString(name, value.ToString());
+	}
 }
