Horde: Log when pools of EC2 instances cannot be expanded at the requested count

diff --git a/Engine/Source/Programs/Horde/Horde.Build/Fleet/Impl/Managers/Fleet/AwsFleetManager.cs b/Engine/Source/Programs/Horde/Horde.Build/Fleet/Impl/Managers/Fleet/AwsFleetManager.cs
--- a/Engine/Source/Programs/Horde/Horde.Build/Fleet/Impl/Managers/Fleet/AwsFleetManager.cs
+++ b/Engine/Source/Programs/Horde/Horde.Build/Fleet/Impl/Managers/Fleet/AwsFleetManager.cs
@@ -91,6 +91,12 @@ namespace HordeServer.Services.Impl
 						}
 					}
 				}
+
+				if (StartRequest.InstanceIds.Count < Count)
+				{
+					Logger.LogInformation("Unable to expand pool with the requested number of instances. " +
+					                      "Num requested instances to add {RequestedCount}. Actual instances started {InstancesStarted}", Count, StartRequest.InstanceIds.Count);
+				}
 			}
 		}
 
