Horde Storage - Compressed blobs endpoints will now always resolve references to content ids (e.g. compressed blobs) never falling back to the uncompressed (without compressed buffer header) asset.

diff --git a/Engine/Source/Programs/Horde/HordeStorage/Horde.Storage/Controllers/CompressedBlobController.cs b/Engine/Source/Programs/Horde/HordeStorage/Horde.Storage/Controllers/CompressedBlobController.cs
--- a/Engine/Source/Programs/Horde/HordeStorage/Horde.Storage/Controllers/CompressedBlobController.cs
+++ b/Engine/Source/Programs/Horde/HordeStorage/Horde.Storage/Controllers/CompressedBlobController.cs
@@ -93,7 +93,7 @@ namespace Horde.Storage.Controllers
             {
                 return Forbid();
             }
-            BlobIdentifier[]? chunks = await _contentIdStore.Resolve(ns, id);
+            BlobIdentifier[]? chunks = await _contentIdStore.Resolve(ns, id, mustBeContentId: true);
             if (chunks == null || chunks.Length == 0)
             {
                 return NotFound();
@@ -136,7 +136,7 @@ namespace Horde.Storage.Controllers
 
             IEnumerable<Task> tasks = id.Select(async blob =>
             {
-                BlobIdentifier[]? chunks = await _contentIdStore.Resolve(ns, blob);
+                BlobIdentifier[]? chunks = await _contentIdStore.Resolve(ns, blob, mustBeContentId: true);
 
                 if (chunks == null)
                 {
@@ -163,7 +163,7 @@ namespace Horde.Storage.Controllers
 
         private async Task<BlobContents> GetImpl(NamespaceId ns, ContentId contentId)
         {
-            BlobIdentifier[]? chunks = await _contentIdStore.Resolve(ns, contentId);
+            BlobIdentifier[]? chunks = await _contentIdStore.Resolve(ns, contentId, mustBeContentId: true);
             if (chunks == null || chunks.Length == 0)
             {
                 throw new ContentIdResolveException(contentId);
diff --git a/Engine/Source/Programs/Horde/HordeStorage/Horde.Storage/Controllers/ContentIdController.cs b/Engine/Source/Programs/Horde/HordeStorage/Horde.Storage/Controllers/ContentIdController.cs
--- a/Engine/Source/Programs/Horde/HordeStorage/Horde.Storage/Controllers/ContentIdController.cs
+++ b/Engine/Source/Programs/Horde/HordeStorage/Horde.Storage/Controllers/ContentIdController.cs
@@ -50,7 +50,7 @@ namespace Horde.Storage.Controllers
                     return Forbid();
                 }
             }
-            BlobIdentifier[]? blobs = await _contentIdStore.Resolve(ns, contentId);
+            BlobIdentifier[]? blobs = await _contentIdStore.Resolve(ns, contentId, mustBeContentId: true);
 
             if (blobs == null)
             {
diff --git a/Engine/Source/Programs/Horde/HordeStorage/Horde.Storage/Implementation/ContentId/CacheContentIdStore.cs b/Engine/Source/Programs/Horde/HordeStorage/Horde.Storage/Implementation/ContentId/CacheContentIdStore.cs
--- a/Engine/Source/Programs/Horde/HordeStorage/Horde.Storage/Implementation/ContentId/CacheContentIdStore.cs
+++ b/Engine/Source/Programs/Horde/HordeStorage/Horde.Storage/Implementation/ContentId/CacheContentIdStore.cs
@@ -17,7 +17,7 @@ namespace Horde.Storage.Implementation
         {
         }
 
-        public async Task<BlobIdentifier[]?> Resolve(NamespaceId ns, ContentId contentId)
+        public async Task<BlobIdentifier[]?> Resolve(NamespaceId ns, ContentId contentId, bool mustBeContentId)
         {
             HttpRequestMessage getContentIdRequest = BuildHttpRequest(HttpMethod.Get, $"api/v1/content-id/{ns}/{contentId}");
             HttpResponseMessage response = await HttpClient.SendAsync(getContentIdRequest);
diff --git a/Engine/Source/Programs/Horde/HordeStorage/Horde.Storage/Implementation/ContentId/IContentIdStore.cs b/Engine/Source/Programs/Horde/HordeStorage/Horde.Storage/Implementation/ContentId/IContentIdStore.cs
--- a/Engine/Source/Programs/Horde/HordeStorage/Horde.Storage/Implementation/ContentId/IContentIdStore.cs
+++ b/Engine/Source/Programs/Horde/HordeStorage/Horde.Storage/Implementation/ContentId/IContentIdStore.cs
@@ -1,9 +1,9 @@
 // Copyright Epic Games, Inc. All Rights Reserved.
 
-using System;
-using System.Threading.Tasks;
 using EpicGames.Horde.Storage;
 using Jupiter.Implementation;
+using System;
+using System.Threading.Tasks;
 
 namespace Horde.Storage.Implementation
 {
@@ -14,8 +14,9 @@ namespace Horde.Storage.Implementation
         /// </summary>
         /// <param name="ns">The namespace to operate in</param>
         /// <param name="contentId">The identifier for the content id</param>
+        /// <param name="mustBeContentId"></param>
         /// <returns></returns>
-        Task<BlobIdentifier[]?> Resolve(NamespaceId ns, ContentId contentId);
+        Task<BlobIdentifier[]?> Resolve(NamespaceId ns, ContentId contentId, bool mustBeContentId = false);
 
         /// <summary>
         /// Add a mapping from contentId to blobIdentifier
@@ -32,7 +33,7 @@ namespace Horde.Storage.Implementation
     {
         public InvalidContentIdException(ContentId contentId) : base($"Unknown content id {contentId}")
         {
-            
+
         }
     }
 
@@ -40,7 +41,7 @@ namespace Horde.Storage.Implementation
     {
         public ContentId ContentId { get; }
 
-        public ContentIdResolveException(ContentId contentId): base($"Unable to find any mapping of contentId {contentId} that has all blobs present")
+        public ContentIdResolveException(ContentId contentId) : base($"Unable to find any mapping of contentId {contentId} that has all blobs present")
         {
             ContentId = contentId;
         }
diff --git a/Engine/Source/Programs/Horde/HordeStorage/Horde.Storage/Implementation/ContentId/MemoryContentIdStore.cs b/Engine/Source/Programs/Horde/HordeStorage/Horde.Storage/Implementation/ContentId/MemoryContentIdStore.cs
--- a/Engine/Source/Programs/Horde/HordeStorage/Horde.Storage/Implementation/ContentId/MemoryContentIdStore.cs
+++ b/Engine/Source/Programs/Horde/HordeStorage/Horde.Storage/Implementation/ContentId/MemoryContentIdStore.cs
@@ -18,7 +18,7 @@ namespace Horde.Storage.Implementation
             _blobStore = blobStore;
         }
 
-        public async Task<BlobIdentifier[]?> Resolve(NamespaceId ns, ContentId contentId)
+        public async Task<BlobIdentifier[]?> Resolve(NamespaceId ns, ContentId contentId, bool mustBeContentId)
         {
             if (_contentIds.TryGetValue(ns, out ConcurrentDictionary<ContentId, SortedList<int, BlobIdentifier[]>>? contentIdsForNamespace))
             {
@@ -36,7 +36,7 @@ namespace Horde.Storage.Implementation
 
             BlobIdentifier uncompressedBlobIdentifier = contentId.AsBlobIdentifier();
             // if no content id is found, but we have a blob that matches the content id (so a unchunked and uncompressed version of the data) we use that instead
-            if (await _blobStore.Exists(ns, uncompressedBlobIdentifier))
+            if (!mustBeContentId && await _blobStore.Exists(ns, uncompressedBlobIdentifier))
                 return new[] { uncompressedBlobIdentifier };
 
             return null;
diff --git a/Engine/Source/Programs/Horde/HordeStorage/Horde.Storage/Implementation/ContentId/MongoContentIdStore.cs b/Engine/Source/Programs/Horde/HordeStorage/Horde.Storage/Implementation/ContentId/MongoContentIdStore.cs
--- a/Engine/Source/Programs/Horde/HordeStorage/Horde.Storage/Implementation/ContentId/MongoContentIdStore.cs
+++ b/Engine/Source/Programs/Horde/HordeStorage/Horde.Storage/Implementation/ContentId/MongoContentIdStore.cs
@@ -36,7 +36,7 @@ namespace Horde.Storage.Implementation
             AddIndexFor<MongoContentIdModelV0>().CreateOne(indexModel);
         }
 
-        public async Task<BlobIdentifier[]?> Resolve(NamespaceId ns, ContentId contentId)
+        public async Task<BlobIdentifier[]?> Resolve(NamespaceId ns, ContentId contentId, bool mustBeContentId)
         {
             IMongoCollection<MongoContentIdModelV0> collection = GetCollection<MongoContentIdModelV0>();
 
@@ -62,7 +62,7 @@ namespace Horde.Storage.Implementation
 
             BlobIdentifier contentIdAsBlobIdentifier = contentId.AsBlobIdentifier();
             // no content id where all blobs are present, check if its present in the blob store as a uncompressed version of the blob
-            if (await _blobStore.Exists(ns, contentIdAsBlobIdentifier))
+            if (!mustBeContentId && await _blobStore.Exists(ns, contentIdAsBlobIdentifier))
             {
                 return new[] { contentIdAsBlobIdentifier };
             }
diff --git a/Engine/Source/Programs/Horde/HordeStorage/Horde.Storage/Implementation/ContentId/ScyllaContentIdStore.cs b/Engine/Source/Programs/Horde/HordeStorage/Horde.Storage/Implementation/ContentId/ScyllaContentIdStore.cs
--- a/Engine/Source/Programs/Horde/HordeStorage/Horde.Storage/Implementation/ContentId/ScyllaContentIdStore.cs
+++ b/Engine/Source/Programs/Horde/HordeStorage/Horde.Storage/Implementation/ContentId/ScyllaContentIdStore.cs
@@ -35,13 +35,17 @@ namespace Horde.Storage.Implementation
         }
 
 
-        public async Task<BlobIdentifier[]?> Resolve(NamespaceId ns, ContentId contentId)
+        public async Task<BlobIdentifier[]?> Resolve(NamespaceId ns, ContentId contentId, bool mustBeContentId)
         {
             using IScope scope = Tracer.Instance.StartActive("ScyllaContentIdStore.ResolveContentId");
             scope.Span.ResourceName = contentId.ToString();
 
             BlobIdentifier contentIdBlob = contentId.AsBlobIdentifier();
-            Task<bool> blobStoreExistsTask = _blobStore.Exists(ns, contentIdBlob);
+            Task<bool>? blobStoreExistsTask = null;
+            if (!mustBeContentId)
+            {
+                blobStoreExistsTask = _blobStore.Exists(ns, contentIdBlob);
+            }
 
             // lower content_weight means its a better candidate to resolve to
             foreach (ScyllaContentId? resolvedContentId in await _mapper.FetchAsync<ScyllaContentId>("WHERE content_id = ? ORDER BY content_weight DESC", new ScyllaBlobIdentifier(contentId)))
@@ -61,11 +65,14 @@ namespace Horde.Storage.Implementation
                 // blobs are missing continue testing with the next content id in the weighted list as that might exist
             }
 
-            // if no content id is found, but we have a blob that matches the content id (so a unchunked and uncompressed version of the data) we use that instead
-            bool contentIdBlobExists = await blobStoreExistsTask;
+            if (!mustBeContentId)
+            {
+                // if no content id is found, but we have a blob that matches the content id (so a unchunked and uncompressed version of the data) we use that instead
+                bool contentIdBlobExists = await blobStoreExistsTask!;
 
-            if (contentIdBlobExists)
-                return new[] { contentIdBlob };
+                if (contentIdBlobExists)
+                    return new[] { contentIdBlob };
+            }
             
             // unable to resolve the content id
             return null;
