Horde Storage - No longer returning all the blobs removed by the cleanup. This was never really used and meant keeping a lot of hashes in memory for no reason.

diff --git a/Engine/Source/Programs/Horde/HordeStorage/Horde.Storage/Implementation/Blob/FileSystemStore.cs b/Engine/Source/Programs/Horde/HordeStorage/Horde.Storage/Implementation/Blob/FileSystemStore.cs
--- a/Engine/Source/Programs/Horde/HordeStorage/Horde.Storage/Implementation/Blob/FileSystemStore.cs
+++ b/Engine/Source/Programs/Horde/HordeStorage/Horde.Storage/Implementation/Blob/FileSystemStore.cs
@@ -199,7 +199,7 @@ namespace Horde.Storage.Implementation
             return true;
         }
 
-        public async Task<List<BlobIdentifier>> Cleanup(CancellationToken cancellationToken)
+        public async Task<ulong> Cleanup(CancellationToken cancellationToken)
         {
             return await CleanupInternal(cancellationToken);
         }
@@ -212,17 +212,17 @@ namespace Horde.Storage.Implementation
         /// <param name="cancellationToken">Cancellation token</param>
         /// <param name="batchSize">Number of files to scan for clean up. A higher number is recommended since blob store can contain many small blobs</param>
         /// <returns></returns>
-        public async Task<List<BlobIdentifier>> CleanupInternal(CancellationToken cancellationToken, int batchSize = 100000)
+        public async Task<ulong> CleanupInternal(CancellationToken cancellationToken, int batchSize = 100000)
         {
             long size = await CalculateDiskSpaceUsed();
             ulong maxSizeBytes = _settings.CurrentValue.MaxSizeBytes;
             long triggerSize = (long) (maxSizeBytes * _settings.CurrentValue.TriggerThresholdPercentage);
             long targetSize = (long) (maxSizeBytes * _settings.CurrentValue.TargetThresholdPercentage); // Target to shrink to if triggered
-            List<BlobIdentifier> blobsRemoved = new List<BlobIdentifier>(batchSize);
+            ulong countOfBlobsRemoved = 0;
 
             if (size < triggerSize)
             {
-                return blobsRemoved;
+                return 0;
             }
 
             // Perform a maximum of 5 clean up runs
@@ -231,7 +231,7 @@ namespace Horde.Storage.Implementation
                 size = await CalculateDiskSpaceUsed();
                 if (size <= targetSize)
                 {
-                    return blobsRemoved;
+                    return countOfBlobsRemoved;
                 }
                 
                 IEnumerable<FileInfo> fileInfos = GetLeastRecentlyAccessedObjects(maxResults: batchSize);
@@ -245,12 +245,12 @@ namespace Horde.Storage.Implementation
                     {
                         totalBytesDeleted += fi.Length;
                         fi.Delete();
-                        blobsRemoved.Add(new BlobIdentifier(fi.Name));
+                        ++countOfBlobsRemoved;
 
                         long currentSize = size - totalBytesDeleted;
                         if (currentSize <= targetSize || cancellationToken.IsCancellationRequested)
                         {
-                            return blobsRemoved;
+                            return countOfBlobsRemoved;
                         }
                     }
                     catch (FileNotFoundException)
@@ -261,11 +261,11 @@ namespace Horde.Storage.Implementation
 
                 if (!hadFiles)
                 {
-                    return blobsRemoved;
+                    return countOfBlobsRemoved;
                 }
             }
             
-            return blobsRemoved;
+            return countOfBlobsRemoved;
         }
 
         /// <summary>
diff --git a/Engine/Source/Programs/Horde/HordeStorage/Horde.Storage/Implementation/GC/BlobCleanupService.cs b/Engine/Source/Programs/Horde/HordeStorage/Horde.Storage/Implementation/GC/BlobCleanupService.cs
--- a/Engine/Source/Programs/Horde/HordeStorage/Horde.Storage/Implementation/GC/BlobCleanupService.cs
+++ b/Engine/Source/Programs/Horde/HordeStorage/Horde.Storage/Implementation/GC/BlobCleanupService.cs
@@ -77,7 +77,7 @@ namespace Horde.Storage.Implementation
             }
         }
 
-        public async Task<List<BlobIdentifier>> Cleanup(BlobCleanupState state, CancellationToken cancellationToken)
+        public async Task Cleanup(BlobCleanupState state, CancellationToken cancellationToken)
         {
             List<BlobIdentifier> removedBlobs = new List<BlobIdentifier>();
             
@@ -94,17 +94,14 @@ namespace Horde.Storage.Implementation
                 _logger.Information("Attempting to run Blob Cleanup {BlobCleanup}. ", type);
                 try
                 {
-                    List<BlobIdentifier> tempBlobs = await blobCleanup.Cleanup(cancellationToken);
-                    _logger.Information("Ran blob cleanup {BlobCleanup}. Deleted {CountBlobRecords}", type, tempBlobs.Count);
-                    removedBlobs.AddRange(tempBlobs);
+                    ulong countOfBlobsCleaned = await blobCleanup.Cleanup(cancellationToken);
+                    _logger.Information("Ran blob cleanup {BlobCleanup}. Deleted {CountBlobRecords}", type, countOfBlobsCleaned);
                 }
                 catch (Exception e)
                 {
                     _logger.Error(e, "Exception running Blob Cleanup {BlobCleanup} .", type);
                 }
             }
-
-            return removedBlobs;
         }
     }
 }
diff --git a/Engine/Source/Programs/Horde/HordeStorage/Horde.Storage/Implementation/GC/OrphanBlobCleanup.cs b/Engine/Source/Programs/Horde/HordeStorage/Horde.Storage/Implementation/GC/OrphanBlobCleanup.cs
--- a/Engine/Source/Programs/Horde/HordeStorage/Horde.Storage/Implementation/GC/OrphanBlobCleanup.cs
+++ b/Engine/Source/Programs/Horde/HordeStorage/Horde.Storage/Implementation/GC/OrphanBlobCleanup.cs
@@ -21,7 +21,7 @@ namespace Horde.Storage.Implementation
     public interface IBlobCleanup
     {
         bool ShouldRun();
-        Task<List<BlobIdentifier>> Cleanup(CancellationToken none);
+        Task<ulong> Cleanup(CancellationToken none);
     }
 
     public class OrphanBlobCleanup : IBlobCleanup
@@ -78,12 +78,12 @@ namespace Horde.Storage.Implementation
             return true;
         }
 
-        public async Task<List<BlobIdentifier>> Cleanup(CancellationToken cancellationToken)
+        public async Task<ulong> Cleanup(CancellationToken cancellationToken)
         {
             if (!_leaderElection.IsThisInstanceLeader())
             {
                 _logger.Information("Skipped orphan blob cleanup run as this instance is not the leader");
-                return new List<BlobIdentifier>();
+                return 0;
             }
 
             List<NamespaceId> namespaces = await ListNamespaces().Where(NamespaceShouldBeCleaned).ToListAsync();
@@ -94,8 +94,8 @@ namespace Horde.Storage.Implementation
                 perNamespaceRoots.TryAdd(ns, gcRootState);
             }, cancellationToken);
 
+            ulong countOfBlobsRemoved = 0;
             List<BlobIdentifier> blobsToRemove = new List<BlobIdentifier>();
-            List<BlobIdentifier> removedBlobs = new List<BlobIdentifier>();
 
             // enumerate all namespaces, and check if the old blob is valid in any of them to allow for a blob store to just store them in a single pile if it wants to
             foreach (NamespaceId @namespace in namespaces)
@@ -158,11 +158,11 @@ namespace Horde.Storage.Implementation
                 }
                 if (deleted)
                 {
-                    removedBlobs.Add(blob);
+                    ++countOfBlobsRemoved;
                 }
             }
 
-            return removedBlobs;
+            return countOfBlobsRemoved;
         }
 
         private bool NamespaceShouldBeCleaned(NamespaceId ns)
diff --git a/Engine/Source/Programs/Horde/HordeStorage/Horde.Storage/Implementation/GC/OrphanBlobCleanupRefs.cs b/Engine/Source/Programs/Horde/HordeStorage/Horde.Storage/Implementation/GC/OrphanBlobCleanupRefs.cs
--- a/Engine/Source/Programs/Horde/HordeStorage/Horde.Storage/Implementation/GC/OrphanBlobCleanupRefs.cs
+++ b/Engine/Source/Programs/Horde/HordeStorage/Horde.Storage/Implementation/GC/OrphanBlobCleanupRefs.cs
@@ -50,19 +50,19 @@ namespace Horde.Storage.Implementation
             return true;
         }
 
-        public async Task<List<BlobIdentifier>> Cleanup(CancellationToken cancellationToken)
+        public async Task<ulong> Cleanup(CancellationToken cancellationToken)
         {
             if (!_leaderElection.IsThisInstanceLeader())
             {
                 _logger.Information("Skipped orphan blob (refs) cleanup run as this instance is not the leader");
-                return new List<BlobIdentifier>();
+                return 0;
             }
 
             List<NamespaceId> namespaces = await ListNamespaces().Where(NamespaceShouldBeCleaned).ToListAsync();
           
 
             // enumerate all namespaces, and check if the old blob is valid in any of them to allow for a blob store to just store them in a single pile if it wants to
-            List<BlobIdentifier> removedBlobs = new List<BlobIdentifier>();
+            ulong countOfBlobsRemoved = 0;
             foreach (NamespaceId @namespace in namespaces)
             {
                 if (cancellationToken.IsCancellationRequested)
@@ -122,12 +122,12 @@ namespace Horde.Storage.Implementation
                         using IScope removeBlobScope = Tracer.Instance.StartActive("gc.blob.remove-blob");
                         removeBlobScope.Span.ResourceName = blob.ToString();
                         await RemoveBlob(@namespace, blob);
-                        removedBlobs.Add(blob);
+                        ++countOfBlobsRemoved;
                     }
                 }
             }
 
-            return removedBlobs;
+            return countOfBlobsRemoved;
         }
 
         private async Task RemoveBlob(NamespaceId ns, BlobIdentifier blob)
diff --git a/Engine/Source/Programs/Horde/HordeStorage/Horde.Storage/Tests/Functional/GC/GCBlobTests.cs b/Engine/Source/Programs/Horde/HordeStorage/Horde.Storage/Tests/Functional/GC/GCBlobTests.cs
--- a/Engine/Source/Programs/Horde/HordeStorage/Horde.Storage/Tests/Functional/GC/GCBlobTests.cs
+++ b/Engine/Source/Programs/Horde/HordeStorage/Horde.Storage/Tests/Functional/GC/GCBlobTests.cs
@@ -180,14 +180,8 @@ namespace Horde.Storage.FunctionalTests.GC
             Assert.AreEqual(1, namespaces.Count);
 
             CancellationTokenSource cts = new CancellationTokenSource();
-            List<BlobIdentifier> removedBlobs = await cleanup.Cleanup(cts.Token);
-            Assert.AreEqual(4, removedBlobs.Count);
-            
-            BlobIdentifier[] expectedIdentifiers = {object1id, object2id, object4id, object5id};
-            foreach (BlobIdentifier removedBlob in removedBlobs)
-            {
-                Assert.IsTrue(expectedIdentifiers.Contains(removedBlob));
-            }
+            ulong removedBlobsCount = await cleanup.Cleanup(cts.Token);
+            Assert.AreEqual(4, removedBlobsCount);
             
             _callistoBlobMock.Verify();
             _callistoBlobMock.VerifyNoOtherCalls();
@@ -336,11 +330,18 @@ namespace Horde.Storage.FunctionalTests.GC
             Assert.IsNotNull(cleanup);
 
             CancellationTokenSource cts = new CancellationTokenSource();
-            List<BlobIdentifier> removedBlobs = await cleanup.Cleanup(cts.Token);
-            Assert.AreEqual(3, removedBlobs.Count);
-            
-            BlobIdentifier[] expectedIdentifiers = {object1id, object4id, object5id};
-            CollectionAssert.AreEquivalent(removedBlobs, expectedIdentifiers);
+            ulong countOfRemovedBlobs = await cleanup.Cleanup(cts.Token);
+            Assert.AreEqual(3u, countOfRemovedBlobs);
+
+            foreach (BlobIdentifier blob in new BlobIdentifier[] {object1id, object4id, object5id})
+            {
+                Assert.IsFalse(await _blobService!.Exists(TestNamespace, blob));
+            }
+
+            foreach (BlobIdentifier blob in new BlobIdentifier[] {object2id, object3id, object6id})
+            {
+                Assert.IsTrue(await _blobService!.Exists(TestNamespace, blob));
+            }
         }
     }
 }
diff --git a/Engine/Source/Programs/Horde/HordeStorage/Horde.Storage/Tests/Functional/StorageTests.cs b/Engine/Source/Programs/Horde/HordeStorage/Horde.Storage/Tests/Functional/StorageTests.cs
--- a/Engine/Source/Programs/Horde/HordeStorage/Horde.Storage/Tests/Functional/StorageTests.cs
+++ b/Engine/Source/Programs/Horde/HordeStorage/Horde.Storage/Tests/Functional/StorageTests.cs
@@ -244,13 +244,13 @@ namespace Horde.Storage.FunctionalTests.Storage
             Assert.IsNotNull(fsStore);
 
             CancellationTokenSource cts = new CancellationTokenSource();
-            Assert.IsTrue((await fsStore.CleanupInternal(cts.Token, batchSize: 2)).Count == 0); // No garbage to collect, should return false
+            Assert.IsTrue(await fsStore.CleanupInternal(cts.Token, batchSize: 2) == 0); // No garbage to collect, should return false
             
             FileInfo[] fooFiles = CreateFilesInNamespace(FooNamespace, 10);
 
             Assert.AreEqual(10 * 100, await fsStore.CalculateDiskSpaceUsed());
 
-            Assert.IsTrue((await fsStore.CleanupInternal(cts.Token, batchSize: 2)).Count > 0);
+            Assert.IsTrue(await fsStore.CleanupInternal(cts.Token, batchSize: 2) > 0);
 
             Assert.AreEqual(5 * 100, await fsStore.CalculateDiskSpaceUsed());
 
