Horde Storage - Added clarifying comments that credentials used in the test file is not actually real credentials but rather something we specify for the tests (as minio emulates AWS it uses credentials in the same format as AWS)

diff --git a/Engine/Source/Programs/Horde/HordeStorage/Horde.Storage/appsettings.Testing.json b/Engine/Source/Programs/Horde/HordeStorage/Horde.Storage/appsettings.Testing.json
--- a/Engine/Source/Programs/Horde/HordeStorage/Horde.Storage/appsettings.Testing.json
+++ b/Engine/Source/Programs/Horde/HordeStorage/Horde.Storage/appsettings.Testing.json
@@ -64,6 +64,7 @@
   },
   "AWSCredentials": {
     // minio and health checks expect access keys to have been set
+    // note that these credentials are not actually real credentials but rather something as we pass into minio when starting it
     "AwsAccessKey": "FSYPTVCPKC2IV5D1U8RE",
     "AwsSecretKey": "WaUf4Zq+J6YMlvbCVvTRkQJIPc7DvZwTm5UP0tPJ"
   },
