Horde: Fix tests

diff --git a/Engine/Source/Programs/Horde/Horde.Build/Commands/ServerCommand.cs b/Engine/Source/Programs/Horde/Horde.Build/Commands/ServerCommand.cs
--- a/Engine/Source/Programs/Horde/Horde.Build/Commands/ServerCommand.cs
+++ b/Engine/Source/Programs/Horde/Horde.Build/Commands/ServerCommand.cs
@@ -61,13 +61,6 @@ namespace Horde.Build.Commands
 			}
 		}
 
-		// Used by WebApplicationFactory in controller tests. Uses reflection to call this exact function signature.
-		public static IHostBuilder CreateHostBuilder(string[] Args)
-		{
-			ServerSettings HordeSettings = new ServerSettings();
-			return CreateHostBuilderWithCert(Args, new ConfigurationBuilder().Build(), HordeSettings, null);
-		}
-
 		static IHostBuilder CreateHostBuilderWithCert(string[] Args, IConfiguration Config, ServerSettings ServerSettings, X509Certificate2? SslCert)
 		{
 			return Host.CreateDefaultBuilder(Args)
@@ -133,5 +126,11 @@ namespace Horde.Build.Commands
 				return new X509Certificate2(FileReference.ReadAllBytes(ServerPrivateCert));
 			}
 		}
-	}
+
+		public static IHostBuilder CreateHostBuilderForTesting(string[] Args)
+		{
+			ServerSettings HordeSettings = new ServerSettings();
+			return CreateHostBuilderWithCert(Args, new ConfigurationBuilder().Build(), HordeSettings, null);
+		}
+}
 }
diff --git a/Engine/Source/Programs/Horde/Horde.Build/Program.cs b/Engine/Source/Programs/Horde/Horde.Build/Program.cs
--- a/Engine/Source/Programs/Horde/Horde.Build/Program.cs
+++ b/Engine/Source/Programs/Horde/Horde.Build/Program.cs
@@ -202,6 +202,9 @@ namespace HordeServer
 #pragma warning restore ASP0000 // Do not call 'IServiceCollection.BuildServiceProvider' in 'ConfigureServices'
 		}
 
+		// Used by WebApplicationFactory in controller tests. Uses reflection to call this exact function signature.
+		public static IHostBuilder CreateHostBuilder(string[] Args) => ServerCommand.CreateHostBuilderForTesting(Args);
+
 		/// <summary>
 		/// Get the application directory
 		/// </summary>
