Horde: Extend the sample period for job queue-based auto-scaler

A too narrow period can make jobs disappear from the queue.
This also needs to get moved to a settings value.

diff --git a/Engine/Source/Programs/Horde/Horde.Build/Fleet/Autoscale/JobQueueStrategy.cs b/Engine/Source/Programs/Horde/Horde.Build/Fleet/Autoscale/JobQueueStrategy.cs
--- a/Engine/Source/Programs/Horde/Horde.Build/Fleet/Autoscale/JobQueueStrategy.cs
+++ b/Engine/Source/Programs/Horde/Horde.Build/Fleet/Autoscale/JobQueueStrategy.cs
@@ -24,7 +24,6 @@ namespace Horde.Build.Fleet.Autoscale
 	/// </summary>
 	public class JobQueueStrategy : IPoolSizeStrategy
 	{
-		
 		private readonly IJobCollection Jobs;
 		private readonly IGraphCollection Graphs;
 		private readonly StreamService StreamService;
@@ -33,7 +32,7 @@ namespace Horde.Build.Fleet.Autoscale
 		/// <summary>
 		/// How far back in time to look for job batches (that potentially are in the queue)
 		/// </summary>
-		private readonly TimeSpan SamplePeriod = TimeSpan.FromHours(2.0);
+		private readonly TimeSpan SamplePeriod = TimeSpan.FromDays(5);
 		
 		/// <summary>
 		/// Time spent in ready state before considered truly waiting for an agent
