Horde: Hardening for storage bundles.

- CommitService now writes bundles rather than tree pack objects.
- Legacy TreePack classes removed.
- Added commands for creating, extracting, diffing bundles to Horde Agent and Horde Build (implemented in both places using shared code, in order to account for convenience in configuring the matching storage backend).
- Some bug fixes.

diff --git a/Engine/Source/Programs/Horde/Horde.Agent/Commands/Archive/CreateArchiveCommand.cs b/Engine/Source/Programs/Horde/Horde.Agent/Commands/Archive/CreateArchiveCommand.cs
deleted file mode 100644
diff --git a/Engine/Source/Programs/Horde/Horde.Agent/Commands/Archive/ExtractArchiveCommand.cs b/Engine/Source/Programs/Horde/Horde.Agent/Commands/Archive/ExtractArchiveCommand.cs
deleted file mode 100644
diff --git a/Engine/Source/Programs/Horde/Horde.Agent/Commands/Bundles/CreateCommand.cs b/Engine/Source/Programs/Horde/Horde.Agent/Commands/Bundles/CreateCommand.cs
new file mode 100644
--- /dev/null
+++ b/Engine/Source/Programs/Horde/Horde.Agent/Commands/Bundles/CreateCommand.cs
@@ -0,0 +1,80 @@
+// Copyright Epic Games, Inc. All Rights Reserved.
+
+using EpicGames.Core;
+using EpicGames.Horde.Bundles;
+using EpicGames.Horde.Bundles.Nodes;
+using EpicGames.Horde.Storage;
+using EpicGames.Horde.Storage.Impl;
+using EpicGames.Serialization;
+using Microsoft.Extensions.Configuration;
+using Microsoft.Extensions.DependencyInjection;
+using Microsoft.Extensions.Logging;
+using System;
+using System.Collections.Generic;
+using System.IO;
+using System.Linq;
+using System.Threading.Tasks;
+
+namespace HordeAgent.Commands.Bundles
+{
+	abstract class BundleCommandBase : Command
+	{
+		[CommandLine("-Server=")]
+		public string? Server = null;
+
+		[CommandLine("-StorageDir=", Description = "Overrides the default storage server with a local directory")]
+		public DirectoryReference? StorageDir = null;
+
+		protected IStorageClient CreateStorageClient(ILogger Logger)
+		{
+			if (StorageDir != null)
+			{
+				return new FileStorageClient(StorageDir, Logger);
+			}
+			else
+			{
+				IConfiguration Config = new ConfigurationBuilder()
+					.AddJsonFile("appsettings.json", optional: false)
+					.AddJsonFile($"appsettings.{Environment.GetEnvironmentVariable("DOTNET_ENVIRONMENT")}.json", optional: true) // environment variable overrides, also used in k8s setups with Helm
+					.AddJsonFile("appsettings.User.json", optional: true)
+					.AddEnvironmentVariables()
+					.Build();
+
+				IConfigurationSection ConfigSection = Config.GetSection(AgentSettings.SectionName);
+
+				IServiceCollection Services = new ServiceCollection();
+				Services.AddHordeStorage(Settings => ConfigSection.GetCurrentServerProfile().GetSection(nameof(ServerProfile.Storage)).Bind(Settings));
+
+				IServiceProvider ServiceProvider = Services.BuildServiceProvider();
+				return ServiceProvider.GetRequiredService<IStorageClient>();
+			}
+		}
+	}
+
+	[Command("bundle", "create", "Creates a bundle from a folder on the local hard drive")]
+	class CreateCommand : BundleCommandBase
+	{
+		[CommandLine("-Namespace=")]
+		public NamespaceId NamespaceId { get; set; } = new NamespaceId("default-ns");
+
+		[CommandLine("-Bucket=")]
+		public BucketId BucketId { get; set; } = new BucketId("default-bkt");
+
+		[CommandLine("-Ref=")]
+		public RefId RefId { get; set; } = new RefId("default-ref");
+
+		[CommandLine("-InputDir=", Required = true)]
+		public DirectoryReference InputDir { get; set; } = null!;
+
+		public override async Task<int> ExecuteAsync(ILogger Logger)
+		{
+			IStorageClient StorageClient = base.CreateStorageClient(Logger);
+
+			Bundle<DirectoryNode> Bundle = new Bundle<DirectoryNode>(StorageClient, NamespaceId, new BundleOptions(), null);
+			await Bundle.Root.CopyFromDirectoryAsync(InputDir.ToDirectoryInfo(), new ChunkOptions(), Logger);
+			await Bundle.WriteAsync(BucketId, RefId, false, DateTime.UtcNow);
+
+			return 0;
+		}
+	}
+}
diff --git a/Engine/Source/Programs/Horde/Horde.Agent/Commands/Bundles/DeltaCommand.cs b/Engine/Source/Programs/Horde/Horde.Agent/Commands/Bundles/DeltaCommand.cs
new file mode 100644
--- /dev/null
+++ b/Engine/Source/Programs/Horde/Horde.Agent/Commands/Bundles/DeltaCommand.cs
@@ -0,0 +1,42 @@
+// Copyright Epic Games, Inc. All Rights Reserved.
+
+using EpicGames.Core;
+using EpicGames.Horde.Bundles;
+using EpicGames.Horde.Storage;
+using EpicGames.Horde.Storage.Impl;
+using EpicGames.Serialization;
+using Microsoft.Extensions.Caching.Memory;
+using Microsoft.Extensions.Configuration;
+using Microsoft.Extensions.Logging;
+using System;
+using System.Collections.Generic;
+using System.Diagnostics;
+using System.Linq;
+using System.Text;
+using System.Threading.Tasks;
+
+namespace HordeAgent.Commands.Bundles
+{
+	[Command("bundle", "delta", "Desribe the delta between two bundles")]
+	class DeltaCommand : BundleCommandBase
+	{
+		[CommandLine("-Namespace=")]
+		public NamespaceId NamespaceId { get; set; } = new NamespaceId("default-ns");
+
+		[CommandLine("-Bucket=")]
+		public BucketId BucketId { get; set; } = new BucketId("default-bkt");
+
+		[CommandLine("-Ref1=", Required = true)]
+		public RefId RefId1 { get; set; }
+
+		[CommandLine("-Ref2=", Required = true)]
+		public RefId RefId2 { get; set; }
+
+		public override async Task<int> ExecuteAsync(ILogger Logger)
+		{
+			IStorageClient StorageClient = CreateStorageClient(Logger);
+			await BundleTools.FindDeltaAsync(StorageClient, NamespaceId, BucketId, RefId1, RefId2, Logger);
+			return 0;
+		}
+	}
+}
diff --git a/Engine/Source/Programs/Horde/Horde.Agent/Commands/Bundles/DumpCommand.cs b/Engine/Source/Programs/Horde/Horde.Agent/Commands/Bundles/DumpCommand.cs
new file mode 100644
--- /dev/null
+++ b/Engine/Source/Programs/Horde/Horde.Agent/Commands/Bundles/DumpCommand.cs
@@ -0,0 +1,45 @@
+// Copyright Epic Games, Inc. All Rights Reserved.
+
+using EpicGames.Core;
+using EpicGames.Horde.Bundles;
+using EpicGames.Horde.Storage;
+using EpicGames.Horde.Storage.Impl;
+using EpicGames.Serialization;
+using Microsoft.Extensions.Caching.Memory;
+using Microsoft.Extensions.Configuration;
+using Microsoft.Extensions.DependencyInjection;
+using Microsoft.Extensions.Logging;
+using System;
+using System.Collections.Generic;
+using System.IO;
+using System.Linq;
+using System.Threading.Tasks;
+
+namespace HordeAgent.Commands.Bundles
+{
+	[Command("bundle", "dump", "Indexes a tree and creates a table of contents for each blob describing the contents")]
+	class DumpCommand : BundleCommandBase
+	{
+		[CommandLine("-Namespace=")]
+		public NamespaceId NamespaceId { get; set; } = new NamespaceId("default-ns");
+
+		[CommandLine("-Bucket=")]
+		public BucketId BucketId { get; set; } = new BucketId("default-bkt");
+
+		[CommandLine("-Ref=")]
+		public RefId RefId { get; set; } = new RefId("default-ref");
+
+		[CommandLine("-OutputDir=")]
+		public DirectoryReference? OutputDir { get; set; }
+
+		public override async Task<int> ExecuteAsync(ILogger Logger)
+		{
+			OutputDir = DirectoryReference.Combine(Program.DataDir, "Bundles");
+
+			IStorageClient StorageClient = CreateStorageClient(Logger);
+			await BundleTools.WriteSummaryAsync(StorageClient, NamespaceId, BucketId, RefId, OutputDir);
+
+			return 0;
+		}
+	}
+}
diff --git a/Engine/Source/Programs/Horde/Horde.Agent/Commands/Bundles/ExtractCommand.cs b/Engine/Source/Programs/Horde/Horde.Agent/Commands/Bundles/ExtractCommand.cs
new file mode 100644
--- /dev/null
+++ b/Engine/Source/Programs/Horde/Horde.Agent/Commands/Bundles/ExtractCommand.cs
@@ -0,0 +1,48 @@
+// Copyright Epic Games, Inc. All Rights Reserved.
+
+using EpicGames.Core;
+using EpicGames.Horde.Bundles;
+using EpicGames.Horde.Bundles.Nodes;
+using EpicGames.Horde.Storage;
+using EpicGames.Horde.Storage.Impl;
+using EpicGames.Serialization;
+using Microsoft.Extensions.Caching.Memory;
+using Microsoft.Extensions.Configuration;
+using Microsoft.Extensions.DependencyInjection;
+using Microsoft.Extensions.Logging;
+using System;
+using System.Collections.Generic;
+using System.IO;
+using System.Linq;
+using System.Threading.Tasks;
+
+namespace HordeAgent.Commands.Bundles
+{
+	[Command("bundle", "extract", "Extracts data from a bundle to the local hard drive")]
+	class ExtractCommand : BundleCommandBase
+	{
+		[CommandLine("-Namespace=")]
+		public NamespaceId NamespaceId { get; set; } = new NamespaceId("default-ns");
+
+		[CommandLine("-Bucket=")]
+		public BucketId BucketId { get; set; } = new BucketId("default-bkt");
+
+		[CommandLine("-Ref=")]
+		public RefId RefId { get; set; } = new RefId("default-ref");
+
+		[CommandLine("-OutputDir=", Required = true)]
+		public DirectoryReference OutputDir { get; set; } = null!;
+
+		public override async Task<int> ExecuteAsync(ILogger Logger)
+		{
+			IStorageClient StorageClient = CreateStorageClient(Logger);
+			using (IMemoryCache Cache = new MemoryCache(new MemoryCacheOptions { SizeLimit = 1024 * 1024 * 1000 }))
+			{
+				Bundle<DirectoryNode> Bundle = new Bundle<DirectoryNode>(StorageClient, NamespaceId, new BundleOptions(), Cache);
+				await Bundle.ReadAsync(BucketId, RefId);
+				await Bundle.Root.CopyToDirectoryAsync(OutputDir.ToDirectoryInfo(), Logger);
+			}
+			return 0;
+		}
+	}
+}
diff --git a/Engine/Source/Programs/Horde/Horde.Agent/Horde.Agent.csproj b/Engine/Source/Programs/Horde/Horde.Agent/Horde.Agent.csproj
--- a/Engine/Source/Programs/Horde/Horde.Agent/Horde.Agent.csproj
+++ b/Engine/Source/Programs/Horde/Horde.Agent/Horde.Agent.csproj
@@ -31,6 +31,7 @@
       <PrivateAssets>all</PrivateAssets>
       <IncludeAssets>runtime; build; native; contentfiles; analyzers; buildtransitive</IncludeAssets>
     </PackageReference>
+    <PackageReference Include="Microsoft.Extensions.Caching.Memory" Version="5.0.0" />
     <PackageReference Include="Microsoft.Extensions.Hosting" Version="3.1.5" />
     <PackageReference Include="Microsoft.Extensions.Hosting.Abstractions" Version="3.1.5" />
     <PackageReference Include="Microsoft.Extensions.Hosting.WindowsServices" Version="3.1.5" />
diff --git a/Engine/Source/Programs/Horde/Horde.Build/Commands/Bundles/CreateCommand.cs b/Engine/Source/Programs/Horde/Horde.Build/Commands/Bundles/CreateCommand.cs
new file mode 100644
--- /dev/null
+++ b/Engine/Source/Programs/Horde/Horde.Build/Commands/Bundles/CreateCommand.cs
@@ -0,0 +1,58 @@
+// Copyright Epic Games, Inc. All Rights Reserved.
+
+using EpicGames.Core;
+using EpicGames.Horde.Bundles;
+using EpicGames.Horde.Bundles.Nodes;
+using EpicGames.Horde.Storage;
+using EpicGames.Horde.Storage.Impl;
+using EpicGames.Serialization;
+using Microsoft.Extensions.Configuration;
+using Microsoft.Extensions.DependencyInjection;
+using Microsoft.Extensions.Logging;
+using System;
+using System.Collections.Generic;
+using System.IO;
+using System.Linq;
+using System.Threading.Tasks;
+
+namespace HordeServer.Commands.Bundles
+{
+	abstract class BundleCommandBase : Command
+	{
+		[CommandLine("-StorageDir=", Description = "Overrides the default storage server with a local directory")]
+		public DirectoryReference? StorageDir = null;
+
+		protected IStorageClient CreateStorageClient(ILogger Logger)
+		{
+			StorageDir ??= DirectoryReference.Combine(Program.DataDir, "Storage");
+			return new FileStorageClient(StorageDir, Logger);
+		}
+	}
+
+	[Command("bundle", "create", "Creates a bundle from a folder on the local hard drive")]
+	class CreateCommand : BundleCommandBase
+	{
+		[CommandLine("-Namespace=")]
+		public NamespaceId NamespaceId { get; set; } = new NamespaceId("default-ns");
+
+		[CommandLine("-Bucket=")]
+		public BucketId BucketId { get; set; } = new BucketId("default-bkt");
+
+		[CommandLine("-Ref=")]
+		public RefId RefId { get; set; } = new RefId("default-ref");
+
+		[CommandLine("-InputDir=", Required = true)]
+		public DirectoryReference InputDir { get; set; } = null!;
+
+		public override async Task<int> ExecuteAsync(ILogger Logger)
+		{
+			IStorageClient StorageClient = base.CreateStorageClient(Logger);
+
+			Bundle<DirectoryNode> Bundle = new Bundle<DirectoryNode>(StorageClient, NamespaceId, new BundleOptions(), null);
+			await Bundle.Root.CopyFromDirectoryAsync(InputDir.ToDirectoryInfo(), new ChunkOptions(), Logger);
+			await Bundle.WriteAsync(BucketId, RefId, false, DateTime.UtcNow);
+
+			return 0;
+		}
+	}
+}
diff --git a/Engine/Source/Programs/Horde/Horde.Build/Commands/Bundles/DeltaCommand.cs b/Engine/Source/Programs/Horde/Horde.Build/Commands/Bundles/DeltaCommand.cs
new file mode 100644
--- /dev/null
+++ b/Engine/Source/Programs/Horde/Horde.Build/Commands/Bundles/DeltaCommand.cs
@@ -0,0 +1,42 @@
+// Copyright Epic Games, Inc. All Rights Reserved.
+
+using EpicGames.Core;
+using EpicGames.Horde.Bundles;
+using EpicGames.Horde.Storage;
+using EpicGames.Horde.Storage.Impl;
+using EpicGames.Serialization;
+using Microsoft.Extensions.Caching.Memory;
+using Microsoft.Extensions.Configuration;
+using Microsoft.Extensions.Logging;
+using System;
+using System.Collections.Generic;
+using System.Diagnostics;
+using System.Linq;
+using System.Text;
+using System.Threading.Tasks;
+
+namespace HordeServer.Commands.Bundles
+{
+	[Command("bundle", "delta", "Desribe the delta between two bundles")]
+	class DeltaCommand : BundleCommandBase
+	{
+		[CommandLine("-Namespace=")]
+		public NamespaceId NamespaceId { get; set; } = new NamespaceId("default-ns");
+
+		[CommandLine("-Bucket=")]
+		public BucketId BucketId { get; set; } = new BucketId("default-bkt");
+
+		[CommandLine("-Ref1=", Required = true)]
+		public RefId RefId1 { get; set; }
+
+		[CommandLine("-Ref2=", Required = true)]
+		public RefId RefId2 { get; set; }
+
+		public override async Task<int> ExecuteAsync(ILogger Logger)
+		{
+			IStorageClient StorageClient = CreateStorageClient(Logger);
+			await BundleTools.FindDeltaAsync(StorageClient, NamespaceId, BucketId, RefId1, RefId2, Logger);
+			return 0;
+		}
+	}
+}
diff --git a/Engine/Source/Programs/Horde/Horde.Build/Commands/Bundles/DumpCommand.cs b/Engine/Source/Programs/Horde/Horde.Build/Commands/Bundles/DumpCommand.cs
new file mode 100644
--- /dev/null
+++ b/Engine/Source/Programs/Horde/Horde.Build/Commands/Bundles/DumpCommand.cs
@@ -0,0 +1,45 @@
+// Copyright Epic Games, Inc. All Rights Reserved.
+
+using EpicGames.Core;
+using EpicGames.Horde.Bundles;
+using EpicGames.Horde.Storage;
+using EpicGames.Horde.Storage.Impl;
+using EpicGames.Serialization;
+using Microsoft.Extensions.Caching.Memory;
+using Microsoft.Extensions.Configuration;
+using Microsoft.Extensions.DependencyInjection;
+using Microsoft.Extensions.Logging;
+using System;
+using System.Collections.Generic;
+using System.IO;
+using System.Linq;
+using System.Threading.Tasks;
+
+namespace HordeServer.Commands.Bundles
+{
+	[Command("bundle", "dump", "Indexes a tree and creates a table of contents for each blob describing the contents")]
+	class DumpCommand : BundleCommandBase
+	{
+		[CommandLine("-Namespace=")]
+		public NamespaceId NamespaceId { get; set; } = new NamespaceId("default-ns");
+
+		[CommandLine("-Bucket=")]
+		public BucketId BucketId { get; set; } = new BucketId("default-bkt");
+
+		[CommandLine("-Ref=")]
+		public RefId RefId { get; set; } = new RefId("default-ref");
+
+		[CommandLine("-OutputDir=")]
+		public DirectoryReference? OutputDir { get; set; }
+
+		public override async Task<int> ExecuteAsync(ILogger Logger)
+		{
+			OutputDir = DirectoryReference.Combine(Program.DataDir, "Bundles");
+
+			IStorageClient StorageClient = CreateStorageClient(Logger);
+			await BundleTools.WriteSummaryAsync(StorageClient, NamespaceId, BucketId, RefId, OutputDir);
+
+			return 0;
+		}
+	}
+}
diff --git a/Engine/Source/Programs/Horde/Horde.Build/Commands/Bundles/ExtractCommand.cs b/Engine/Source/Programs/Horde/Horde.Build/Commands/Bundles/ExtractCommand.cs
new file mode 100644
--- /dev/null
+++ b/Engine/Source/Programs/Horde/Horde.Build/Commands/Bundles/ExtractCommand.cs
@@ -0,0 +1,48 @@
+// Copyright Epic Games, Inc. All Rights Reserved.
+
+using EpicGames.Core;
+using EpicGames.Horde.Bundles;
+using EpicGames.Horde.Bundles.Nodes;
+using EpicGames.Horde.Storage;
+using EpicGames.Horde.Storage.Impl;
+using EpicGames.Serialization;
+using Microsoft.Extensions.Caching.Memory;
+using Microsoft.Extensions.Configuration;
+using Microsoft.Extensions.DependencyInjection;
+using Microsoft.Extensions.Logging;
+using System;
+using System.Collections.Generic;
+using System.IO;
+using System.Linq;
+using System.Threading.Tasks;
+
+namespace HordeServer.Commands.Bundles
+{
+	[Command("bundle", "extract", "Extracts data from a bundle to the local hard drive")]
+	class ExtractCommand : BundleCommandBase
+	{
+		[CommandLine("-Namespace=")]
+		public NamespaceId NamespaceId { get; set; } = new NamespaceId("default-ns");
+
+		[CommandLine("-Bucket=")]
+		public BucketId BucketId { get; set; } = new BucketId("default-bkt");
+
+		[CommandLine("-Ref=")]
+		public RefId RefId { get; set; } = new RefId("default-ref");
+
+		[CommandLine("-OutputDir=", Required = true)]
+		public DirectoryReference OutputDir { get; set; } = null!;
+
+		public override async Task<int> ExecuteAsync(ILogger Logger)
+		{
+			IStorageClient StorageClient = CreateStorageClient(Logger);
+			using (IMemoryCache Cache = new MemoryCache(new MemoryCacheOptions { SizeLimit = 1024 * 1024 * 1000 }))
+			{
+				Bundle<DirectoryNode> Bundle = new Bundle<DirectoryNode>(StorageClient, NamespaceId, new BundleOptions(), Cache);
+				await Bundle.ReadAsync(BucketId, RefId);
+				await Bundle.Root.CopyToDirectoryAsync(OutputDir.ToDirectoryInfo(), Logger);
+			}
+			return 0;
+		}
+	}
+}
diff --git a/Engine/Source/Programs/Horde/Horde.Build/Commands/Perforce/ReplicateCommand.cs b/Engine/Source/Programs/Horde/Horde.Build/Commands/Perforce/ReplicateCommand.cs
--- a/Engine/Source/Programs/Horde/Horde.Build/Commands/Perforce/ReplicateCommand.cs
+++ b/Engine/Source/Programs/Horde/Horde.Build/Commands/Perforce/ReplicateCommand.cs
@@ -40,19 +40,26 @@ namespace HordeServer.Commands
 		[CommandLine]
 		public bool Content { get; set; }
 
+		[CommandLine]
+		public string Filter { get; set; } = "//...";
+
 		[CommandLine]
 		public DirectoryReference? OutputDir { get; set; }
 
 		IConfiguration Configuration;
+		ILoggerProvider LoggerProvider;
 
-		public ReplicateCommand(IConfiguration Configuration)
+		public ReplicateCommand(IConfiguration Configuration, ILoggerProvider LoggerProvider)
 		{
 			this.Configuration = Configuration;
+			this.LoggerProvider = LoggerProvider;
 		}
 
 		public override async Task<int> ExecuteAsync(ILogger Logger)
 		{
 			IServiceCollection Services = new ServiceCollection();
+			Services.AddLogging(Builder => Builder.AddProvider(LoggerProvider));
+
 			Startup.AddServices(Services, Configuration);
 
 			OutputDir ??= DirectoryReference.Combine(Program.DataDir, "Storage");
@@ -77,7 +84,16 @@ namespace HordeServer.Commands
 			CommitTree? BaseTree = null; 
 			if (BaseChange != 0)
 			{
-				BaseTree = await CommitService.ReadTreeAsync(Stream.Id, BaseChange);
+				ICommit? Commit = await CommitCollection.GetCommitAsync(Stream.Id, BaseChange);
+				if (Commit == null)
+				{
+					throw new InvalidOperationException($"Unable to find existing commit for {StreamId} @ CL {BaseChange}");
+				}
+				if (Commit.TreeRefId == null)
+				{
+					throw new InvalidOperationException($"No tree for commit {StreamId} @ CL {BaseChange}");
+				}
+				BaseTree = new CommitTree(Commit.StreamId, Commit.Change, Commit.TreeRefId.Value);
 			}
 
 			await foreach (NewCommit NewCommit in CommitService.FindCommitsForClusterAsync(Stream.ClusterName, StreamToFirstChange).Take(Count))
@@ -88,8 +104,7 @@ namespace HordeServer.Commands
 
 				if (Content)
 				{
-					BaseTree = await CommitService.FindCommitTreeAsync(Stream, NewCommit.Change, BaseTree);
-					await CommitService.WriteTreeAsync(Stream.Id, BaseTree);
+					BaseTree = await CommitService.FindCommitTreeAsync(Stream, NewCommit.Change, BaseTree, Filter);
 				}
 			}
 
diff --git a/Engine/Source/Programs/Horde/Horde.Build/Commands/Tree/AnalyzeCommand.cs b/Engine/Source/Programs/Horde/Horde.Build/Commands/Tree/AnalyzeCommand.cs
deleted file mode 100644
diff --git a/Engine/Source/Programs/Horde/Horde.Build/Commands/Tree/ExtractCommand.cs b/Engine/Source/Programs/Horde/Horde.Build/Commands/Tree/ExtractCommand.cs
deleted file mode 100644
diff --git a/Engine/Source/Programs/Horde/Horde.Build/Perforce/Impl/CommitService.cs b/Engine/Source/Programs/Horde/Horde.Build/Perforce/Impl/CommitService.cs
--- a/Engine/Source/Programs/Horde/Horde.Build/Perforce/Impl/CommitService.cs
+++ b/Engine/Source/Programs/Horde/Horde.Build/Perforce/Impl/CommitService.cs
@@ -33,6 +33,9 @@ using EpicGames.Horde.Storage;
 using System.Text;
 using Microsoft.Extensions.Options;
 using Horde.Build.Utilities;
+using EpicGames.Horde.Bundles;
+using EpicGames.Horde.Bundles.Nodes;
+using Microsoft.Extensions.Caching.Memory;
 
 namespace HordeServer.Commits.Impl
 {
@@ -79,7 +82,12 @@ namespace HordeServer.Commits.Impl
 		/// <summary>
 		/// Options for how objects are packed together
 		/// </summary>
-		public TreePackOptions Packing { get; set; } = new TreePackOptions();
+		public BundleOptions Bundle { get; set; } = new BundleOptions();
+
+		/// <summary>
+		/// Options for how objects are sliced
+		/// </summary>
+		public ChunkOptions Chunking { get; set; } = new ChunkOptions();
 	}
 
 	/// <summary>
@@ -98,34 +106,39 @@ namespace HordeServer.Commits.Impl
 		public int Change { get; }
 
 		/// <summary>
-		/// Root object describing the tree contents
+		/// The ref id for storing this tree
 		/// </summary>
-		public TreePackObject RootObject { get; }
+		public RefId RefId { get; }
 
 		/// <summary>
 		/// Constructor
 		/// </summary>
-		public CommitTree(StreamId StreamId, int Change, TreePackObject RootObject)
+		public CommitTree(StreamId StreamId, int Change, RefId RefId)
 		{
 			this.StreamId = StreamId;
 			this.Change = Change;
-			this.RootObject = RootObject;
+			this.RefId = RefId;
 		}
 
 		/// <summary>
-		/// The ref id for storing this tree
+		/// Attempt to create a commit tree from a commit
 		/// </summary>
-		public RefId RefId => GetRefId(Change);
+		public static CommitTree? FromCommit(ICommit Commit)
+		{
+			if (Commit.TreeRefId == null)
+			{
+				return null;
+			}
+			else
+			{
+				return new CommitTree(Commit.StreamId, Commit.Change, Commit.TreeRefId.Value);
+			}
+		}
 
 		/// <summary>
 		/// The ref id for storing this tree
 		/// </summary>
-		public static RefId GetRefId(int Change) => new RefId(IoHash.Compute(Encoding.UTF8.GetBytes($"{Change}")));
-
-		/// <summary>
-		/// Gets the directory at the root of this tree
-		/// </summary>
-		public TreePackDirNode RootDirNode => TreePackDirNode.Parse(RootObject.GetRootNode());
+		public static RefId GetDefaultRefId(int Change) => new RefId(IoHash.Compute(Encoding.UTF8.GetBytes($"{Change}")));
 	}
 
 	/// <summary>
@@ -187,6 +200,7 @@ namespace HordeServer.Commits.Impl
 		IStreamCollection StreamCollection;
 		IPerforceService PerforceService;
 		IUserCollection UserCollection;
+		IMemoryCache MemoryCache;
 		ILogger<CommitService> Logger;
 
 		const int MaxBackgroundTasks = 2;
@@ -194,7 +208,7 @@ namespace HordeServer.Commits.Impl
 		/// <summary>
 		/// Constructor
 		/// </summary>
-		public CommitService(IDatabase Redis, ICommitCollection CommitCollection, IStorageClient StorageClient, IStreamCollection StreamCollection, IPerforceService PerforceService, IUserCollection UserCollection, IClock Clock, IOptions<CommitServiceOptions> Options, ILogger<CommitService> Logger)
+		public CommitService(IDatabase Redis, ICommitCollection CommitCollection, IStorageClient StorageClient, IStreamCollection StreamCollection, IPerforceService PerforceService, IUserCollection UserCollection, IClock Clock, IMemoryCache MemoryCache, IOptions<CommitServiceOptions> Options, ILogger<CommitService> Logger)
 		{
 			this.Options = Options.Value;
 
@@ -207,6 +221,7 @@ namespace HordeServer.Commits.Impl
 			this.StreamCollection = StreamCollection;
 			this.PerforceService = PerforceService;
 			this.UserCollection = UserCollection;
+			this.MemoryCache = MemoryCache;
 			this.Logger = Logger;
 			this.UpdateMetadataTicker = Clock.AddSharedTicker<CommitService>(TimeSpan.FromSeconds(30.0), UpdateMetadataAsync, Logger);
 		}
@@ -820,40 +835,6 @@ namespace HordeServer.Commits.Impl
 			}
 		}
 
-		/// <summary>
-		/// Gets the tree for a particular commit
-		/// </summary>
-		/// <param name="Commit">The commit instance</param>
-		/// <returns>The corresponding commit tree</returns>
-		public async Task<CommitTree?> GetTreeAsync(ICommit Commit)
-		{
-			if (Commit.TreeRefId == null)
-			{
-				return null;
-			}
-
-			IRef? Ref = await StorageClient.GetRefAsync(Options.NamespaceId, GetStreamBucketId(Commit.StreamId), Commit.TreeRefId.Value);
-			if(Ref == null)
-			{
-				return null;
-			}
-
-			return new CommitTree(Commit.StreamId, Commit.Change, TreePackObject.Parse(Ref));
-		}
-
-		public async Task<CommitTree> ReadTreeAsync(StreamId StreamId, int Change)
-		{
-			RefId TreeRefId = new RefId(IoHash.Compute(Encoding.UTF8.GetBytes($"{Change}")));
-			IRef Ref = await StorageClient.GetRefAsync(Options.NamespaceId, new BucketId(StreamId.ToString()), TreeRefId);
-			return new CommitTree(StreamId, Change, TreePackObject.Parse(Ref));
-		}
-
-		public async Task WriteTreeAsync(StreamId StreamId, CommitTree Tree)
-		{
-			CbObject Object = Tree.RootObject.ToCbObject();
-			await StorageClient.SetRefAsync(Options.NamespaceId, new BucketId(StreamId.ToString()), Tree.RefId, Object);
-		}
-
 		/// <summary>
 		/// Replicates the contents of a stream to Horde storage, optionally using the given change as a starting point
 		/// </summary>
@@ -873,21 +854,16 @@ namespace HordeServer.Commits.Impl
 				}
 				else
 				{
-					BaseTree = new CommitTree(BaseCommit.StreamId, BaseCommit.Change, TreePackObject.Parse(await StorageClient.GetRefAsync(Options.NamespaceId, GetStreamBucketId(Stream.Id), BaseCommit.TreeRefId.Value)));
+					BaseTree = new CommitTree(BaseCommit.StreamId, BaseCommit.Change, BaseCommit.TreeRefId.Value);
 				}
 			}
 
-			// Create the new root object
+			// Write the new tree
 			CommitTree Tree = await FindCommitTreeAsync(Stream, Commit.Change, BaseTree);
 
-			// Write the new ref
-			string RefName = $"{Commit.Change}";
-			RefId TreeRefId = new RefId(IoHash.Compute(Encoding.UTF8.GetBytes(RefName)));
-			await StorageClient.SetRefAsync(Options.NamespaceId, new BucketId(Stream.Id.ToString()), TreeRefId, Tree.RootObject.ToCbObject());
-
 			// Update the commit
 			NewCommit NewCommit = new NewCommit(Commit);
-			NewCommit.TreeRefId = TreeRefId;
+			NewCommit.TreeRefId = Tree.RefId;
 			return await CommitCollection.AddOrReplaceAsync(NewCommit);
 		}
 
@@ -896,34 +872,35 @@ namespace HordeServer.Commits.Impl
 		/// </summary>
 		/// <param name="Stream">The stream to replicate</param>
 		/// <param name="Change">Commit to store the tree ref</param>
-		/// <param name="BaseTree">The base tree to update from</param>
+		/// <param name="BaseTree">The initial change to update from</param>
+		/// <param name="QueryPath">Depot path to query for changes. Will default to the entire depot (but filtered by the workspace)</param>
 		/// <returns>Root tree object</returns>
-		public async Task<CommitTree> FindCommitTreeAsync(IStream Stream, int Change, CommitTree? BaseTree)
+		public async Task<CommitTree> FindCommitTreeAsync(IStream Stream, int Change, CommitTree? BaseTree, string QueryPath = "//...")
 		{
 			// Create a client to replicate from this stream, and connect to the server
 			ReplicationClient ClientInfo = await FindOrAddReplicationClientAsync(Stream);
 			using IPerforceConnection Perforce = await PerforceConnection.CreateAsync(ClientInfo.Settings, Logger);
 
 			// Get the initial directory state
-			TreePackDirWriter DirWriter;
 			if (BaseTree == null)
 			{
 				await FlushWorkspaceAsync(ClientInfo, Perforce, 0);
-				DirWriter = new TreePackDirWriter(ClientInfo.TreePack);
+				ClientInfo.Contents.Root.Clear();
 			}
 			else
 			{
 				await FlushWorkspaceAsync(ClientInfo, Perforce, BaseTree.Change);
-				DirWriter = new TreePackDirWriter(ClientInfo.TreePack, BaseTree.RootDirNode);
+				await ClientInfo.Contents.ReadAsync(GetStreamBucketId(BaseTree.StreamId), BaseTree.RefId);
 			}
 
 			// Apply all the updates
 			Logger.LogInformation("Updating client {Client} to changelist {Change}", ClientInfo.Client.Name, Change);
 			ClientInfo.Change = -1;
 
+			DirectoryNode Contents = ClientInfo.Contents.Root;
 			if (Options.WriteStubFiles)
 			{
-				await foreach (PerforceResponse<SyncRecord> Record in Perforce.StreamCommandAsync<SyncRecord>("sync", new[] { "-k", $"//...@{Change}" }, null, default))
+				await foreach (PerforceResponse<SyncRecord> Record in Perforce.StreamCommandAsync<SyncRecord>("sync", new[] { "-k", $"{QueryPath}@{Change}" }, null, default))
 				{
 					SyncRecord SyncRecord = Record.Data;
 					if (!SyncRecord.Path.StartsWith(ClientInfo.Client.Root, StringComparison.Ordinal))
@@ -931,22 +908,17 @@ namespace HordeServer.Commits.Impl
 						throw new ArgumentException($"Unable to make path {ClientInfo.Client.Root} relative to client root {ClientInfo.Client.Root}");
 					}
 
-					string RelativePath = SyncRecord.Path.Substring(ClientInfo.Client.Root.Length).Replace('\\', '/').TrimStart('/');
-
-					string DepotPath = $"{SyncRecord.DepotFile}#{SyncRecord.Revision}";
-					byte[] Data = new byte[Encoding.UTF8.GetMaxByteCount(DepotPath.Length) + 1];
-					Data[0] = (byte)TreePackNodeType.Binary;
-					int Length = Encoding.UTF8.GetBytes(DepotPath, Data.AsSpan(1));
-					ReadOnlyMemory<byte> EncodedData = Data.AsMemory(0, Length + 1);
+					string Path = SyncRecord.Path.Substring(ClientInfo.Client.Root.Length).Replace('\\', '/');
+					ChunkNode File = await Contents.CreateFileByPathAsync(Path, DirectoryEntryFlags.PerforceDepotPathAndRevision);
 
-					IoHash Hash = await ClientInfo.TreePack.AddNodeAsync(EncodedData);
-					await DirWriter.FindOrAddFileByPathAsync(RelativePath, TreePackDirEntryFlags.File | TreePackDirEntryFlags.PerforceDepotPathAndRevision, Hash, Sha1Hash.Zero);
+					byte[] Data = Encoding.UTF8.GetBytes($"{SyncRecord.DepotFile}#{SyncRecord.Revision}");
+					File.Append(Data, Options.Chunking);
 				}
 			}
 			else
 			{
-				Dictionary<int, (string Path, TreePackFileWriter Writer)> Files = new Dictionary<int, (string, TreePackFileWriter)>();
-				await foreach (PerforceResponse Response in Perforce.StreamCommandAsync("sync", Array.Empty<string>(), new string[] { $"//...@{Change}" }, null, typeof(SyncRecord), true, default))
+				Dictionary<int, ChunkNode> Files = new Dictionary<int, ChunkNode>();
+				await foreach (PerforceResponse Response in Perforce.StreamCommandAsync("sync", Array.Empty<string>(), new string[] { $"{QueryPath}@{Change}" }, null, typeof(SyncRecord), true, default))
 				{
 					PerforceIo? Io = Response.Io;
 					if (Io != null)
@@ -954,25 +926,21 @@ namespace HordeServer.Commits.Impl
 						if (Io.Command == PerforceIoCommand.Open)
 						{
 							string Path = GetClientRelativePath(Io.Payload, ClientInfo.Client.Root);
-							TreePackFileWriter FileWriter = new TreePackFileWriter(ClientInfo.TreePack);
-							Files[Io.File] = (Path, FileWriter);
+							Files[Io.File] = await Contents.CreateFileByPathAsync(Path, DirectoryEntryFlags.File);
 						}
 						else if (Io.Command == PerforceIoCommand.Write)
 						{
-							TreePackFileWriter FileWriter = Files[Io.File].Writer;
-							await FileWriter.WriteAsync(Io.Payload, false);
+							ChunkNode File = Files[Io.File];
+							File.Append(Io.Payload, Options.Chunking);
 						}
 						else if (Io.Command == PerforceIoCommand.Close)
 						{
-							(string Path, TreePackFileWriter FileWriter) = Files[Io.File];
-							IoHash FileHash = await FileWriter.FinalizeAsync();
-							await DirWriter.FindOrAddFileByPathAsync(Path, TreePackDirEntryFlags.File, FileHash, Sha1Hash.Zero);
 							Files.Remove(Io.File);
 						}
 						else if (Io.Command == PerforceIoCommand.Unlink)
 						{
 							string Path = GetClientRelativePath(Io.Payload, ClientInfo.Client.Root);
-							await DirWriter.RemoveFileByPathAsync(Path);
+							await Contents.DeleteFileByPathAsync(Path);
 						}
 					}
 				}
@@ -981,9 +949,10 @@ namespace HordeServer.Commits.Impl
 			ClientInfo.Change = Change;
 
 			// Return the new root object
-			IoHash RootHash = await DirWriter.FlushAsync();
-			TreePackObject RootObject = await ClientInfo.TreePack.FlushAsync(RootHash, DateTime.UtcNow);
-			return new CommitTree(Stream.Id, Change, RootObject);
+			RefId RefId = CommitTree.GetDefaultRefId(Change);
+			Logger.LogInformation("Writing ref {RefId} for {StreamId} change {Change}", RefId, Stream.Id, Change);
+			await ClientInfo.Contents.WriteAsync(GetStreamBucketId(Stream.Id), RefId, true, DateTime.UtcNow);
+			return new CommitTree(Stream.Id, Change, RefId);
 		}
 
 		async Task FlushWorkspaceAsync(ReplicationClient ClientInfo, IPerforceConnection Perforce, int Change)
@@ -1025,16 +994,16 @@ namespace HordeServer.Commits.Impl
 			public InfoRecord ServerInfo { get; }
 			public ClientRecord Client { get; }
 			public int Change { get; set; }
-			public TreePack TreePack { get; }
+			public Bundle<DirectoryNode> Contents { get; }
 
-			public ReplicationClient(PerforceSettings Settings, string ClusterName, InfoRecord ServerInfo, ClientRecord Client, int Change, TreePack TreePack)
+			public ReplicationClient(PerforceSettings Settings, string ClusterName, InfoRecord ServerInfo, ClientRecord Client, int Change, Bundle<DirectoryNode> Contents)
 			{
 				this.Settings = Settings;
 				this.ClusterName = ClusterName;
 				this.ServerInfo = ServerInfo;
 				this.Client = Client;
 				this.Change = Change;
-				this.TreePack = TreePack;
+				this.Contents = Contents;
 			}
 		}
 
@@ -1085,8 +1054,8 @@ namespace HordeServer.Commits.Impl
 				Settings.ClientName = NewClient.Name;
 				Settings.PreferNativeClient = true;
 
-				TreePack TreePack = new TreePack(StorageClient, Options.NamespaceId, Options.Packing);
-				ClientInfo = new ReplicationClient(Settings, Stream.ClusterName, ServerInfo, NewClient, -1, TreePack);
+				Bundle<DirectoryNode> Bundle = new Bundle<DirectoryNode>(StorageClient, Options.NamespaceId, Options.Bundle, MemoryCache);
+				ClientInfo = new ReplicationClient(Settings, Stream.ClusterName, ServerInfo, NewClient, -1, Bundle);
 				CachedPerforceClients.Add(Stream.Id, ClientInfo);
 			}
 			return ClientInfo;
diff --git a/Engine/Source/Programs/Shared/EpicGames.Core/StringUtils.cs b/Engine/Source/Programs/Shared/EpicGames.Core/StringUtils.cs
--- a/Engine/Source/Programs/Shared/EpicGames.Core/StringUtils.cs
+++ b/Engine/Source/Programs/Shared/EpicGames.Core/StringUtils.cs
@@ -342,7 +342,8 @@ namespace EpicGames.Core
 		{
 			if((Text.Length & 1) != 0)
 			{
-				throw new FormatException("Length of hex string must be a multiple of two characters");
+				OutBytes = null;
+				return false;
 			}
 
 			byte[] Bytes = new byte[Text.Length / 2];
@@ -369,7 +370,8 @@ namespace EpicGames.Core
 		{
 			if ((Text.Length & 1) != 0)
 			{
-				throw new FormatException("Length of hex string must be a multiple of two characters");
+				OutBytes = null;
+				return false;
 			}
 
 			byte[] Bytes = new byte[Text.Length / 2];
diff --git a/Engine/Source/Programs/Shared/EpicGames.Horde.Tests/BundleTests.cs b/Engine/Source/Programs/Shared/EpicGames.Horde.Tests/BundleTests.cs
--- a/Engine/Source/Programs/Shared/EpicGames.Horde.Tests/BundleTests.cs
+++ b/Engine/Source/Programs/Shared/EpicGames.Horde.Tests/BundleTests.cs
@@ -89,8 +89,8 @@ namespace HordeServerTests
 		public async Task BasicTestDirectory()
 		{
 			Bundle<DirectoryNode> Bundle = new Bundle<DirectoryNode>(StorageClient, NamespaceId, new BundleOptions(), null);
-			DirectoryNode Node = await Bundle.Root.AddDirectoryAsync("hello");
-			DirectoryNode Node2 = await Node.AddDirectoryAsync("world");
+			DirectoryNode Node = Bundle.Root.AddDirectory("hello");
+			DirectoryNode Node2 = Node.AddDirectory("world");
 
 			RefId RefId = new RefId("testref");
 			await Bundle.WriteAsync(BucketId, RefId, false, DateTime.MinValue);
@@ -121,11 +121,6 @@ namespace HordeServerTests
 			Assert.AreEqual(0, OutputNode2!.Entries.Count);
 		}
 
-
-
-		// caching
-		// hierarchical chunking
-
 		[TestMethod]
 		public async Task DedupTests()
 		{
@@ -135,9 +130,9 @@ namespace HordeServerTests
 
 			Bundle<DirectoryNode> Bundle = new Bundle<DirectoryNode>(StorageClient, NamespaceId, Options, null);
 
-			await Bundle.Root.AddDirectoryAsync("node1");
-			await Bundle.Root.AddDirectoryAsync("node2");
-			await Bundle.Root.AddDirectoryAsync("node3");
+			Bundle.Root.AddDirectory("node1");
+			Bundle.Root.AddDirectory("node2");
+			Bundle.Root.AddDirectory("node3");
 
 			RefId RefId = new RefId("ref");
 			await Bundle.WriteAsync(BucketId, RefId, false, DateTime.UtcNow);
@@ -155,10 +150,10 @@ namespace HordeServerTests
 
 			Bundle<DirectoryNode> Bundle = new Bundle<DirectoryNode>(StorageClient, NamespaceId, Options, null);
 
-			DirectoryNode Node1 = await Bundle.Root.AddDirectoryAsync("node1");
-			DirectoryNode Node2 = await Node1.AddDirectoryAsync("node2");
-			DirectoryNode Node3 = await Node2.AddDirectoryAsync("node3");
-			DirectoryNode Node4 = await Node3.AddDirectoryAsync("node4");
+			DirectoryNode Node1 = Bundle.Root.AddDirectory("node1");
+			DirectoryNode Node2 = Node1.AddDirectory("node2");
+			DirectoryNode Node3 = Node2.AddDirectory("node3");
+			DirectoryNode Node4 = Node3.AddDirectory("node4");
 
 			RefId RefId = new RefId("ref");
 			await Bundle.WriteAsync(BucketId, RefId, false, DateTime.UtcNow);
@@ -193,10 +188,10 @@ namespace HordeServerTests
 
 			Bundle<DirectoryNode> Bundle = new Bundle<DirectoryNode>(StorageClient, NamespaceId, Options, null);
 
-			DirectoryNode Node1 = await Bundle.Root.AddDirectoryAsync("node1");
-			DirectoryNode Node2 = await Node1.AddDirectoryAsync("node2");
-			DirectoryNode Node3 = await Node2.AddDirectoryAsync("node3");
-			DirectoryNode Node4 = await Bundle.Root.AddDirectoryAsync("node4"); // same contents as node 3
+			DirectoryNode Node1 = Bundle.Root.AddDirectory("node1");
+			DirectoryNode Node2 = Node1.AddDirectory("node2");
+			DirectoryNode Node3 = Node2.AddDirectory("node3");
+			DirectoryNode Node4 = Bundle.Root.AddDirectory("node4"); // same contents as node 3
 
 			RefId RefId1 = new RefId("ref1");
 			await Bundle.WriteAsync(BucketId, RefId1, false, Time);
diff --git a/Engine/Source/Programs/Shared/EpicGames.Horde/Bundles/Bundle.cs b/Engine/Source/Programs/Shared/EpicGames.Horde/Bundles/Bundle.cs
--- a/Engine/Source/Programs/Shared/EpicGames.Horde/Bundles/Bundle.cs
+++ b/Engine/Source/Programs/Shared/EpicGames.Horde/Bundles/Bundle.cs
@@ -218,7 +218,7 @@ namespace EpicGames.Horde.Bundles
 		{
 			if(!Blob.Mounted)
 			{
-				BundleObject Object = await StorageClient.ReadObjectAsync<BundleObject>(NamespaceId, Blob.Hash);
+				BundleObject Object = await GetObjectAsync(Blob.Hash);
 				RegisterObject(Blob, Object, ReadOnlyMemory<byte>.Empty);
 				Blob.Mounted = true;
 			}
@@ -571,6 +571,9 @@ namespace EpicGames.Horde.Bundles
 				Object.ImportObjects.Add(ImportObject);
 			}
 
+			// Size of data currently stored in the block buffer
+			int BlockSize = 0;
+
 			// Compress all the nodes into the encoded buffer
 			BundleCompressionPacket Packet = new BundleCompressionPacket(0);
 			foreach (NodeInfo Node in Nodes)
@@ -578,10 +581,11 @@ namespace EpicGames.Horde.Bundles
 				ReadOnlyMemory<byte> NodeData = await GetDataAsync(Node);
 
 				// If we can't fit this data into the current block, flush the contents of it first
-				if (Packet.DecodedLength + NodeData.Length > Options.MinCompressionPacketSize)
+				if (BlockSize > 0 && BlockSize + NodeData.Length > Options.MinCompressionPacketSize)
 				{
-					FlushPacket(Packet);
+					FlushPacket(BlockBuffer.AsMemory(0, BlockSize), Packet);
 					Packet = new BundleCompressionPacket(Packet.Offset + Packet.EncodedLength);
+					BlockSize = 0;
 				}
 
 				// Create the export for this node
@@ -593,18 +597,18 @@ namespace EpicGames.Horde.Bundles
 				int Offset = Packet.EncodedLength;
 				if (NodeData.Length < Options.MinCompressionPacketSize)
 				{
-					int RequiredSize = Math.Max(Packet.DecodedLength + NodeData.Length, (int)(Options.MaxBlobSize * 1.2));
-					CreateFreeSpace(ref BlockBuffer, Packet.DecodedLength, RequiredSize);
-					NodeData.CopyTo(BlockBuffer.AsMemory(Packet.DecodedLength));
-					Packet.DecodedLength += NodeData.Length;
+					int RequiredSize = Math.Max(BlockSize + NodeData.Length, (int)(Options.MaxBlobSize * 1.2));
+					CreateFreeSpace(ref BlockBuffer, BlockSize, RequiredSize);
+					NodeData.CopyTo(BlockBuffer.AsMemory(BlockSize));
+					BlockSize += NodeData.Length;
 				}
 				else
 				{
-					FlushPacket(Packet);
+					FlushPacket(NodeData, Packet);
 					Packet = new BundleCompressionPacket(Packet.Offset + Packet.EncodedLength);
 				}
 			}
-			FlushPacket(Packet);
+			FlushPacket(BlockBuffer.AsMemory(0, BlockSize), Packet);
 
 			// Flush the data
 			Object.Data = EncodedBuffer.AsSpan(0, Packet.Offset + Packet.EncodedLength).ToArray();
@@ -612,17 +616,20 @@ namespace EpicGames.Horde.Bundles
 			return Object;
 		}
 
-		void FlushPacket(BundleCompressionPacket Packet)
+		void FlushPacket(ReadOnlyMemory<byte> InputData, BundleCompressionPacket Packet)
 		{
-			if (Packet.DecodedLength > 0)
+			if (InputData.Length > 0)
 			{
-				int MinFreeSpace = LZ4Codec.MaximumOutputSize(Packet.DecodedLength);
+				int MinFreeSpace = LZ4Codec.MaximumOutputSize(InputData.Length);
 				CreateFreeSpace(ref EncodedBuffer, Packet.Offset, Packet.Offset + MinFreeSpace);
 
-				ReadOnlySpan<byte> InputSpan = BlockBuffer.AsSpan(0, Packet.DecodedLength);
+				ReadOnlySpan<byte> InputSpan = InputData.Span;
 				Span<byte> OutputSpan = EncodedBuffer.AsSpan(Packet.Offset);
 
+				Packet.DecodedLength = InputData.Length;
 				Packet.EncodedLength = LZ4Codec.Encode(InputSpan, OutputSpan);
+
+				Debug.Assert(Packet.EncodedLength >= 0);
 			}
 		}
 
diff --git a/Engine/Source/Programs/Shared/EpicGames.Horde/Bundles/BundleTools.cs b/Engine/Source/Programs/Shared/EpicGames.Horde/Bundles/BundleTools.cs
new file mode 100644
--- /dev/null
+++ b/Engine/Source/Programs/Shared/EpicGames.Horde/Bundles/BundleTools.cs
@@ -0,0 +1,323 @@
+// Copyright Epic Games, Inc. All Rights Reserved.
+
+using EpicGames.Core;
+using EpicGames.Horde.Bundles.Nodes;
+using EpicGames.Horde.Storage;
+using EpicGames.Serialization;
+using K4os.Compression.LZ4;
+using Microsoft.Extensions.Logging;
+using System;
+using System.Buffers.Binary;
+using System.Collections;
+using System.Collections.Generic;
+using System.Diagnostics;
+using System.Diagnostics.CodeAnalysis;
+using System.IO;
+using System.IO.Compression;
+using System.Linq;
+using System.Security.Cryptography;
+using System.Text;
+using System.Text.Json;
+using System.Threading.Tasks;
+
+namespace EpicGames.Horde.Bundles
+{
+	/// <summary>
+	/// Tools for manipulating bundles
+	/// </summary>
+	public static class BundleTools
+	{
+		class NodeInfo
+		{
+			public IoHash Hash { get; }
+			public int Rank { get; }
+			public int Length { get; }
+			public BlobInfo Blob { get; }
+
+			public BundleExport? Export { get; set; }
+
+			public NodeInfo(IoHash Hash, int Rank, int Length, BlobInfo Blob)
+			{
+				this.Blob = Blob;
+				this.Hash = Hash;
+				this.Rank = Rank;
+				this.Length = Length;
+			}
+		}
+
+		class BlobInfo
+		{
+			public IoHash Hash { get; }
+			public int Size { get; set; }
+			public int DataSize { get; set; }
+			public List<NodeInfo>? References { get; set; }
+			public Dictionary<IoHash, NodeInfo> Nodes { get; } = new Dictionary<IoHash, NodeInfo>();
+
+			public BlobInfo(IoHash Hash, int DataSize)
+			{
+				this.Hash = Hash;
+				this.DataSize = DataSize;
+			}
+		}
+
+		/// <summary>
+		/// Finds the delta between two bundles
+		/// </summary>
+		/// <param name="StorageClient"></param>
+		/// <param name="NamespaceId"></param>
+		/// <param name="BucketId"></param>
+		/// <param name="RefId1">First ref to compare</param>
+		/// <param name="RefId2">Second ref to compare</param>
+		/// <param name="Logger">Logger to write output data</param>
+		public static async Task FindDeltaAsync(IStorageClient StorageClient, NamespaceId NamespaceId, BucketId BucketId, RefId RefId1, RefId RefId2, ILogger Logger)
+		{
+			Dictionary<IoHash, BlobInfo> BlobCache = new Dictionary<IoHash, BlobInfo>();
+
+			Logger.LogInformation("Reading internal blobs...");
+			HashSet<NodeInfo> Nodes1 = new HashSet<NodeInfo>();
+			BlobInfo RootBlob1 = await ScanTreeAsync(StorageClient, NamespaceId, BucketId, RefId1, BlobCache, Nodes1);
+
+			HashSet<NodeInfo> Nodes2 = new HashSet<NodeInfo>();
+			BlobInfo RootBlob2 = await ScanTreeAsync(StorageClient, NamespaceId, BucketId, RefId2, BlobCache, Nodes2);
+
+			Logger.LogInformation("Reading leaf blobs...");
+			foreach (BlobInfo Blob in BlobCache.Values)
+			{
+				if (Blob.References == null)
+				{
+					ReadOnlyMemory<byte> Data = await StorageClient.ReadBlobToMemoryAsync(NamespaceId, Blob.Hash);
+					MountObject(StorageClient, NamespaceId, Blob, Data, null!);
+				}
+			}
+
+			Logger.LogInformation("");
+
+			Logger.LogInformation("{Row}", "                                               |            |       Total       |        Used (Ref 1)        |        Used (Ref 2)        ");
+			Logger.LogInformation("{Row}", " Identifier                                    |    RawSize |    Num       Size |    Num         Size    Pct |    Num         Size    Pct ");
+			Logger.LogInformation("{Row}", " ----------------------------------------------|------------|-------------------|----------------------------|----------------------------");
+
+			WriteBlobInfo($"Ref1 {RefId1.Hash}", RootBlob1, Nodes1, Nodes2, Logger);
+			WriteBlobInfo($"Ref2 {RefId2.Hash}", RootBlob2, Nodes1, Nodes2, Logger);
+
+			foreach (BlobInfo Blob in BlobCache.Values)
+			{
+				WriteBlobInfo($"Blob {Blob.Hash}", Blob, Nodes1, Nodes2, Logger);
+			}
+
+			Logger.LogInformation("");
+
+			List<BlobInfo> NewBlobs = Nodes2.Select(x => x.Blob).Distinct().Except(Nodes2.Select(x => x.Blob)).ToList();
+			Logger.LogInformation("Total blob download size: {NumBytes:n0} bytes", NewBlobs.Sum(x => x.Size));
+		}
+
+		static void WriteBlobInfo(string Name, BlobInfo Blob, HashSet<NodeInfo> Nodes1, HashSet<NodeInfo> Nodes2, ILogger Logger)
+		{
+			int UsedCount1 = 0;
+			int UsedSize1 = 0;
+
+			int UsedCount2 = 0;
+			int UsedSize2 = 0;
+
+			foreach (NodeInfo Node in Blob.Nodes.Values)
+			{
+				if (Nodes1.Contains(Node))
+				{
+					UsedCount1++;
+					UsedSize1 += Node.Length;
+				}
+				if (Nodes2.Contains(Node))
+				{
+					UsedCount2++;
+					UsedSize2 += Node.Length;
+				}
+			}
+
+			string UsedPct1 = $"{(UsedSize1 * 100) / Blob.DataSize}%";
+			string UsedPct2 = $"{(UsedSize2 * 100) / Blob.DataSize}%";
+
+			string Row = $"{Name,42} | {Blob.Size,10:n0} | {Blob.Nodes.Count,6:n0} {Blob.DataSize,10:n0} | {UsedCount1,6:n0} {UsedSize1,12:n0} {UsedPct1,6} | {UsedCount2,6:n0} {UsedSize2,12:n0} {UsedPct2,6}";
+			Logger.LogInformation(" {Row}", Row);
+		}
+
+		static async Task<BlobInfo> ScanTreeAsync(IStorageClient StorageClient, NamespaceId NamespaceId, BucketId BucketId, RefId RefId, Dictionary<IoHash, BlobInfo> BlobCache, HashSet<NodeInfo> Nodes)
+		{
+			IRef Ref = await StorageClient.GetRefAsync(NamespaceId, BucketId, RefId);
+
+			BlobInfo RootBlob = new BlobInfo(IoHash.Zero, 0);
+			BundleObject Object = MountObject(StorageClient, NamespaceId, RootBlob, Ref.Value.GetView(), BlobCache);
+
+			NodeInfo RootNode = RootBlob.Nodes[Object.Exports[^1].Hash];
+			await ScanTreeAsync(StorageClient, NamespaceId, RootNode, BlobCache, Nodes);
+
+			return RootBlob;
+		}
+
+		static async Task ScanTreeAsync(IStorageClient StorageClient, NamespaceId NamespaceId, NodeInfo Node, Dictionary<IoHash, BlobInfo> BlobCache, HashSet<NodeInfo> Nodes)
+		{
+			if (Nodes.Add(Node) && Node.Rank > 0)
+			{
+				if (Node.Export == null)
+				{
+					ReadOnlyMemory<byte> Data = await StorageClient.ReadBlobToMemoryAsync(NamespaceId, Node.Blob.Hash);
+					MountObject(StorageClient, NamespaceId, Node.Blob, Data, BlobCache);
+				}
+
+				Debug.Assert(Node.Export != null);
+
+				foreach (int RefIdx in Node.Export.References)
+				{
+					NodeInfo RefNode = Node.Blob.References![RefIdx];
+					await ScanTreeAsync(StorageClient, NamespaceId, RefNode, BlobCache, Nodes);
+				}
+			}
+		}
+
+		static BundleObject MountObject(IStorageClient StorageClient, NamespaceId NamespaceId, BlobInfo Blob, ReadOnlyMemory<byte> BlobData, Dictionary<IoHash, BlobInfo> BlobCache)
+		{
+			BundleObject Object = CbSerializer.Deserialize<BundleObject>(BlobData);
+
+			Blob.Size = BlobData.Length;
+			Blob.DataSize = Object.Data.Length;
+			Blob.References = new List<NodeInfo>();
+
+			foreach (BundleExport Export in Object.Exports)
+			{
+				NodeInfo? Node;
+				if (!Blob.Nodes.TryGetValue(Export.Hash, out Node))
+				{
+					Node = new NodeInfo(Export.Hash, Export.Rank, Export.Cost, Blob);
+					Blob.Nodes.Add(Export.Hash, Node);
+				}
+				Blob.References.Add(Node);
+
+				Node.Export = Export;
+			}
+
+			foreach (BundleImportObject ImportObject in Object.ImportObjects)
+			{
+				BlobInfo? ImportBlob;
+				if (!BlobCache.TryGetValue(ImportObject.Object.Hash, out ImportBlob))
+				{
+					ImportBlob = new BlobInfo(ImportObject.Object.Hash, ImportObject.TotalCost);
+					BlobCache[ImportObject.Object.Hash] = ImportBlob;
+				}
+				foreach (BundleImport Import in ImportObject.Imports)
+				{
+					NodeInfo? Node;
+					if (!ImportBlob.Nodes.TryGetValue(Import.Hash, out Node))
+					{
+						Node = new NodeInfo(Import.Hash, Import.Rank, Import.Cost, ImportBlob);
+						ImportBlob.Nodes.Add(Import.Hash, Node);
+					}
+					Blob.References.Add(Node);
+				}
+			}
+
+			return Object;
+		}
+
+		/// <summary>
+		/// Creates a YAML summary of the given bundle tree
+		/// </summary>
+		/// <param name="StorageClient"></param>
+		/// <param name="NamespaceId"></param>
+		/// <param name="BucketId"></param>
+		/// <param name="RefId"></param>
+		/// <param name="OutputDir"></param>
+		public static async Task WriteSummaryAsync(IStorageClient StorageClient, NamespaceId NamespaceId, BucketId BucketId, RefId RefId, DirectoryReference OutputDir)
+		{
+			IRef Ref = await StorageClient.GetRefAsync(NamespaceId, BucketId, RefId);
+			BundleObject RootObject = CbSerializer.Deserialize<BundleObject>(Ref.Value);
+
+			FileReference RootFile = FileReference.Combine(OutputDir, "root.yml");
+			WriteSummaryForObject(RootObject, RootFile);
+
+			await WriteSummaryForImportsAsync(StorageClient, NamespaceId, RootObject.ImportObjects, new HashSet<IoHash>(), OutputDir);
+		}
+
+		static async Task WriteSummaryForImportsAsync(IStorageClient StorageClient, NamespaceId NamespaceId, IEnumerable<BundleImportObject> ImportObjects, HashSet<IoHash> Hashes, DirectoryReference OutputDir)
+		{
+			foreach (BundleImportObject ImportObject in ImportObjects)
+			{
+				IoHash Hash = ImportObject.Object.Hash;
+				if (Hashes.Add(Hash))
+				{
+					BundleObject Object = await StorageClient.ReadObjectAsync<BundleObject>(NamespaceId, Hash);
+
+					FileReference OutputFile = FileReference.Combine(OutputDir, $"{Hash}.yml");
+					WriteSummaryForObject(Object, OutputFile);
+
+					await WriteSummaryForImportsAsync(StorageClient, NamespaceId, Object.ImportObjects, Hashes, OutputDir);
+				}
+			}
+		}
+
+		/// <summary>
+		/// Writes a summary of the given bundle object to disk in YAML format
+		/// </summary>
+		/// <param name="Object"></param>
+		/// <param name="File"></param>
+		public static void WriteSummaryForObject(BundleObject Object, FileReference File)
+		{
+			DirectoryReference.CreateDirectory(File.Directory);
+
+			using StreamWriter Writer = new StreamWriter(FileReference.Open(File, FileMode.Create));
+
+			Writer.WriteLine("schema: {0}", Object.Schema);
+			Writer.WriteLine("time: \"{0}\"", Object.CreationTimeUtc.ToString("u"));
+
+			List<IoHash> References = new List<IoHash>(Object.Exports.Select(x => x.Hash));
+			if (Object.ImportObjects.Count > 0)
+			{
+				Writer.WriteLine("import-objects:");
+				foreach (BundleImportObject ImportObject in Object.ImportObjects)
+				{
+					Writer.WriteLine($"- blob: \"{ImportObject.Object.Hash}\"");
+					Writer.WriteLine($"  imports:");
+					foreach (BundleImport Import in ImportObject.Imports)
+					{
+						Writer.WriteLine($"  - hash: \"{Import.Hash}\"");
+						Writer.WriteLine($"    rank: \"{Import.Rank}\"");
+						Writer.WriteLine($"    cost: \"{Import.Cost}\"");
+						Writer.WriteLine("");
+						References.Add(Import.Hash);
+					}
+					Writer.WriteLine();
+				}
+			}
+
+			if (Object.Exports.Count > 0)
+			{
+				Writer.WriteLine("packets:");
+				foreach (IGrouping<int, BundleCompressionPacket> Group in Object.Exports.Select(x => x.Packet).GroupBy(x => x.Offset).OrderBy(x => x.Key))
+				{
+					BundleCompressionPacket Packet = Group.First();
+					Writer.WriteLine($"- offset: {Packet.Offset}");
+					Writer.WriteLine($"  encodedLength: {Packet.EncodedLength}");
+					Writer.WriteLine($"  decodedLength: {Packet.DecodedLength}");
+					Writer.WriteLine($"  compression-ratio: \"{(Packet.EncodedLength * 100) / Packet.DecodedLength}%\"");
+					Writer.WriteLine($"  exports: {Group.Count()}");
+					Writer.WriteLine();
+				}
+
+				Writer.WriteLine("exports:");
+				foreach (BundleExport Export in Object.Exports)
+				{
+					Writer.WriteLine($"- hash: {Export.Hash}");
+					Writer.WriteLine($"  packet-offset: {Export.Packet.Offset}");
+					Writer.WriteLine($"  size: {Export.Length}");
+					if (Export.References.Length > 0)
+					{
+						Writer.WriteLine($"  references:");
+						foreach (int Reference in Export.References)
+						{
+							IoHash ReferenceHash = References[Reference];
+							Writer.WriteLine($"  - \"{ReferenceHash}\"");
+						}
+					}
+					Writer.WriteLine();
+				}
+			}
+		}
+	}
+}
diff --git a/Engine/Source/Programs/Shared/EpicGames.Horde/Bundles/Nodes/ChunkNode.cs b/Engine/Source/Programs/Shared/EpicGames.Horde/Bundles/Nodes/ChunkNode.cs
--- a/Engine/Source/Programs/Shared/EpicGames.Horde/Bundles/Nodes/ChunkNode.cs
+++ b/Engine/Source/Programs/Shared/EpicGames.Horde/Bundles/Nodes/ChunkNode.cs
@@ -90,8 +90,8 @@ namespace EpicGames.Horde.Bundles.Nodes
 		/// <summary>
 		/// Constructor
 		/// </summary>
-		public ChunkNode(Bundle Owner, ChunkOptions Options)
-			: this(Owner, null, 0, Options)
+		public ChunkNode(Bundle Owner, BundleNode? Parent)
+			: this(Owner, Parent, 0, new ChunkOptions())
 		{
 		}
 
@@ -202,23 +202,57 @@ namespace EpicGames.Horde.Bundles.Nodes
 		{
 			using (MemoryStream Stream = new MemoryStream())
 			{
-				await CopyToAsync(Stream);
+				await CopyToStreamAsync(Stream);
 				return Stream.ToArray();
 			}
 		}
 
+		/// <summary>
+		/// Copy data from the given stream into this chunk node
+		/// </summary>
+		/// <param name="InputStream"></param>
+		/// <param name="Options"></param>
+		/// <returns></returns>
+		public async Task CopyFromStreamAsync(Stream InputStream, ChunkOptions Options)
+		{
+			byte[] Buffer = new byte[64 * 1024];
+			for (; ; )
+			{
+				int ReadSize = await InputStream.ReadAsync(Buffer);
+				if(ReadSize == 0)
+				{
+					break;
+				}
+				Append(Buffer.AsMemory(0, ReadSize), Options);
+			}
+		}
+
+		/// <summary>
+		/// Copies the contents of this file from disk
+		/// </summary>
+		/// <param name="File"></param>
+		/// <param name="Options"></param>
+		/// <returns></returns>
+		public async Task CopyFromFileAsync(FileInfo File, ChunkOptions Options)
+		{
+			using (FileStream Stream = File.OpenRead())
+			{
+				await CopyFromStreamAsync(Stream, Options);
+			}
+		}
+
 		/// <summary>
 		/// Copies the contents of this node and its children to the given output stream
 		/// </summary>
 		/// <param name="OutputStream">The output stream to receive the data</param>
-		public async Task CopyToAsync(Stream OutputStream)
+		public async Task CopyToStreamAsync(Stream OutputStream)
 		{
 			if (Depth > 0)
 			{
 				for (int Idx = 0; Idx < GetChildCount(); Idx++)
 				{
 					ChunkNode Node = await GetChildNodeAsync(Idx);
-					await Node.CopyToAsync(OutputStream);
+					await Node.CopyToStreamAsync(OutputStream);
 				}
 			}
 			else
@@ -227,6 +261,19 @@ namespace EpicGames.Horde.Bundles.Nodes
 			}
 		}
 
+		/// <summary>
+		/// Extracts the contents of this node to a file
+		/// </summary>
+		/// <param name="File">File to write with the contents of this node</param>
+		/// <returns></returns>
+		public async Task CopyToFileAsync(FileInfo File)
+		{
+			using (FileStream Stream = File.OpenWrite())
+			{
+				await CopyToStreamAsync(Stream);
+			}
+		}
+
 		/// <summary>
 		/// Append data to this chunk. Must only be called on the root node in a chunk tree.
 		/// </summary>
@@ -390,7 +437,7 @@ namespace EpicGames.Horde.Bundles.Nodes
 				ChunkNode? LastNode = ChildNodes[^1];
 				if (LastNode != null && LastNode.WriteBuffer != null)
 				{
-					Span<byte> LastHashSpan = WriteBuffer.AsSpan(Payload.Length - IoHash.NumBytes);
+					Span<byte> LastHashSpan = WriteBuffer.AsSpan(Data.Length - IoHash.NumBytes);
 					IoHash LastHash = LastNode.Serialize();
 					LastHash.CopyTo(LastHashSpan);
 				}
@@ -420,7 +467,7 @@ namespace EpicGames.Horde.Bundles.Nodes
 		/// <inheritdoc/>
 		public override ChunkNode CreateRoot(Bundle Bundle)
 		{
-			return new ChunkNode(Bundle, new ChunkOptions());
+			return new ChunkNode(Bundle, null);
 		}
 
 		/// <inheritdoc/>
diff --git a/Engine/Source/Programs/Shared/EpicGames.Horde/Bundles/Nodes/DirectoryNode.cs b/Engine/Source/Programs/Shared/EpicGames.Horde/Bundles/Nodes/DirectoryNode.cs
--- a/Engine/Source/Programs/Shared/EpicGames.Horde/Bundles/Nodes/DirectoryNode.cs
+++ b/Engine/Source/Programs/Shared/EpicGames.Horde/Bundles/Nodes/DirectoryNode.cs
@@ -2,9 +2,11 @@
 
 using EpicGames.Core;
 using EpicGames.Serialization;
+using Microsoft.Extensions.Logging;
 using System;
 using System.Collections.Generic;
 using System.Diagnostics.CodeAnalysis;
+using System.IO;
 using System.Linq;
 using System.Text;
 using System.Threading.Tasks;
@@ -196,8 +198,82 @@ namespace EpicGames.Horde.Bundles.Nodes
 		/// Adds a new directory with the given name
 		/// </summary>
 		/// <param name="Name">Name of the new directory</param>
+		/// <param name="Flags">Flags for the new file</param>
 		/// <returns>The new directory object</returns>
-		public ValueTask<DirectoryNode> AddDirectoryAsync(Utf8String Name)
+		public ChunkNode CreateFile(Utf8String Name, DirectoryEntryFlags Flags)
+		{
+			if((Flags & DirectoryEntryFlags.Directory) != 0)
+			{
+				throw new ArgumentException("Directory flag cannot be specified for new file");
+			}
+
+			ChunkNode NewNode = new ChunkNode(Owner, this);
+
+			DirectoryEntry Entry = new DirectoryEntry(this, Name, Flags | DirectoryEntryFlags.File, NewNode);
+			NameToEntry[Name] = Entry;
+			MarkAsDirty();
+
+			return NewNode;
+		}
+
+		/// <summary>
+		/// Finds or adds a file with the given path
+		/// </summary>
+		/// <param name="Path">Path to the file</param>
+		/// <param name="Flags">Flags for the new file</param>
+		/// <returns>The new directory object</returns>
+		public async ValueTask<ChunkNode> CreateFileByPathAsync(Utf8String Path, DirectoryEntryFlags Flags)
+		{
+			DirectoryNode Directory = this;
+
+			Utf8String RemainingPath = Path;
+			for (; ; )
+			{
+				int Length = RemainingPath.IndexOf('/');
+				if (Length == -1)
+				{
+					return Directory.CreateFile(RemainingPath, Flags);
+				}
+				if (Length > 0)
+				{
+					Directory = await Directory.FindOrAddDirectoryAsync(RemainingPath.Slice(0, Length));
+				}
+				RemainingPath = RemainingPath.Slice(Length + 1);
+			}
+		}
+
+		/// <summary>
+		/// 
+		/// </summary>
+		/// <param name="Path"></param>
+		/// <returns></returns>
+		public async ValueTask DeleteFileByPathAsync(Utf8String Path)
+		{
+			DirectoryNode Directory = this;
+
+			Utf8String RemainingPath = Path;
+			for (; ; )
+			{
+				int Length = RemainingPath.IndexOf('/');
+				if (Length == -1)
+				{
+					Directory.Delete(RemainingPath);
+					return;
+				}
+				if (Length > 0)
+				{
+					Directory = await Directory.FindOrAddDirectoryAsync(RemainingPath.Slice(0, Length));
+				}
+				RemainingPath = RemainingPath.Slice(Length + 1);
+			}
+		}
+
+		/// <summary>
+		/// Adds a new directory with the given name
+		/// </summary>
+		/// <param name="Name">Name of the new directory</param>
+		/// <returns>The new directory object</returns>
+		public DirectoryNode AddDirectory(Utf8String Name)
 		{
 			DirectoryNode NewNode = new DirectoryNode(Owner, this);
 
@@ -205,7 +281,7 @@ namespace EpicGames.Horde.Bundles.Nodes
 			NameToEntry.Add(Name, Entry);
 			MarkAsDirty();
 
-			return new ValueTask<DirectoryNode>(NewNode);
+			return NewNode;
 		}
 
 		/// <summary>
@@ -233,7 +309,7 @@ namespace EpicGames.Horde.Bundles.Nodes
 		public async ValueTask<DirectoryNode> FindOrAddDirectoryAsync(Utf8String Name)
 		{
 			DirectoryNode? Node = await FindDirectoryAsync(Name);
-			return Node ?? await AddDirectoryAsync(Name);
+			return Node ?? AddDirectory(Name);
 		}
 
 		/// <summary>
@@ -311,6 +387,59 @@ namespace EpicGames.Horde.Bundles.Nodes
 			}
 			return Node;
 		}
+
+		/// <summary>
+		/// Adds files from a directory on disk
+		/// </summary>
+		/// <param name="DirectoryInfo"></param>
+		/// <param name="Options">Options for chunking file content</param>
+		/// <param name="Logger"></param>
+		/// <returns></returns>
+		public async Task CopyFromDirectoryAsync(DirectoryInfo DirectoryInfo, ChunkOptions Options, ILogger Logger)
+		{
+			foreach (DirectoryInfo SubDirectoryInfo in DirectoryInfo.EnumerateDirectories())
+			{
+				Logger.LogInformation("Adding {Directory}", SubDirectoryInfo.FullName);
+				DirectoryNode SubDirectoryNode = AddDirectory(SubDirectoryInfo.Name);
+				await SubDirectoryNode.CopyFromDirectoryAsync(SubDirectoryInfo, Options, Logger);
+			}
+			foreach (FileInfo FileInfo in DirectoryInfo.EnumerateFiles())
+			{
+				Logger.LogInformation("Adding {File}", FileInfo.FullName);
+				ChunkNode FileNode = CreateFile(FileInfo.Name, DirectoryEntryFlags.File);
+ 				await FileNode.CopyFromFileAsync(FileInfo, Options);
+			}
+		}
+
+		/// <summary>
+		/// Utility function to allow extracting a packed directory to disk
+		/// </summary>
+		/// <param name="DirectoryInfo"></param>
+		/// <param name="Logger"></param>
+		/// <returns></returns>
+		public async Task CopyToDirectoryAsync(DirectoryInfo DirectoryInfo, ILogger Logger)
+		{
+			DirectoryInfo.Create();
+
+			List<Task> Tasks = new List<Task>();
+			foreach (DirectoryEntry Entry in Entries)
+			{
+				if ((Entry.Flags & DirectoryEntryFlags.Directory) != 0)
+				{
+					DirectoryInfo SubDirectoryInfo = DirectoryInfo.CreateSubdirectory(Entry.Name.ToString());
+					DirectoryNode SubDirectoryNode = await Entry.GetDirectoryAsync();
+					Tasks.Add(Task.Run(() => SubDirectoryNode.CopyToDirectoryAsync(SubDirectoryInfo, Logger)));
+				}
+				else
+				{
+					FileInfo FileInfo = new FileInfo(Path.Combine(DirectoryInfo.FullName, Entry.Name.ToString()));
+					ChunkNode FileNode = await Entry.GetFileAsync();
+					Tasks.Add(Task.Run(() => FileNode.CopyToFileAsync(FileInfo)));
+				}
+			}
+
+			await Task.WhenAll(Tasks);
+		}
 	}
 
 	/// <summary>
diff --git a/Engine/Source/Programs/Shared/EpicGames.Horde/Storage/RefId.cs b/Engine/Source/Programs/Shared/EpicGames.Horde/Storage/RefId.cs
--- a/Engine/Source/Programs/Shared/EpicGames.Horde/Storage/RefId.cs
+++ b/Engine/Source/Programs/Shared/EpicGames.Horde/Storage/RefId.cs
@@ -88,7 +88,15 @@ namespace EpicGames.Horde.Storage
 		/// <inheritdoc/>
 		public override object ConvertFrom(ITypeDescriptorContext Context, CultureInfo Culture, object Value)
 		{
-			return new RefId(IoHash.Parse((string)Value));
+			IoHash Hash;
+			if (IoHash.TryParse((string)Value, out Hash))
+			{
+				return new RefId(Hash);
+			}
+			else
+			{
+				return new RefId((string)Value);
+			}
 		}
 	}
 
diff --git a/Engine/Source/Programs/Shared/EpicGames.Horde/Storage/TreePack.cs b/Engine/Source/Programs/Shared/EpicGames.Horde/Storage/TreePack.cs
deleted file mode 100644
diff --git a/Engine/Source/Programs/Shared/EpicGames.Horde/Storage/TreePackDirWriter.cs b/Engine/Source/Programs/Shared/EpicGames.Horde/Storage/TreePackDirWriter.cs
deleted file mode 100644
diff --git a/Engine/Source/Programs/Shared/EpicGames.Horde/Storage/TreePackFileWriter.cs b/Engine/Source/Programs/Shared/EpicGames.Horde/Storage/TreePackFileWriter.cs
deleted file mode 100644
diff --git a/Engine/Source/Programs/Shared/EpicGames.IoHash/IoHash.cs b/Engine/Source/Programs/Shared/EpicGames.IoHash/IoHash.cs
--- a/Engine/Source/Programs/Shared/EpicGames.IoHash/IoHash.cs
+++ b/Engine/Source/Programs/Shared/EpicGames.IoHash/IoHash.cs
@@ -103,6 +103,48 @@ namespace EpicGames.Core
 			return new IoHash(StringUtils.ParseHexString(Text));
 		}
 
+		/// <summary>
+		/// Parses a digest from the given hex string
+		/// </summary>
+		/// <param name="Text"></param>
+		/// <param name="Hash">Receives the hash on success</param>
+		/// <returns></returns>
+		public static bool TryParse(string Text, out IoHash Hash)
+		{
+			byte[]? Bytes;
+			if (StringUtils.TryParseHexString(Text, out Bytes) && Bytes.Length == IoHash.NumBytes)
+			{
+				Hash = new IoHash(Bytes);
+				return true;
+			}
+			else
+			{
+				Hash = default;
+				return false;
+			}
+		}
+
+		/// <summary>
+		/// Parses a digest from the given hex string
+		/// </summary>
+		/// <param name="Text"></param>
+		/// <param name="Hash">Receives the hash on success</param>
+		/// <returns></returns>
+		public static bool TryParse(ReadOnlySpan<byte> Text, out IoHash Hash)
+		{
+			byte[]? Bytes;
+			if (StringUtils.TryParseHexString(Text, out Bytes) && Bytes.Length == IoHash.NumBytes)
+			{
+				Hash = new IoHash(Bytes);
+				return true;
+			}
+			else
+			{
+				Hash = default;
+				return false;
+			}
+		}
+
 		/// <inheritdoc cref="IComparable{T}.CompareTo(T)"/>
 		public int CompareTo(IoHash Other)
 		{
