Horde: Enable mirroring of content through commit service (currently just writing out stub entries for each file).

diff --git a/Engine/Source/Programs/Horde/Horde.Build/Perforce/Impl/CommitService.cs b/Engine/Source/Programs/Horde/Horde.Build/Perforce/Impl/CommitService.cs
--- a/Engine/Source/Programs/Horde/Horde.Build/Perforce/Impl/CommitService.cs
+++ b/Engine/Source/Programs/Horde/Horde.Build/Perforce/Impl/CommitService.cs
@@ -62,7 +62,7 @@ namespace HordeServer.Commits.Impl
 		/// <summary>
 		/// Whether to mirror content to Horde storage
 		/// </summary>
-		public bool Content { get; set; }
+		public bool Content { get; set; } = true;
 
 		/// <summary>
 		/// List of streams for which content mirroring should be enabled
