Horde: Add per-pool settings for autoscaler strategies

Job queue scaler need this controlled at the pool level since the use-cases can vary quite a bit.

diff --git a/Engine/Source/Programs/Horde/Horde.Build/Api/Pools.cs b/Engine/Source/Programs/Horde/Horde.Build/Api/Pools.cs
--- a/Engine/Source/Programs/Horde/Horde.Build/Api/Pools.cs
+++ b/Engine/Source/Programs/Horde/Horde.Build/Api/Pools.cs
@@ -11,6 +11,60 @@ using Horde.Build.Fleet.Autoscale;
 
 namespace HordeServer.Api
 {
+	/// <see cref="Horde.Build.Fleet.Autoscale.LeaseUtilizationSettings" />
+	public class LeaseUtilizationSettings
+	{
+		/// <summary>
+		/// Construct a public REST API representation from the internal one
+		/// </summary>
+		/// <param name="Settings"></param>
+		public LeaseUtilizationSettings(Horde.Build.Fleet.Autoscale.LeaseUtilizationSettings Settings)
+		{
+		}
+		
+		/// <summary>
+		/// Convert this public REST API object to an internal representation
+		/// </summary>
+		/// <returns></returns>
+		public Horde.Build.Fleet.Autoscale.LeaseUtilizationSettings Convert()
+		{
+			return new Horde.Build.Fleet.Autoscale.LeaseUtilizationSettings();
+		}
+	}
+	
+	/// <see cref="Horde.Build.Fleet.Autoscale.JobQueueSettings" />
+	public class JobQueueSettings
+	{
+		/// <see cref="Horde.Build.Fleet.Autoscale.JobQueueSettings.MinQueueSizeForScaleOut" />
+		private int MinQueueSizeForScaleOut { get; set; }
+
+		/// <see cref="Horde.Build.Fleet.Autoscale.JobQueueSettings.ScaleOutFactor" />
+		private double ScaleOutFactor { get; set; }
+
+		/// <see cref="Horde.Build.Fleet.Autoscale.JobQueueSettings.ScaleInFactor" />
+		private double ScaleInFactor { get; set; }
+
+		/// <summary>
+		/// Construct a public REST API representation from the internal one
+		/// </summary>
+		/// <param name="Settings"></param>
+		public JobQueueSettings(Horde.Build.Fleet.Autoscale.JobQueueSettings Settings)
+		{
+			MinQueueSizeForScaleOut = Settings.MinQueueSizeForScaleOut;
+			ScaleInFactor = Settings.ScaleInFactor;
+			ScaleOutFactor = Settings.ScaleOutFactor;
+		}
+		
+		/// <summary>
+		/// Convert this public REST API object to an internal representation
+		/// </summary>
+		/// <returns></returns>
+		public Horde.Build.Fleet.Autoscale.JobQueueSettings Convert()
+		{
+			return new Horde.Build.Fleet.Autoscale.JobQueueSettings(MinQueueSizeForScaleOut, ScaleOutFactor, ScaleInFactor);
+		}
+	}
+	
 	/// <summary>
 	/// Parameters to create a new pool
 	/// </summary>
@@ -36,6 +90,16 @@ namespace HordeServer.Api
 		/// Pool sizing strategy
 		/// </summary>
 		public PoolSizeStrategy? SizeStrategy { get; set; }
+		
+		/// <summary>
+		/// Settings for lease utilization pool sizing strategy (if used)
+		/// </summary>
+		public LeaseUtilizationSettings? LeaseUtilizationSettings { get; set; }
+		
+		/// <summary>
+		/// Settings for job queue pool sizing strategy (if used) 
+		/// </summary>
+		public JobQueueSettings? JobQueueSettings { get; set; }
 
 		/// <summary>
 		/// The minimum nunmber of agents to retain in this pool
@@ -97,6 +161,16 @@ namespace HordeServer.Api
 		/// Pool sizing strategy
 		/// </summary>
 		public PoolSizeStrategy? SizeStrategy { get; set; }
+		
+		/// <summary>
+		/// Settings for lease utilization pool sizing strategy (if used)
+		/// </summary>
+		public LeaseUtilizationSettings? LeaseUtilizationSettings { get; set; }
+		
+		/// <summary>
+		/// Settings for job queue pool sizing strategy (if used) 
+		/// </summary>
+		public JobQueueSettings? JobQueueSettings { get; set; }
 
 		/// <summary>
 		/// The minimum nunmber of agents to retain in this pool
@@ -165,6 +239,16 @@ namespace HordeServer.Api
 		/// Pool sizing strategy to be used for this pool
 		/// </summary>
 		public PoolSizeStrategy? SizeStrategy { get; set; }
+		
+		/// <summary>
+		/// Settings for lease utilization pool sizing strategy (if used)
+		/// </summary>
+		public LeaseUtilizationSettings? LeaseUtilizationSettings { get; set; }
+		
+		/// <summary>
+		/// Settings for job queue pool sizing strategy (if used) 
+		/// </summary>
+		public JobQueueSettings? JobQueueSettings { get; set; }
 
 		/// <summary>
 		/// The minimum nunmber of agents to retain in this pool
@@ -197,6 +281,8 @@ namespace HordeServer.Api
 			this.Condition = Pool.Condition;
 			this.EnableAutoscaling = Pool.EnableAutoscaling;
 			this.SizeStrategy = Pool.SizeStrategy;
+			this.LeaseUtilizationSettings = Pool.LeaseUtilizationSettings == null ? null : new LeaseUtilizationSettings(Pool.LeaseUtilizationSettings);
+			this.JobQueueSettings = Pool.JobQueueSettings == null ? null : new JobQueueSettings(Pool.JobQueueSettings);
 			this.MinAgents = Pool.MinAgents;
 			this.NumReserveAgents = Pool.NumReserveAgents;
 			this.Workspaces = Pool.Workspaces.Select(x => new GetAgentWorkspaceResponse(x)).ToList();
diff --git a/Engine/Source/Programs/Horde/Horde.Build/Controllers/PoolsController.cs b/Engine/Source/Programs/Horde/Horde.Build/Controllers/PoolsController.cs
--- a/Engine/Source/Programs/Horde/Horde.Build/Controllers/PoolsController.cs
+++ b/Engine/Source/Programs/Horde/Horde.Build/Controllers/PoolsController.cs
@@ -60,7 +60,10 @@ namespace HordeServer.Controllers
 				return Forbid(AclAction.CreatePool);
 			}
 
-			IPool NewPool = await PoolService.CreatePoolAsync(Create.Name, Create.Condition, Create.EnableAutoscaling, Create.MinAgents, Create.NumReserveAgents, Create.SizeStrategy, Create.Properties);
+			Horde.Build.Fleet.Autoscale.LeaseUtilizationSettings? LuSettings = Create.LeaseUtilizationSettings?.Convert();
+			Horde.Build.Fleet.Autoscale.JobQueueSettings? JqSettings = Create.JobQueueSettings?.Convert();
+
+			IPool NewPool = await PoolService.CreatePoolAsync(Create.Name, Create.Condition, Create.EnableAutoscaling, Create.MinAgents, Create.NumReserveAgents, Create.SizeStrategy, LuSettings, JqSettings, Create.Properties);
 			return new CreatePoolResponse(NewPool.Id.ToString());
 		}
 
diff --git a/Engine/Source/Programs/Horde/Horde.Build/Fleet/Autoscale/JobQueueStrategy.cs b/Engine/Source/Programs/Horde/Horde.Build/Fleet/Autoscale/JobQueueStrategy.cs
--- a/Engine/Source/Programs/Horde/Horde.Build/Fleet/Autoscale/JobQueueStrategy.cs
+++ b/Engine/Source/Programs/Horde/Horde.Build/Fleet/Autoscale/JobQueueStrategy.cs
@@ -16,6 +16,46 @@ namespace Horde.Build.Fleet.Autoscale
 	using PoolId = StringId<IPool>;
 	using StreamId = StringId<IStream>;
 	
+	/// <summary>
+	/// Job queue sizing settings for a pool
+	/// </summary>
+	public class JobQueueSettings
+	{
+		/// <summary>
+		/// Minimum number of jobs in queue for scale-out logic to activate
+		///
+		/// Useful to avoid a very small queue size triggering any scaling. 
+		/// </summary>
+		public int MinQueueSizeForScaleOut { get; set; }
+
+		/// <summary>
+		/// Factor translating queue size to additional agents to grow the pool with
+		///
+		/// Example: if there are 20 jobs in queue, a factor 0.25 will result in 5 new agents being added (20 * 0.25)
+		/// </summary>
+		public double ScaleOutFactor { get; set;  } = 0.25;
+		
+		/// <summary>
+		/// Factor by which to shrink the pool size with when queue is empty
+		///
+		/// Example: when the queue size is zero, a default value of 0.9 will shrink the pool by 10% (current agent count * 0.9)
+		/// </summary>
+		public double ScaleInFactor { get; set; } = 0.9;
+
+		/// <summary>
+		/// Constructor
+		/// </summary>
+		/// <param name="MinQueueSizeForScaleOut"></param>
+		/// <param name="ScaleOutFactor"></param>
+		/// <param name="ScaleInFactor"></param>
+		public JobQueueSettings(int? MinQueueSizeForScaleOut = null, double? ScaleOutFactor = null, double? ScaleInFactor = null)
+		{
+			this.MinQueueSizeForScaleOut = MinQueueSizeForScaleOut.GetValueOrDefault(this.MinQueueSizeForScaleOut);
+			this.ScaleOutFactor = ScaleOutFactor.GetValueOrDefault(this.ScaleOutFactor);
+			this.ScaleInFactor = ScaleInFactor.GetValueOrDefault(this.ScaleInFactor);
+		}
+	}
+	
 	/// <summary>
 	/// Calculate pool size by observing the number of jobs in waiting state
 	///
@@ -28,7 +68,7 @@ namespace Horde.Build.Fleet.Autoscale
 		private readonly IGraphCollection Graphs;
 		private readonly StreamService StreamService;
 		private readonly IClock Clock;
-			
+		
 		/// <summary>
 		/// How far back in time to look for job batches (that potentially are in the queue)
 		/// </summary>
@@ -41,27 +81,8 @@ namespace Horde.Build.Fleet.Autoscale
 		/// This threshold will help ensure only batches that have been waiting longer than this value will be considered.
 		/// </summary>
 		internal readonly TimeSpan ReadyTimeThreshold = TimeSpan.FromSeconds(45.0);
-		
-		/// <summary>
-		/// Minimum number of jobs in queue for scale-out logic to activate
-		///
-		/// Useful to avoid a very small queue size triggering any scaling. 
-		/// </summary>
-		private readonly int MinQueueSizeForScaleOut = 0;
 
-		/// <summary>
-		/// Factor translating queue size to additional agents to grow the pool with
-		///
-		/// Example: if there are 20 jobs in queue, a factor 0.25 will result in 5 new agents being added (20 * 0.25)
-		/// </summary>
-		private readonly double ScaleOutFactor = 0.25;
-		
-		/// <summary>
-		/// Factor by which to shrink the pool size with when queue is empty
-		///
-		/// Example: when the queue size is zero, a default value of 0.9 will shrink the pool by 10% (current agent count * 0.9)
-		/// </summary>
-		private readonly double ScaleInFactor = 0.9;
+		private readonly JobQueueSettings DefaultPoolSettings = new();
 
 		/// <summary>
 		/// Constructor
@@ -143,10 +164,11 @@ namespace Horde.Build.Fleet.Autoscale
 
 			return Pools.Select(Current =>
 			{
+				JobQueueSettings Settings = Current.Pool.JobQueueSettings ?? DefaultPoolSettings;
 				PoolQueueSizes.TryGetValue(Current.Pool.Id, out var QueueSize);
 				if (QueueSize > 0)
 				{
-					int AdditionalAgentCount = (int)Math.Round(QueueSize * ScaleOutFactor);
+					int AdditionalAgentCount = (int)Math.Round(QueueSize * Settings.ScaleOutFactor);
 
 					int NumAgentsEnabled = Current.Agents.Count(x => x.Enabled);
 					if (NumAgentsEnabled == 0)
@@ -157,7 +179,7 @@ namespace Horde.Build.Fleet.Autoscale
 						AdditionalAgentCount = Math.Max(AdditionalAgentCount, 1);	
 					}
 					
-					if (QueueSize < MinQueueSizeForScaleOut)
+					if (QueueSize < Settings.MinQueueSizeForScaleOut)
 						AdditionalAgentCount = 0;
 
 					int DesiredAgentCount = Current.Agents.Count + AdditionalAgentCount;
@@ -165,7 +187,7 @@ namespace Horde.Build.Fleet.Autoscale
 				}
 				else
 				{
-					int DesiredAgentCount = (int)(Current.Agents.Count * ScaleInFactor);
+					int DesiredAgentCount = (int)(Current.Agents.Count * Settings.ScaleInFactor);
 					return new PoolSizeData(Current.Pool, Current.Agents, DesiredAgentCount, "Empty job queue");
 				}
 			}).ToList();
diff --git a/Engine/Source/Programs/Horde/Horde.Build/Fleet/Autoscale/LeaseUtilizationStrategy.cs b/Engine/Source/Programs/Horde/Horde.Build/Fleet/Autoscale/LeaseUtilizationStrategy.cs
--- a/Engine/Source/Programs/Horde/Horde.Build/Fleet/Autoscale/LeaseUtilizationStrategy.cs
+++ b/Engine/Source/Programs/Horde/Horde.Build/Fleet/Autoscale/LeaseUtilizationStrategy.cs
@@ -17,6 +17,13 @@ using OpenTracing.Util;
 namespace Horde.Build.Fleet.Autoscale
 {
 	using PoolId = StringId<IPool>;
+
+	/// <summary>
+	/// Lease utilization sizing settings for a pool
+	/// </summary>
+	public class LeaseUtilizationSettings
+	{
+	}
 	
 	/// <summary>
 	/// Calculate pool size by looking at previously finished leases
diff --git a/Engine/Source/Programs/Horde/Horde.Build/Fleet/IPool.cs b/Engine/Source/Programs/Horde/Horde.Build/Fleet/IPool.cs
--- a/Engine/Source/Programs/Horde/Horde.Build/Fleet/IPool.cs
+++ b/Engine/Source/Programs/Horde/Horde.Build/Fleet/IPool.cs
@@ -1,15 +1,9 @@
 // Copyright Epic Games, Inc. All Rights Reserved.
 
 using EpicGames.Horde.Common;
-using HordeServer.Api;
 using HordeServer.Utilities;
-using MongoDB.Bson;
-using MongoDB.Bson.Serialization.Attributes;
 using System;
 using System.Collections.Generic;
-using System.Linq;
-using System.Security.Permissions;
-using System.Threading.Tasks;
 using Horde.Build.Fleet.Autoscale;
 
 namespace HordeServer.Models
@@ -80,6 +74,16 @@ namespace HordeServer.Models
 		/// Pool sizing strategy to be used for this pool
 		/// </summary>
 		public PoolSizeStrategy? SizeStrategy { get; }
+		
+		/// <summary>
+		/// Settings for lease utilization pool sizing strategy (if used)
+		/// </summary>
+		public LeaseUtilizationSettings? LeaseUtilizationSettings { get; }
+		
+		/// <summary>
+		/// Settings for job queue pool sizing strategy (if used)
+		/// </summary>
+		public JobQueueSettings? JobQueueSettings { get; }
 
 		/// <summary>
 		/// Update index for this document
diff --git a/Engine/Source/Programs/Horde/Horde.Build/Fleet/IPoolCollection.cs b/Engine/Source/Programs/Horde/Horde.Build/Fleet/IPoolCollection.cs
--- a/Engine/Source/Programs/Horde/Horde.Build/Fleet/IPoolCollection.cs
+++ b/Engine/Source/Programs/Horde/Horde.Build/Fleet/IPoolCollection.cs
@@ -34,9 +34,11 @@ namespace HordeServer.Collections
 		/// <param name="MinAgents">Minimum number of agents in the pool</param>
 		/// <param name="NumReserveAgents">Minimum number of idle agents to maintain</param>
 		/// <param name="SizeStrategy">Pool sizing strategy</param>
+		/// <param name="LeaseUtilizationSettings">Settings for lease utilization strategy</param>
+		/// <param name="JobQueueSettings">Settings for job queue strategy</param>
 		/// <param name="Properties">Properties for the pool</param>
 		/// <returns>The new pool document</returns>
-		Task<IPool> AddAsync(PoolId Id, string Name, Condition? Condition = null, bool? EnableAutoscaling = null, int? MinAgents = null, int? NumReserveAgents = null, PoolSizeStrategy? SizeStrategy = null, IEnumerable<KeyValuePair<string, string>>? Properties = null);
+		Task<IPool> AddAsync(PoolId Id, string Name, Condition? Condition = null, bool? EnableAutoscaling = null, int? MinAgents = null, int? NumReserveAgents = null, PoolSizeStrategy? SizeStrategy = null, LeaseUtilizationSettings? LeaseUtilizationSettings = null, JobQueueSettings? JobQueueSettings = null, IEnumerable<KeyValuePair<string, string>>? Properties = null);
 
 		/// <summary>
 		/// Enumerates all the pools
@@ -79,7 +81,9 @@ namespace HordeServer.Collections
 		/// <param name="LastScaleUpTime">New time for last (auto) scale up</param>
 		/// <param name="LastScaleDownTime">New time for last (auto) scale down</param>
 		/// <param name="SizeStrategy">Pool sizing strategy</param>
+		/// <param name="LeaseUtilizationSettings">Settings for lease utilization strategy</param>
+		/// <param name="JobQueueSettings">Settings for job queue strategy</param>
 		/// <returns>Async task</returns>
-		Task<IPool?> TryUpdateAsync(IPool Pool, string? NewName = null, Condition? NewCondition = null, bool? NewEnableAutoscaling = null, int? NewMinAgents = null, int? NewNumReserveAgents = null, List<AgentWorkspace>? NewWorkspaces = null, bool? NewUseAutoSdk = null, Dictionary<string, string?>? NewProperties = null, DateTime? LastScaleUpTime = null, DateTime? LastScaleDownTime = null, PoolSizeStrategy? SizeStrategy = null);
+		Task<IPool?> TryUpdateAsync(IPool Pool, string? NewName = null, Condition? NewCondition = null, bool? NewEnableAutoscaling = null, int? NewMinAgents = null, int? NewNumReserveAgents = null, List<AgentWorkspace>? NewWorkspaces = null, bool? NewUseAutoSdk = null, Dictionary<string, string?>? NewProperties = null, DateTime? LastScaleUpTime = null, DateTime? LastScaleDownTime = null, PoolSizeStrategy? SizeStrategy = null, LeaseUtilizationSettings? LeaseUtilizationSettings = null, JobQueueSettings? JobQueueSettings = null);
 	}
 }
diff --git a/Engine/Source/Programs/Horde/Horde.Build/Fleet/Impl/PoolCollection.cs b/Engine/Source/Programs/Horde/Horde.Build/Fleet/Impl/PoolCollection.cs
--- a/Engine/Source/Programs/Horde/Horde.Build/Fleet/Impl/PoolCollection.cs
+++ b/Engine/Source/Programs/Horde/Horde.Build/Fleet/Impl/PoolCollection.cs
@@ -42,6 +42,8 @@ namespace HordeServer.Collections.Impl
 			public DateTime? LastScaleUpTime { get; set; }
 			public DateTime? LastScaleDownTime { get; set; }
 			public PoolSizeStrategy? SizeStrategy { get; set; }
+			public LeaseUtilizationSettings? LeaseUtilizationSettings { get; set; }
+			public JobQueueSettings? JobQueueSettings { get; set; }
 			public int UpdateIndex { get; set; }
 
 			// Read-only wrappers
@@ -66,6 +68,8 @@ namespace HordeServer.Collections.Impl
 				LastScaleUpTime = Other.LastScaleUpTime;
 				LastScaleDownTime = Other.LastScaleDownTime;
 				SizeStrategy = Other.SizeStrategy;
+				LeaseUtilizationSettings = Other.LeaseUtilizationSettings;
+				JobQueueSettings = Other.JobQueueSettings;
 				UpdateIndex = Other.UpdateIndex;
 			}
 		}
@@ -85,7 +89,7 @@ namespace HordeServer.Collections.Impl
 		}
 
 		/// <inheritdoc/>
-		public async Task<IPool> AddAsync(PoolId Id, string Name, Condition? Condition, bool? EnableAutoscaling, int? MinAgents, int? NumReserveAgents, PoolSizeStrategy? SizeStrategy, IEnumerable<KeyValuePair<string, string>>? Properties)
+		public async Task<IPool> AddAsync(PoolId Id, string Name, Condition? Condition, bool? EnableAutoscaling, int? MinAgents, int? NumReserveAgents, PoolSizeStrategy? SizeStrategy, LeaseUtilizationSettings? LeaseUtilizationSettings, JobQueueSettings? JobQueueSettings, IEnumerable<KeyValuePair<string, string>>? Properties)
 		{
 			PoolDocument Pool = new PoolDocument();
 			Pool.Id = Id;
@@ -103,6 +107,8 @@ namespace HordeServer.Collections.Impl
 			}
 
 			Pool.SizeStrategy = SizeStrategy;
+			Pool.LeaseUtilizationSettings = LeaseUtilizationSettings;
+			Pool.JobQueueSettings = JobQueueSettings;
 			await Pools.InsertOneAsync(Pool);
 			return Pool;
 		}
@@ -177,7 +183,7 @@ namespace HordeServer.Collections.Impl
 		}
 
 		/// <inheritdoc/>
-		public Task<IPool?> TryUpdateAsync(IPool Pool, string? NewName, Condition? NewCondition, bool? NewEnableAutoscaling, int? NewMinAgents, int? NewNumReserveAgents, List<AgentWorkspace>? NewWorkspaces, bool? NewUseAutoSdk, Dictionary<string, string?>? NewProperties, DateTime? LastScaleUpTime, DateTime? LastScaleDownTime, PoolSizeStrategy? SizeStrategy)
+		public Task<IPool?> TryUpdateAsync(IPool Pool, string? NewName, Condition? NewCondition, bool? NewEnableAutoscaling, int? NewMinAgents, int? NewNumReserveAgents, List<AgentWorkspace>? NewWorkspaces, bool? NewUseAutoSdk, Dictionary<string, string?>? NewProperties, DateTime? LastScaleUpTime, DateTime? LastScaleDownTime, PoolSizeStrategy? SizeStrategy, LeaseUtilizationSettings? LeaseUtilizationSettings, JobQueueSettings? JobQueueSettings)
 		{
 			TransactionBuilder<PoolDocument> Transaction = new TransactionBuilder<PoolDocument>();
 			if (NewName != null)
@@ -252,6 +258,14 @@ namespace HordeServer.Collections.Impl
 			{
 				Transaction.Set(x => x.SizeStrategy, SizeStrategy);
 			}
+			if (LeaseUtilizationSettings != null)
+			{
+				Transaction.Set(x => x.LeaseUtilizationSettings, LeaseUtilizationSettings);
+			}
+			if (JobQueueSettings != null)
+			{
+				Transaction.Set(x => x.JobQueueSettings, JobQueueSettings);
+			}
 			return TryUpdateAsync(Pool, Transaction);
 		}
 	}
diff --git a/Engine/Source/Programs/Horde/Horde.Build/Fleet/PoolService.cs b/Engine/Source/Programs/Horde/Horde.Build/Fleet/PoolService.cs
--- a/Engine/Source/Programs/Horde/Horde.Build/Fleet/PoolService.cs
+++ b/Engine/Source/Programs/Horde/Horde.Build/Fleet/PoolService.cs
@@ -2,14 +2,9 @@
 
 using EpicGames.Horde.Common;
 using HordeCommon;
-using HordeServer.Api;
 using HordeServer.Collections;
 using HordeServer.Models;
 using HordeServer.Utilities;
-using Microsoft.Extensions.Logging;
-using Microsoft.Extensions.Options;
-using MongoDB.Bson;
-using MongoDB.Driver;
 using System;
 using System.Collections.Generic;
 using System.Linq;
@@ -68,11 +63,13 @@ namespace HordeServer.Services
 		/// <param name="MinAgents">Minimum number of agents in the pool</param>
 		/// <param name="NumReserveAgents">Minimum number of idle agents to maintain</param>
 		/// <param name="SizeStrategy">Pool sizing strategy</param>
+		/// <param name="LeaseUtilizationSettings">Settings for lease utilization strategy</param>
+		/// <param name="JobQueueSettings">Settings for job queue strategy</param>
 		/// <param name="Properties">Properties for the new pool</param>
 		/// <returns>The new pool document</returns>
-		public Task<IPool> CreatePoolAsync(string Name, Condition? Condition = null, bool? EnableAutoscaling = null, int? MinAgents = null, int? NumReserveAgents = null, PoolSizeStrategy? SizeStrategy = null, Dictionary<string, string>? Properties = null)
+		public Task<IPool> CreatePoolAsync(string Name, Condition? Condition = null, bool? EnableAutoscaling = null, int? MinAgents = null, int? NumReserveAgents = null, PoolSizeStrategy? SizeStrategy = null, LeaseUtilizationSettings? LeaseUtilizationSettings = null, JobQueueSettings? JobQueueSettings = null, Dictionary<string, string>? Properties = null)
 		{
-			return Pools.AddAsync(PoolId.Sanitize(Name), Name, Condition, EnableAutoscaling, MinAgents, NumReserveAgents, SizeStrategy, Properties);
+			return Pools.AddAsync(PoolId.Sanitize(Name), Name, Condition, EnableAutoscaling, MinAgents, NumReserveAgents, SizeStrategy, LeaseUtilizationSettings, JobQueueSettings, Properties);
 		}
 
 		/// <summary>
