Horde: Fix server startup errors.

diff --git a/Engine/Source/Programs/Horde/Horde.Build/Perforce/Impl/CommitCollection.cs b/Engine/Source/Programs/Horde/Horde.Build/Perforce/Impl/CommitCollection.cs
--- a/Engine/Source/Programs/Horde/Horde.Build/Perforce/Impl/CommitCollection.cs
+++ b/Engine/Source/Programs/Horde/Horde.Build/Perforce/Impl/CommitCollection.cs
@@ -50,7 +50,7 @@ namespace HordeServer.Commits.Impl
 			[BsonElement("p")]
 			public string BasePath { get; set; } = String.Empty;
 
-			[BsonElement("r")]
+			[BsonElement("rid")]
 			public RefId? TreeRefId { get; set; }
 
 			[BsonElement("t")]
diff --git a/Engine/Source/Programs/Horde/Horde.Build/Program.cs b/Engine/Source/Programs/Horde/Horde.Build/Program.cs
--- a/Engine/Source/Programs/Horde/Horde.Build/Program.cs
+++ b/Engine/Source/Programs/Horde/Horde.Build/Program.cs
@@ -195,6 +195,7 @@ namespace HordeServer
 			Services.AddCommandsFromAssembly(Assembly.GetExecutingAssembly());
 			Services.AddLogging(Builder => Builder.AddSerilog());
 			Services.AddSingleton<IConfiguration>(Config);
+			Services.AddSingleton<ServerSettings>(HordeSettings);
 
 #pragma warning disable ASP0000 // Do not call 'IServiceCollection.BuildServiceProvider' in 'ConfigureServices'
 			IServiceProvider ServiceProvider = Services.BuildServiceProvider();
diff --git a/Engine/Source/Programs/Horde/Horde.Build/Startup.cs b/Engine/Source/Programs/Horde/Horde.Build/Startup.cs
--- a/Engine/Source/Programs/Horde/Horde.Build/Startup.cs
+++ b/Engine/Source/Programs/Horde/Horde.Build/Startup.cs
@@ -91,6 +91,7 @@ using EpicGames.Horde.Storage;
 using System.Net.Http;
 using EpicGames.AspNet;
 using EpicGames.Horde.Storage.Impl;
+using MongoDB.Bson.Serialization.Serializers;
 
 namespace HordeServer
 {
@@ -780,30 +781,46 @@ namespace HordeServer
 				return Storage;
 			});
 		}
-/*
-		private static void ConfigureLogFileWriteCache(IServiceCollection Services, ServerSettings Settings)
-		{
-			bool RedisConfigured = !String.IsNullOrEmpty(Settings.RedisConnectionConfig);
-			string CacheType = Settings.LogServiceWriteCacheType.ToLower(CultureInfo.CurrentCulture);
+		/*
+				private static void ConfigureLogFileWriteCache(IServiceCollection Services, ServerSettings Settings)
+				{
+					bool RedisConfigured = !String.IsNullOrEmpty(Settings.RedisConnectionConfig);
+					string CacheType = Settings.LogServiceWriteCacheType.ToLower(CultureInfo.CurrentCulture);
 
-			if (CacheType == "inmemory")
-			{
-				Services.AddSingleton<ILogFileWriteCache2>(Sp => new InMemoryLogFileWriteCache2());
-			}
-			else if (CacheType == "redis" && RedisConfigured)
-			{
-				Services.AddSingleton<ILogFileWriteCache2>(Sp => new RedisLogFileWriteCache2(Sp.GetService<ILogger<RedisLogFileWriteCache>>(), Sp.GetService<IDatabase>()));
-			}
-			else if (CacheType == "redis" && !RedisConfigured)
+					if (CacheType == "inmemory")
+					{
+						Services.AddSingleton<ILogFileWriteCache2>(Sp => new InMemoryLogFileWriteCache2());
+					}
+					else if (CacheType == "redis" && RedisConfigured)
+					{
+						Services.AddSingleton<ILogFileWriteCache2>(Sp => new RedisLogFileWriteCache2(Sp.GetService<ILogger<RedisLogFileWriteCache>>(), Sp.GetService<IDatabase>()));
+					}
+					else if (CacheType == "redis" && !RedisConfigured)
+					{
+						throw new Exception("Redis must be configured to use the Redis-backed log write cache");
+					}
+					else
+					{
+						throw new Exception("Unknown value set for LogServiceWriteCacheType in config: " + Settings.LogServiceWriteCacheType);
+					}
+				}
+		*/
+
+		public sealed class RefIdBsonSerializer : SerializerBase<RefId>
+		{
+			/// <inheritdoc/>
+			public override RefId Deserialize(BsonDeserializationContext Context, BsonDeserializationArgs Args)
 			{
-				throw new Exception("Redis must be configured to use the Redis-backed log write cache");
+				return new RefId(IoHash.Parse(Context.Reader.ReadString()));
 			}
-			else
+
+			/// <inheritdoc/>
+			public override void Serialize(BsonSerializationContext Context, BsonSerializationArgs Args, RefId Value)
 			{
-				throw new Exception("Unknown value set for LogServiceWriteCacheType in config: " + Settings.LogServiceWriteCacheType);
+				Context.Writer.WriteString(Value.ToString());
 			}
 		}
-*/
+
 		public static void ConfigureMongoDbClient()
 		{
 			// Ignore extra elements on deserialized documents
@@ -813,6 +830,7 @@ namespace HordeServer
 			ConventionRegistry.Register("Horde", ConventionPack, Type => true);
 
 			// Register the custom serializers
+			BsonSerializer.RegisterSerializer(new RefIdBsonSerializer());
 			BsonSerializer.RegisterSerializationProvider(new BsonSerializationProvider());
 			BsonSerializer.RegisterSerializationProvider(new StringIdSerializationProvider());
 			BsonSerializer.RegisterSerializationProvider(new ObjectIdSerializationProvider());
