Horde: Fix non deterministic sort order causing test failures.

diff --git a/Engine/Source/Programs/Horde/Horde.Build.Tests/LeasesControllerDbTest.cs b/Engine/Source/Programs/Horde/Horde.Build.Tests/LeasesControllerDbTest.cs
--- a/Engine/Source/Programs/Horde/Horde.Build.Tests/LeasesControllerDbTest.cs
+++ b/Engine/Source/Programs/Horde/Horde.Build.Tests/LeasesControllerDbTest.cs
@@ -26,7 +26,7 @@ namespace Horde.Build.Tests
 			DateTimeOffset MaxTime = Clock.UtcNow;
 
 			ILease Lease1 = await CreateLease(Clock.UtcNow - TimeSpan.FromMinutes(10), TimeSpan.FromMinutes(6));
-			ILease Lease2 = await CreateLease(Clock.UtcNow - TimeSpan.FromMinutes(7), TimeSpan.FromMinutes(3));
+			ILease Lease2 = await CreateLease(Clock.UtcNow - TimeSpan.FromMinutes(7), TimeSpan.FromMinutes(3.1));
 			ILease OutOfTimeWindow = await CreateLease(Clock.UtcNow - TimeSpan.FromMinutes(7), TimeSpan.FromMinutes(25));
 			
 			ActionResult<List<object>> Res = await LeasesController.FindLeasesAsync(null, null, null, null, MinTime, MaxTime);
