Horde: Prevent MongoDB serializers being registered more than once, causing test failures.

diff --git a/Engine/Source/Programs/Horde/Horde.Build/Startup.cs b/Engine/Source/Programs/Horde/Horde.Build/Startup.cs
--- a/Engine/Source/Programs/Horde/Horde.Build/Startup.cs
+++ b/Engine/Source/Programs/Horde/Horde.Build/Startup.cs
@@ -821,19 +821,24 @@ namespace HordeServer
 			}
 		}
 
+		static int HaveConfiguredMongoDb = 0;
+
 		public static void ConfigureMongoDbClient()
 		{
-			// Ignore extra elements on deserialized documents
-			ConventionPack ConventionPack = new ConventionPack();
-			ConventionPack.Add(new IgnoreExtraElementsConvention(true));
-			ConventionPack.Add(new EnumRepresentationConvention(BsonType.String));
-			ConventionRegistry.Register("Horde", ConventionPack, Type => true);
-
-			// Register the custom serializers
-			BsonSerializer.RegisterSerializer(new RefIdBsonSerializer());
-			BsonSerializer.RegisterSerializationProvider(new BsonSerializationProvider());
-			BsonSerializer.RegisterSerializationProvider(new StringIdSerializationProvider());
-			BsonSerializer.RegisterSerializationProvider(new ObjectIdSerializationProvider());
+			if (Interlocked.CompareExchange(ref HaveConfiguredMongoDb, 1, 0) == 0)
+			{
+				// Ignore extra elements on deserialized documents
+				ConventionPack ConventionPack = new ConventionPack();
+				ConventionPack.Add(new IgnoreExtraElementsConvention(true));
+				ConventionPack.Add(new EnumRepresentationConvention(BsonType.String));
+				ConventionRegistry.Register("Horde", ConventionPack, Type => true);
+
+				// Register the custom serializers
+				BsonSerializer.RegisterSerializer(new RefIdBsonSerializer());
+				BsonSerializer.RegisterSerializationProvider(new BsonSerializationProvider());
+				BsonSerializer.RegisterSerializationProvider(new StringIdSerializationProvider());
+				BsonSerializer.RegisterSerializationProvider(new ObjectIdSerializationProvider());
+			}
 		}
 
 		private static void OnAddHealthChecks(IServiceCollection Services)
