Horde: Log error when an exception occurs during a p4 stream update.

diff --git a/Engine/Source/Programs/Horde/Horde.Build/Perforce/Impl/CommitService.cs b/Engine/Source/Programs/Horde/Horde.Build/Perforce/Impl/CommitService.cs
--- a/Engine/Source/Programs/Horde/Horde.Build/Perforce/Impl/CommitService.cs
+++ b/Engine/Source/Programs/Horde/Horde.Build/Perforce/Impl/CommitService.cs
@@ -775,7 +775,12 @@ namespace HordeServer.Commits.Impl
 						Logger.LogInformation("Extending reservation for tree update of {StreamId} (elapsed: {Time}s)", StreamId, (int)Timer.Elapsed.TotalSeconds);
 					}
 				}
-				await InternalTask;
+
+				// Log any error during the update
+				if (InternalTask.Exception != null)
+				{
+					Logger.LogError(InternalTask.Exception.InnerException, "Exception while updating stream content for {StreamId}", StreamId);
+				}
 
 				// Remove this stream from the dirty list if it's empty
 				ITransaction Transaction = Redis.CreateTransaction();
