Horde: Catch another location where background tasks may throw an exception.

diff --git a/Engine/Source/Programs/Horde/Horde.Build/Perforce/Impl/CommitService.cs b/Engine/Source/Programs/Horde/Horde.Build/Perforce/Impl/CommitService.cs
--- a/Engine/Source/Programs/Horde/Horde.Build/Perforce/Impl/CommitService.cs
+++ b/Engine/Source/Programs/Horde/Horde.Build/Perforce/Impl/CommitService.cs
@@ -688,7 +688,10 @@ namespace HordeServer.Commits.Impl
 				(_, Task BackgroundTask) = BackgroundTasks[Idx];
 				if (BackgroundTask.IsCompleted)
 				{
-					await BackgroundTask;
+					if (BackgroundTask.Exception != null)
+					{
+						Logger.LogError(BackgroundTask.Exception, "Update background task faulted");
+					}
 					BackgroundTasks.RemoveAt(Idx);
 				}
 			}
@@ -779,7 +782,7 @@ namespace HordeServer.Commits.Impl
 				// Log any error during the update
 				if (InternalTask.Exception != null)
 				{
-					Logger.LogError(InternalTask.Exception.InnerException, "Exception while updating stream content for {StreamId}", StreamId);
+					Logger.LogError(InternalTask.Exception, "Exception while updating stream content for {StreamId}", StreamId);
 				}
 
 				// Remove this stream from the dirty list if it's empty
