Horde: Additional logging for tree updates.

diff --git a/Engine/Source/Programs/Horde/Horde.Build/Perforce/Impl/CommitService.cs b/Engine/Source/Programs/Horde/Horde.Build/Perforce/Impl/CommitService.cs
--- a/Engine/Source/Programs/Horde/Horde.Build/Perforce/Impl/CommitService.cs
+++ b/Engine/Source/Programs/Horde/Horde.Build/Perforce/Impl/CommitService.cs
@@ -780,10 +780,14 @@ namespace HordeServer.Commits.Impl
 				}
 
 				// Log any error during the update
-				if (InternalTask.Exception != null)
+				if (InternalTask.IsFaulted)
 				{
 					Logger.LogError(InternalTask.Exception, "Exception while updating stream content for {StreamId}", StreamId);
 				}
+				else
+				{
+					Logger.LogInformation("Finished update for {StreamId}", StreamId);
+				}
 
 				// Remove this stream from the dirty list if it's empty
 				ITransaction Transaction = Redis.CreateTransaction();
