Horde Storage - Fixed blob cleanup tests (diffing on expected int type)

diff --git a/Engine/Source/Programs/Horde/HordeStorage/Horde.Storage/Tests/Functional/GC/GCBlobTests.cs b/Engine/Source/Programs/Horde/HordeStorage/Horde.Storage/Tests/Functional/GC/GCBlobTests.cs
--- a/Engine/Source/Programs/Horde/HordeStorage/Horde.Storage/Tests/Functional/GC/GCBlobTests.cs
+++ b/Engine/Source/Programs/Horde/HordeStorage/Horde.Storage/Tests/Functional/GC/GCBlobTests.cs
@@ -181,7 +181,7 @@ namespace Horde.Storage.FunctionalTests.GC
 
             CancellationTokenSource cts = new CancellationTokenSource();
             ulong removedBlobsCount = await cleanup.Cleanup(cts.Token);
-            Assert.AreEqual(4, removedBlobsCount);
+            Assert.AreEqual(4ul, removedBlobsCount);
             
             _callistoBlobMock.Verify();
             _callistoBlobMock.VerifyNoOtherCalls();
