Horde Storage - Added PeerStatus service, which regularly pings other Horde Storage instances to track which instances we have better latency to. Will be used by inline-replication to pick which peer to replicate from.

diff --git a/Engine/Source/Programs/Horde/HordeStorage/Horde.Storage/Controllers/StatusController.cs b/Engine/Source/Programs/Horde/HordeStorage/Horde.Storage/Controllers/StatusController.cs
--- a/Engine/Source/Programs/Horde/HordeStorage/Horde.Storage/Controllers/StatusController.cs
+++ b/Engine/Source/Programs/Horde/HordeStorage/Horde.Storage/Controllers/StatusController.cs
@@ -19,12 +19,14 @@ namespace Horde.Storage.Controllers
         private readonly VersionFile _versionFile;
         private readonly IOptionsMonitor<JupiterSettings> _jupiterSettings;
         private readonly IOptionsMonitor<ClusterSettings> _clusterSettings;
+        private readonly IPeerStatusService _statusService;
 
-        public StatusController(VersionFile versionFile, IOptionsMonitor<JupiterSettings> jupiterSettings, IOptionsMonitor<ClusterSettings> clusterSettings)
+        public StatusController(VersionFile versionFile, IOptionsMonitor<JupiterSettings> jupiterSettings, IOptionsMonitor<ClusterSettings> clusterSettings, IPeerStatusService statusService)
         {
             _versionFile = versionFile;
             _jupiterSettings = jupiterSettings;
             _clusterSettings = clusterSettings;
+            _statusService = statusService;
         }
 
         /// <summary>
@@ -73,7 +75,7 @@ namespace Horde.Storage.Controllers
         [ProducesResponseType(type: typeof(PeersResponse), 200)]
         public IActionResult Peers([FromQuery] bool includeInternalEndpoints = false)
         {
-            return Ok(new PeersResponse(_jupiterSettings, _clusterSettings, includeInternalEndpoints));
+            return Ok(new PeersResponse(_jupiterSettings, _clusterSettings, includeInternalEndpoints, _statusService));
         }
     }
 
@@ -86,10 +88,10 @@ namespace Horde.Storage.Controllers
             Peers = peers;
         }
 
-        public PeersResponse(IOptionsMonitor<JupiterSettings> jupiterSettings, IOptionsMonitor<ClusterSettings> clusterSettings, bool includeInternalEndpoints)
+        public PeersResponse(IOptionsMonitor<JupiterSettings> jupiterSettings, IOptionsMonitor<ClusterSettings> clusterSettings, bool includeInternalEndpoints, IPeerStatusService peerStatusService)
         {
             CurrentSite = jupiterSettings.CurrentValue.CurrentSite;
-            Peers = clusterSettings.CurrentValue.Peers.Select(settings => new KnownPeer(settings, includeInternalEndpoints)).ToList();
+            Peers = clusterSettings.CurrentValue.Peers.Select(settings => new KnownPeer(settings, includeInternalEndpoints, peerStatusService)).ToList();
         }
 
         public string CurrentSite { get; set; }
@@ -100,14 +102,15 @@ namespace Horde.Storage.Controllers
     public class KnownPeer
     {
         [JsonConstructor]
-        public KnownPeer(string site, string fullName, List<string> endpoints)
+        public KnownPeer(string site, string fullName, List<string> endpoints, int latency)
         {
             Site = site;
             FullName = fullName;
             Endpoints = endpoints;
+            Latency = latency;
         }
 
-        public KnownPeer(PeerSettings peerSettings, bool includeInternalEndpoints)
+        public KnownPeer(PeerSettings peerSettings, bool includeInternalEndpoints, IPeerStatusService statusService)
         {
             Site = peerSettings.Name;
             FullName = peerSettings.FullName;
@@ -115,11 +118,15 @@ namespace Horde.Storage.Controllers
             if (!includeInternalEndpoints)
                 endpoints = endpoints.Where(s => !s.IsInternal);
             Endpoints = endpoints.Select(e => e.Url).ToList();
+            IPeerStatusService.PeerStatus peerStatus = statusService.GetPeerStatus(peerSettings.Name);
+            Latency = peerStatus.Latency;
         }
 
         public string Site { get; set; }
         public string FullName { get; set; }
 
+        public int Latency { get; set; }
+
         public List<string> Endpoints { get; set; }
     }
 
diff --git a/Engine/Source/Programs/Horde/HordeStorage/Horde.Storage/Horde.Storage.Startup.cs b/Engine/Source/Programs/Horde/HordeStorage/Horde.Storage/Horde.Storage.Startup.cs
--- a/Engine/Source/Programs/Horde/HordeStorage/Horde.Storage/Horde.Storage.Startup.cs
+++ b/Engine/Source/Programs/Horde/HordeStorage/Horde.Storage/Horde.Storage.Startup.cs
@@ -251,6 +251,9 @@ namespace Horde.Storage
 
             services.AddSingleton<ConsistencyCheckService>();
             services.AddHostedService<ConsistencyCheckService>(p => p.GetService<ConsistencyCheckService>()!);
+
+            services.AddSingleton(typeof(IPeerStatusService), typeof(PeerStatusService));
+            services.AddHostedService<PeerStatusService>(p => (PeerStatusService)p.GetService<IPeerStatusService>()!);
         
             services.AddTransient(typeof(IRefCleanup), typeof(RefCleanup));
 
diff --git a/Engine/Source/Programs/Horde/HordeStorage/Horde.Storage/Implementation/PeerStatusService.cs b/Engine/Source/Programs/Horde/HordeStorage/Horde.Storage/Implementation/PeerStatusService.cs
new file mode 100644
--- /dev/null
+++ b/Engine/Source/Programs/Horde/HordeStorage/Horde.Storage/Implementation/PeerStatusService.cs
@@ -0,0 +1,134 @@
+// Copyright Epic Games, Inc. All Rights Reserved.
+
+using System;
+using System.Collections.Generic;
+using System.Diagnostics;
+using System.Net.Http;
+using System.Threading;
+using System.Threading.Tasks;
+using Jupiter;
+using Microsoft.Extensions.Options;
+
+namespace Horde.Storage.Implementation
+{
+    public interface IPeerStatusService
+    {
+        public class PeerStatus
+        {
+            public int Latency { get; set; }
+            public bool Reachable { get; set; }
+        }
+
+        PeerStatus GetPeerStatus(string regionName);
+
+        SortedList<int, string> GetPeersByLatency(IEnumerable<string> peerNames);
+    }
+
+    public class PeerStatusService : PollingService<PeerStatusService.PeerStatusServiceState>, IPeerStatusService
+    {
+        private readonly IOptionsMonitor<ClusterSettings> _clusterSettings;
+        private readonly IHttpClientFactory _clientFactory;
+        private readonly Dictionary<string, IPeerStatusService.PeerStatus> _peers = new Dictionary<string, IPeerStatusService.PeerStatus>();
+        private volatile bool _alreadyPolling = false;
+
+        public class PeerStatusServiceState
+        {
+        }
+
+        public IPeerStatusService.PeerStatus GetPeerStatus(string regionName)
+        {
+            if (_peers.TryGetValue(regionName, out IPeerStatusService.PeerStatus? peerStatus))
+            {
+                return peerStatus;
+            }
+
+            throw new ArgumentException($"No Peer known for region: {regionName}", nameof(regionName));
+        }
+
+        public SortedList<int, string> GetPeersByLatency(IEnumerable<string> peerNames)
+        {
+            SortedList<int, string> sortedPeers = new();
+            foreach (string peerName in peerNames)
+            {
+                IPeerStatusService.PeerStatus peerStatus = GetPeerStatus(peerName);
+                sortedPeers.Add(peerStatus.Latency, peerName);
+            }
+
+            return sortedPeers;
+        }
+
+        public PeerStatusService(IOptionsMonitor<ClusterSettings> clusterSettings, IHttpClientFactory clientFactory) : base("PeerStatus", TimeSpan.FromMinutes(15), new PeerStatusServiceState())
+        {
+            _clusterSettings = clusterSettings;
+            _clientFactory = clientFactory;
+
+            foreach (PeerSettings peerSettings in clusterSettings.CurrentValue.Peers)
+            {
+                _peers[peerSettings.Name] = new IPeerStatusService.PeerStatus
+                {
+                    Latency = int.MaxValue
+                };
+            }
+        }
+
+        public override async Task<bool> OnPoll(PeerStatusServiceState state, CancellationToken cancellationToken)
+        {
+            if (_alreadyPolling)
+                return false;
+
+            _alreadyPolling = true;
+
+            await UpdatePeerStatus(cancellationToken);
+            return true;
+        }
+
+        public async Task UpdatePeerStatus(CancellationToken cancellationToken)
+        {
+            foreach (PeerSettings peerSettings in _clusterSettings.CurrentValue.Peers)
+            {
+                IPeerStatusService.PeerStatus status = _peers[peerSettings.Name];
+
+                int bestLatency = int.MaxValue;
+                bool reachable = false;
+
+                await Parallel.ForEachAsync(peerSettings.Endpoints, cancellationToken, async (endpoint, token) =>
+                {
+                    int latency = await MeasureLatency(endpoint);
+                    bestLatency = Math.Min(latency, bestLatency);
+
+                    if (latency != int.MaxValue)
+                       reachable = true;
+                });
+
+                status.Reachable = reachable;
+                status.Latency = bestLatency;
+            }
+        }
+
+        private async Task<int> MeasureLatency(PeerEndpoints endpoint)
+        {
+            Stopwatch stopwatch = Stopwatch.StartNew();
+            using HttpClient client = _clientFactory.CreateClient();
+
+            // treat any connection that takes more then 5 seconds to establish as timed out
+            client.Timeout = TimeSpan.FromSeconds(5);
+
+            Uri uri = new Uri(endpoint.Url + "/health/live");
+            try
+            {
+                HttpResponseMessage result = await client.GetAsync(uri);
+                // ignore error responses as they may not have reached the actual instance
+                if (!result.IsSuccessStatusCode)
+                    return int.MaxValue;
+
+                return (int)stopwatch.ElapsedMilliseconds;
+            }
+            catch (Exception)
+            {
+                // error reaching the endpoint is just considered to max latency
+                return int.MaxValue;
+            }
+
+        }
+    }
+}
diff --git a/Engine/Source/Programs/Horde/HordeStorage/Horde.Storage/Tests/Functional/PeerStatusTests.cs b/Engine/Source/Programs/Horde/HordeStorage/Horde.Storage/Tests/Functional/PeerStatusTests.cs
new file mode 100644
--- /dev/null
+++ b/Engine/Source/Programs/Horde/HordeStorage/Horde.Storage/Tests/Functional/PeerStatusTests.cs
@@ -0,0 +1,164 @@
+// Copyright Epic Games, Inc. All Rights Reserved.
+
+using System.Linq;
+using System.Net;
+using System.Net.Http;
+using System.Net.Sockets;
+using System.Threading;
+using System.Threading.Tasks;
+using Horde.Storage.Implementation;
+using Microsoft.Extensions.Options;
+using Microsoft.VisualStudio.TestTools.UnitTesting;
+using Moq;
+using Moq.Contrib.HttpClient;
+
+namespace Horde.Storage.FunctionalTests.Status
+{
+    [TestClass]
+    public class PeerStatusTests
+    {
+        [TestMethod]
+        public async Task GetPeerConnection()
+        {
+            ClusterSettings settings = new ClusterSettings
+            {
+                Peers = new PeerSettings[]
+                {
+                    new PeerSettings()
+                    {
+                        Name = "siteA",
+                        FullName = "siteA",
+                        Endpoints = new PeerEndpoints[]
+                        {
+                            new PeerEndpoints()
+                            {
+                                Url = "http://siteA.com/internal",
+                                IsInternal = true
+                            },
+                            new PeerEndpoints()
+                            {
+                                Url = "http://siteA.com/public"
+                            },
+                        }.ToList()
+                    },
+                    new PeerSettings()
+                    {
+                        Name = "siteB",
+                        FullName = "siteB",
+                        Endpoints = new PeerEndpoints[]
+                        {
+                            new PeerEndpoints()
+                            {
+                                Url = "http://siteB.com/internal",
+                                IsInternal = true
+                            },
+                            new PeerEndpoints()
+                            {
+                                Url = "http://siteB.com/public"
+                            },
+                        }.ToList()
+                    },
+
+                }.ToList()
+            };
+            IOptionsMonitor<ClusterSettings> settingsMock = Mock.Of<IOptionsMonitor<ClusterSettings>>(_ => _.CurrentValue == settings);
+
+            Mock<HttpMessageHandler> handler = new Mock<HttpMessageHandler>();
+            string endpoint = "/health/live";
+            // this emulates response times from calling the different endpoints
+            handler.SetupRequest("http://sitea.com/public" + endpoint).ReturnsResponse(HttpStatusCode.OK).Callback(() => Task.Delay(220).Wait()).Verifiable();
+            handler.SetupRequest("http://sitea.com/internal" + endpoint).ReturnsResponse(HttpStatusCode.OK).Callback(() => Task.Delay(200).Wait()).Verifiable();
+            handler.SetupRequest("http://siteb.com/public" + endpoint).ReturnsResponse(HttpStatusCode.OK).Callback(() => Task.Delay(420).Wait()).Verifiable();
+            handler.SetupRequest("http://siteb.com/internal" + endpoint).ReturnsResponse(HttpStatusCode.OK).Callback(() => Task.Delay(400).Wait()).Verifiable();
+
+            IHttpClientFactory httpClientFactory = handler.CreateClientFactory();
+            PeerStatusService statusService = new PeerStatusService(settingsMock, httpClientFactory);
+
+            await statusService.UpdatePeerStatus(CancellationToken.None);
+
+            handler.Verify();
+
+            // as we are actually measuring the time it takes to call our mocked handler the latency expected is not going to be exact so we allow some delta
+            IPeerStatusService.PeerStatus siteAPeerStatus = statusService.GetPeerStatus("siteA");
+            Assert.AreEqual(220, siteAPeerStatus.Latency, 50);
+            Assert.IsTrue(siteAPeerStatus.Reachable);
+
+            IPeerStatusService.PeerStatus siteBPeerStatus = statusService.GetPeerStatus("siteB");
+            Assert.AreEqual(420, siteBPeerStatus.Latency, 50);
+            Assert.IsTrue(siteBPeerStatus.Reachable);
+        }
+
+        [TestMethod]
+        public async Task PeerUnreachable()
+        {
+            ClusterSettings settings = new ClusterSettings
+            {
+                Peers = new PeerSettings[]
+                {
+                    new PeerSettings()
+                    {
+                        Name = "siteA",
+                        FullName = "siteA",
+                        Endpoints = new PeerEndpoints[]
+                        {
+                            new PeerEndpoints()
+                            {
+                                Url = "http://siteA.com/internal",
+                                IsInternal = true
+                            },
+                            new PeerEndpoints()
+                            {
+                                Url = "http://siteA.com/public"
+                            },
+                        }.ToList()
+                    },
+                    new PeerSettings()
+                    {
+                        Name = "siteB",
+                        FullName = "siteB",
+                        Endpoints = new PeerEndpoints[]
+                        {
+                            new PeerEndpoints()
+                            {
+                                Url = "http://siteB.com/internal",
+                                IsInternal = true
+                            },
+                            new PeerEndpoints()
+                            {
+                                Url = "http://siteB.com/public"
+                            },
+                        }.ToList()
+                    },
+                }.ToList()
+            };
+            IOptionsMonitor<ClusterSettings> settingsMock = Mock.Of<IOptionsMonitor<ClusterSettings>>(_ => _.CurrentValue == settings);
+
+            Mock<HttpMessageHandler> handler = new Mock<HttpMessageHandler>();
+            string endpoint = "/health/live";
+            handler.SetupRequest("http://sitea.com/public" + endpoint).ReturnsResponse(HttpStatusCode.OK).Callback(() => Task.Delay(220).Wait()).Verifiable();
+            handler.SetupRequest("http://sitea.com/internal" + endpoint).Throws<SocketException>().Verifiable();
+
+            // site b is not reachable at all
+            handler.SetupRequest("http://siteb.com/public" + endpoint).Throws<SocketException>().Verifiable();
+            handler.SetupRequest("http://siteb.com/internal" + endpoint).Throws<SocketException>().Verifiable();
+
+            IHttpClientFactory httpClientFactory = handler.CreateClientFactory();
+            PeerStatusService statusService = new PeerStatusService(settingsMock, httpClientFactory);
+
+            await statusService.UpdatePeerStatus(CancellationToken.None);
+
+            handler.Verify();
+
+            // as we are actually measuring the time it takes to call our mocked handler the latency expected is not going to be exact so we allow some delta
+            IPeerStatusService.PeerStatus siteAPeerStatus = statusService.GetPeerStatus("siteA");
+            
+            // verify that we ignore the failing internal endpoint
+            Assert.AreEqual(220, siteAPeerStatus.Latency, 50);
+            Assert.IsTrue(siteAPeerStatus.Reachable);
+
+            IPeerStatusService.PeerStatus siteBPeerStatus = statusService.GetPeerStatus("siteB");
+            Assert.AreEqual(int.MaxValue, siteBPeerStatus.Latency);
+            Assert.IsFalse(siteBPeerStatus.Reachable);
+        }
+    }
+}
