Horde Storage - Avoid determining the latency to the local site.

diff --git a/Engine/Source/Programs/Horde/HordeStorage/Horde.Storage/Implementation/PeerStatusService.cs b/Engine/Source/Programs/Horde/HordeStorage/Horde.Storage/Implementation/PeerStatusService.cs
--- a/Engine/Source/Programs/Horde/HordeStorage/Horde.Storage/Implementation/PeerStatusService.cs
+++ b/Engine/Source/Programs/Horde/HordeStorage/Horde.Storage/Implementation/PeerStatusService.cs
@@ -27,6 +27,7 @@ namespace Horde.Storage.Implementation
     public class PeerStatusService : PollingService<PeerStatusService.PeerStatusServiceState>, IPeerStatusService
     {
         private readonly IOptionsMonitor<ClusterSettings> _clusterSettings;
+        private readonly IOptionsMonitor<JupiterSettings> _jupiterSettings;
         private readonly IHttpClientFactory _clientFactory;
         private readonly Dictionary<string, IPeerStatusService.PeerStatus> _peers = new Dictionary<string, IPeerStatusService.PeerStatus>();
         private volatile bool _alreadyPolling = false;
@@ -50,6 +51,10 @@ namespace Horde.Storage.Implementation
             SortedList<int, string> sortedPeers = new();
             foreach (string peerName in peerNames)
             {
+                // skip the local site
+                if (string.Equals(peerName, _jupiterSettings.CurrentValue.CurrentSite))
+                    continue;
+
                 IPeerStatusService.PeerStatus peerStatus = GetPeerStatus(peerName);
                 sortedPeers.Add(peerStatus.Latency, peerName);
             }
@@ -57,13 +62,18 @@ namespace Horde.Storage.Implementation
             return sortedPeers;
         }
 
-        public PeerStatusService(IOptionsMonitor<ClusterSettings> clusterSettings, IHttpClientFactory clientFactory) : base("PeerStatus", TimeSpan.FromMinutes(15), new PeerStatusServiceState())
+        public PeerStatusService(IOptionsMonitor<ClusterSettings> clusterSettings, IOptionsMonitor<JupiterSettings> jupiterSettings, IHttpClientFactory clientFactory) : base("PeerStatus", TimeSpan.FromMinutes(15), new PeerStatusServiceState())
         {
             _clusterSettings = clusterSettings;
+            _jupiterSettings = jupiterSettings;
             _clientFactory = clientFactory;
 
             foreach (PeerSettings peerSettings in clusterSettings.CurrentValue.Peers)
             {
+                // skip the local site
+                if (string.Equals(peerSettings.Name, jupiterSettings.CurrentValue.CurrentSite))
+                    continue;
+
                 _peers[peerSettings.Name] = new IPeerStatusService.PeerStatus
                 {
                     Latency = int.MaxValue
@@ -86,6 +96,10 @@ namespace Horde.Storage.Implementation
         {
             foreach (PeerSettings peerSettings in _clusterSettings.CurrentValue.Peers)
             {
+                // skip the local site
+                if (string.Equals(peerSettings.Name, _jupiterSettings.CurrentValue.CurrentSite))
+                    continue;
+
                 IPeerStatusService.PeerStatus status = _peers[peerSettings.Name];
 
                 int bestLatency = int.MaxValue;
diff --git a/Engine/Source/Programs/Horde/HordeStorage/Horde.Storage/Tests/Functional/PeerStatusTests.cs b/Engine/Source/Programs/Horde/HordeStorage/Horde.Storage/Tests/Functional/PeerStatusTests.cs
--- a/Engine/Source/Programs/Horde/HordeStorage/Horde.Storage/Tests/Functional/PeerStatusTests.cs
+++ b/Engine/Source/Programs/Horde/HordeStorage/Horde.Storage/Tests/Functional/PeerStatusTests.cs
@@ -7,6 +7,7 @@ using System.Net.Sockets;
 using System.Threading;
 using System.Threading.Tasks;
 using Horde.Storage.Implementation;
+using Jupiter;
 using Microsoft.Extensions.Options;
 using Microsoft.VisualStudio.TestTools.UnitTesting;
 using Moq;
@@ -63,6 +64,9 @@ namespace Horde.Storage.FunctionalTests.Status
             };
             IOptionsMonitor<ClusterSettings> settingsMock = Mock.Of<IOptionsMonitor<ClusterSettings>>(_ => _.CurrentValue == settings);
 
+            JupiterSettings jupiterSettings = new JupiterSettings();
+            IOptionsMonitor<JupiterSettings> jupiterSettingsMock = Mock.Of<IOptionsMonitor<JupiterSettings>>(_ => _.CurrentValue == jupiterSettings);
+
             Mock<HttpMessageHandler> handler = new Mock<HttpMessageHandler>();
             string endpoint = "/health/live";
             // this emulates response times from calling the different endpoints
@@ -72,7 +76,7 @@ namespace Horde.Storage.FunctionalTests.Status
             handler.SetupRequest("http://siteb.com/internal" + endpoint).ReturnsResponse(HttpStatusCode.OK).Callback(() => Task.Delay(400).Wait()).Verifiable();
 
             IHttpClientFactory httpClientFactory = handler.CreateClientFactory();
-            PeerStatusService statusService = new PeerStatusService(settingsMock, httpClientFactory);
+            PeerStatusService statusService = new PeerStatusService(settingsMock, jupiterSettingsMock, httpClientFactory);
 
             await statusService.UpdatePeerStatus(CancellationToken.None);
 
@@ -133,6 +137,9 @@ namespace Horde.Storage.FunctionalTests.Status
             };
             IOptionsMonitor<ClusterSettings> settingsMock = Mock.Of<IOptionsMonitor<ClusterSettings>>(_ => _.CurrentValue == settings);
 
+            JupiterSettings jupiterSettings = new JupiterSettings();
+            IOptionsMonitor<JupiterSettings> jupiterSettingsMock = Mock.Of<IOptionsMonitor<JupiterSettings>>(_ => _.CurrentValue == jupiterSettings);
+
             Mock<HttpMessageHandler> handler = new Mock<HttpMessageHandler>();
             string endpoint = "/health/live";
             handler.SetupRequest("http://sitea.com/public" + endpoint).ReturnsResponse(HttpStatusCode.OK).Callback(() => Task.Delay(220).Wait()).Verifiable();
@@ -143,7 +150,7 @@ namespace Horde.Storage.FunctionalTests.Status
             handler.SetupRequest("http://siteb.com/internal" + endpoint).Throws<SocketException>().Verifiable();
 
             IHttpClientFactory httpClientFactory = handler.CreateClientFactory();
-            PeerStatusService statusService = new PeerStatusService(settingsMock, httpClientFactory);
+            PeerStatusService statusService = new PeerStatusService(settingsMock, jupiterSettingsMock, httpClientFactory);
 
             await statusService.UpdatePeerStatus(CancellationToken.None);
 
