Horde Storage - Ignore casing differences between cluster settings and the current site setting.

diff --git a/Engine/Source/Programs/Horde/HordeStorage/Horde.Storage/Implementation/PeerStatusService.cs b/Engine/Source/Programs/Horde/HordeStorage/Horde.Storage/Implementation/PeerStatusService.cs
--- a/Engine/Source/Programs/Horde/HordeStorage/Horde.Storage/Implementation/PeerStatusService.cs
+++ b/Engine/Source/Programs/Horde/HordeStorage/Horde.Storage/Implementation/PeerStatusService.cs
@@ -52,7 +52,7 @@ namespace Horde.Storage.Implementation
             foreach (string peerName in peerNames)
             {
                 // skip the local site
-                if (string.Equals(peerName, _jupiterSettings.CurrentValue.CurrentSite))
+                if (string.Equals(peerName, _jupiterSettings.CurrentValue.CurrentSite, StringComparison.InvariantCultureIgnoreCase))
                     continue;
 
                 IPeerStatusService.PeerStatus peerStatus = GetPeerStatus(peerName);
@@ -71,7 +71,7 @@ namespace Horde.Storage.Implementation
             foreach (PeerSettings peerSettings in clusterSettings.CurrentValue.Peers)
             {
                 // skip the local site
-                if (string.Equals(peerSettings.Name, jupiterSettings.CurrentValue.CurrentSite))
+                if (string.Equals(peerSettings.Name, jupiterSettings.CurrentValue.CurrentSite, StringComparison.InvariantCultureIgnoreCase))
                     continue;
 
                 _peers[peerSettings.Name] = new IPeerStatusService.PeerStatus
@@ -97,7 +97,7 @@ namespace Horde.Storage.Implementation
             foreach (PeerSettings peerSettings in _clusterSettings.CurrentValue.Peers)
             {
                 // skip the local site
-                if (string.Equals(peerSettings.Name, _jupiterSettings.CurrentValue.CurrentSite))
+                if (string.Equals(peerSettings.Name, _jupiterSettings.CurrentValue.CurrentSite, StringComparison.InvariantCultureIgnoreCase))
                     continue;
 
                 IPeerStatusService.PeerStatus status = _peers[peerSettings.Name];
