Horde Storage - Return null peer status if the region asked for is not known (or is the current region) instead of throwing.

diff --git a/Engine/Source/Programs/Horde/HordeStorage/Horde.Storage/Controllers/StatusController.cs b/Engine/Source/Programs/Horde/HordeStorage/Horde.Storage/Controllers/StatusController.cs
--- a/Engine/Source/Programs/Horde/HordeStorage/Horde.Storage/Controllers/StatusController.cs
+++ b/Engine/Source/Programs/Horde/HordeStorage/Horde.Storage/Controllers/StatusController.cs
@@ -118,8 +118,12 @@ namespace Horde.Storage.Controllers
             if (!includeInternalEndpoints)
                 endpoints = endpoints.Where(s => !s.IsInternal);
             Endpoints = endpoints.Select(e => e.Url).ToList();
-            IPeerStatusService.PeerStatus peerStatus = statusService.GetPeerStatus(peerSettings.Name);
-            Latency = peerStatus.Latency;
+
+            IPeerStatusService.PeerStatus? peerStatus = statusService.GetPeerStatus(peerSettings.Name);
+            if (peerStatus != null)
+            {
+                Latency = peerStatus.Latency;
+            }
         }
 
         public string Site { get; set; }
diff --git a/Engine/Source/Programs/Horde/HordeStorage/Horde.Storage/Implementation/PeerStatusService.cs b/Engine/Source/Programs/Horde/HordeStorage/Horde.Storage/Implementation/PeerStatusService.cs
--- a/Engine/Source/Programs/Horde/HordeStorage/Horde.Storage/Implementation/PeerStatusService.cs
+++ b/Engine/Source/Programs/Horde/HordeStorage/Horde.Storage/Implementation/PeerStatusService.cs
@@ -19,7 +19,7 @@ namespace Horde.Storage.Implementation
             public bool Reachable { get; set; }
         }
 
-        PeerStatus GetPeerStatus(string regionName);
+        PeerStatus? GetPeerStatus(string regionName);
 
         SortedList<int, string> GetPeersByLatency(IEnumerable<string> peerNames);
     }
@@ -36,14 +36,14 @@ namespace Horde.Storage.Implementation
         {
         }
 
-        public IPeerStatusService.PeerStatus GetPeerStatus(string regionName)
+        public IPeerStatusService.PeerStatus? GetPeerStatus(string regionName)
         {
             if (_peers.TryGetValue(regionName, out IPeerStatusService.PeerStatus? peerStatus))
             {
                 return peerStatus;
             }
 
-            throw new ArgumentException($"No Peer known for region: {regionName}", nameof(regionName));
+            return null;
         }
 
         public SortedList<int, string> GetPeersByLatency(IEnumerable<string> peerNames)
@@ -55,7 +55,9 @@ namespace Horde.Storage.Implementation
                 if (string.Equals(peerName, _jupiterSettings.CurrentValue.CurrentSite, StringComparison.InvariantCultureIgnoreCase))
                     continue;
 
-                IPeerStatusService.PeerStatus peerStatus = GetPeerStatus(peerName);
+                IPeerStatusService.PeerStatus? peerStatus = GetPeerStatus(peerName);
+                if (peerStatus == null)
+                    continue;
                 sortedPeers.Add(peerStatus.Latency, peerName);
             }
 
diff --git a/Engine/Source/Programs/Horde/HordeStorage/Horde.Storage/Tests/Functional/PeerStatusTests.cs b/Engine/Source/Programs/Horde/HordeStorage/Horde.Storage/Tests/Functional/PeerStatusTests.cs
--- a/Engine/Source/Programs/Horde/HordeStorage/Horde.Storage/Tests/Functional/PeerStatusTests.cs
+++ b/Engine/Source/Programs/Horde/HordeStorage/Horde.Storage/Tests/Functional/PeerStatusTests.cs
@@ -83,11 +83,13 @@ namespace Horde.Storage.FunctionalTests.Status
             handler.Verify();
 
             // as we are actually measuring the time it takes to call our mocked handler the latency expected is not going to be exact so we allow some delta
-            IPeerStatusService.PeerStatus siteAPeerStatus = statusService.GetPeerStatus("siteA");
+            IPeerStatusService.PeerStatus? siteAPeerStatus = statusService.GetPeerStatus("siteA");
+            Assert.IsNotNull(siteAPeerStatus);
             Assert.AreEqual(220, siteAPeerStatus.Latency, 50);
             Assert.IsTrue(siteAPeerStatus.Reachable);
 
-            IPeerStatusService.PeerStatus siteBPeerStatus = statusService.GetPeerStatus("siteB");
+            IPeerStatusService.PeerStatus? siteBPeerStatus = statusService.GetPeerStatus("siteB");
+            Assert.IsNotNull(siteBPeerStatus);
             Assert.AreEqual(420, siteBPeerStatus.Latency, 50);
             Assert.IsTrue(siteBPeerStatus.Reachable);
         }
@@ -157,13 +159,14 @@ namespace Horde.Storage.FunctionalTests.Status
             handler.Verify();
 
             // as we are actually measuring the time it takes to call our mocked handler the latency expected is not going to be exact so we allow some delta
-            IPeerStatusService.PeerStatus siteAPeerStatus = statusService.GetPeerStatus("siteA");
-            
+            IPeerStatusService.PeerStatus? siteAPeerStatus = statusService.GetPeerStatus("siteA");
+            Assert.IsNotNull(siteAPeerStatus);
             // verify that we ignore the failing internal endpoint
             Assert.AreEqual(220, siteAPeerStatus.Latency, 50);
             Assert.IsTrue(siteAPeerStatus.Reachable);
 
-            IPeerStatusService.PeerStatus siteBPeerStatus = statusService.GetPeerStatus("siteB");
+            IPeerStatusService.PeerStatus? siteBPeerStatus = statusService.GetPeerStatus("siteB");
+            Assert.IsNotNull(siteBPeerStatus);
             Assert.AreEqual(int.MaxValue, siteBPeerStatus.Latency);
             Assert.IsFalse(siteBPeerStatus.Reachable);
         }
