Horde Storage - Added more debug output for intermitently failing test

diff --git a/Engine/Source/Programs/Horde/HordeStorage/Horde.Storage/Tests/Functional/Objects/ReferencesTests.cs b/Engine/Source/Programs/Horde/HordeStorage/Horde.Storage/Tests/Functional/Objects/ReferencesTests.cs
--- a/Engine/Source/Programs/Horde/HordeStorage/Horde.Storage/Tests/Functional/Objects/ReferencesTests.cs
+++ b/Engine/Source/Programs/Horde/HordeStorage/Horde.Storage/Tests/Functional/Objects/ReferencesTests.cs
@@ -1764,13 +1764,13 @@ namespace Horde.Storage.FunctionalTests.References
                 
                 BatchOpsResponse.OpResponses op0 = response.Results.First(r => r.OpId == 0);
                 Assert.IsNotNull(op0.Response);
-                Assert.AreEqual(200, op0.StatusCode);
+                Assert.AreEqual(200, op0.StatusCode, $"Expected 200 response, got {op0.StatusCode} . Response: {op0.Response.ToJson()}");
                 Assert.IsTrue(!op0.Response["needs"].Equals(CbField.Empty));
                 CollectionAssert.AreEqual(Array.Empty<IoHash>(), op0.Response["needs"].ToArray());
 
                 BatchOpsResponse.OpResponses op1 = response.Results.First(r => r.OpId == 1);
                 Assert.IsNotNull(op1.Response);
-                Assert.AreEqual(200, op1.StatusCode);
+                Assert.AreEqual(200, op1.StatusCode, $"Expected 200 response, got {op1.StatusCode} . Response: {op1.Response.ToJson()}");
                 Assert.IsTrue(!op1.Response["needs"].Equals(CbField.Empty));
                 CollectionAssert.AreEqual(Array.Empty<IoHash>(), op1.Response["needs"].ToArray());
             }
