Horde: Disable mirroring of P4 content to Horde Storage due to memory pressure.

diff --git a/Engine/Source/Programs/Horde/Horde.Build/Perforce/Impl/CommitService.cs b/Engine/Source/Programs/Horde/Horde.Build/Perforce/Impl/CommitService.cs
--- a/Engine/Source/Programs/Horde/Horde.Build/Perforce/Impl/CommitService.cs
+++ b/Engine/Source/Programs/Horde/Horde.Build/Perforce/Impl/CommitService.cs
@@ -62,7 +62,7 @@ namespace HordeServer.Commits.Impl
 		/// <summary>
 		/// Whether to mirror content to Horde storage
 		/// </summary>
-		public bool Content { get; set; } = true;
+		public bool Content { get; set; } = false;
 
 		/// <summary>
 		/// List of streams for which content mirroring should be enabled
