Horde: Add per-pool settings for scaling cooldowns

diff --git a/Engine/Source/Programs/Horde/Horde.Build.Tests/Fleet/AutoscaleServiceV2Test.cs b/Engine/Source/Programs/Horde/Horde.Build.Tests/Fleet/AutoscaleServiceV2Test.cs
--- a/Engine/Source/Programs/Horde/Horde.Build.Tests/Fleet/AutoscaleServiceV2Test.cs
+++ b/Engine/Source/Programs/Horde/Horde.Build.Tests/Fleet/AutoscaleServiceV2Test.cs
@@ -72,11 +72,11 @@ namespace Horde.Build.Tests.Fleet
 		{
 			var Service = GetAutoscaleService(FleetManagerSpy);
 			
-			IPool Pool1 = await PoolService.CreatePoolAsync("bogusPoolLease1", null, true, 0, 0, PoolSizeStrategy.LeaseUtilization);
-			IPool Pool2 = await PoolService.CreatePoolAsync("bogusPoolLease2", null, true, 0, 0, null);
-			IPool Pool3 = await PoolService.CreatePoolAsync("bogusPoolJobQueue1", null, true, 0, 0, PoolSizeStrategy.JobQueue);
-			IPool Pool4 = await PoolService.CreatePoolAsync("bogusPoolJobQueue2", null, true, 0, 0, PoolSizeStrategy.JobQueue);
-			IPool Pool5 = await PoolService.CreatePoolAsync("bogusPoolNoOp", null, true, 0, 0, PoolSizeStrategy.NoOp);
+			IPool Pool1 = await PoolService.CreatePoolAsync("bogusPoolLease1", null, true, 0, 0, SizeStrategy: PoolSizeStrategy.LeaseUtilization);
+			IPool Pool2 = await PoolService.CreatePoolAsync("bogusPoolLease2", null, true, 0, 0, SizeStrategy: null);
+			IPool Pool3 = await PoolService.CreatePoolAsync("bogusPoolJobQueue1", null, true, 0, 0, SizeStrategy: PoolSizeStrategy.JobQueue);
+			IPool Pool4 = await PoolService.CreatePoolAsync("bogusPoolJobQueue2", null, true, 0, 0, SizeStrategy: PoolSizeStrategy.JobQueue);
+			IPool Pool5 = await PoolService.CreatePoolAsync("bogusPoolNoOp", null, true, 0, 0, SizeStrategy: PoolSizeStrategy.NoOp);
 
 			PoolSizeStrategySpy LeaseUtilizationSpy = new();
 			PoolSizeStrategySpy JobQueueSpy = new();
@@ -97,6 +97,48 @@ namespace Horde.Build.Tests.Fleet
 			Assert.IsTrue(NoOpSpy.PoolIdsSeen.Contains(Pool5.Id));
 		}
 
+		[TestMethod]
+		public async Task ScaleOutCooldown()
+		{
+			AutoscaleServiceV2 Service = GetAutoscaleService(FleetManagerSpy);
+			IPool Pool = await PoolService.CreatePoolAsync("testPool", null, true, 0, 0, SizeStrategy: PoolSizeStrategy.NoOp);
+
+			// First scale-out will succeed
+			await Service.ResizePools(new() { new PoolSizeData(Pool, new List<IAgent>(), 1, "Testing") });
+			Assert.AreEqual(1, FleetManagerSpy.ExpandPoolAsyncCallCount);
+			
+			// Cannot scale-out due to cool-down
+			await Service.ResizePools(new() { new PoolSizeData(Pool, new List<IAgent>(), 2, "Testing") });
+			Assert.AreEqual(1, FleetManagerSpy.ExpandPoolAsyncCallCount);
+
+			// Wait some time and then try again
+			await Clock.AdvanceAsync(TimeSpan.FromHours(2));
+			await Service.ResizePools(new() { new PoolSizeData(Pool, new List<IAgent>(), 2, "Testing") });
+			Assert.AreEqual(2, FleetManagerSpy.ExpandPoolAsyncCallCount);
+		}
+		
+		[TestMethod]
+		public async Task ScaleInCooldown()
+		{
+			AutoscaleServiceV2 Service = GetAutoscaleService(FleetManagerSpy);
+			IPool Pool = await PoolService.CreatePoolAsync("testPool", null, true, 0, 0, SizeStrategy: PoolSizeStrategy.NoOp);
+			IAgent Agent1 = await CreateAgentAsync(Pool);
+			IAgent Agent2 = await CreateAgentAsync(Pool);
+
+			// First scale-out will succeed
+			await Service.ResizePools(new() { new PoolSizeData(Pool, new () { Agent1, Agent2 }, 1, "Testing") });
+			Assert.AreEqual(1, FleetManagerSpy.ShrinkPoolAsyncCallCount);
+			
+			// Cannot scale-out due to cool-down
+			await Service.ResizePools(new() { new PoolSizeData(Pool, new () { Agent1 }, 0, "Testing") });
+			Assert.AreEqual(1, FleetManagerSpy.ShrinkPoolAsyncCallCount);
+
+			// Wait some time and then try again
+			await Clock.AdvanceAsync(TimeSpan.FromHours(2));
+			await Service.ResizePools(new() { new PoolSizeData(Pool, new () { Agent1 }, 0, "Testing") });
+			Assert.AreEqual(2, FleetManagerSpy.ShrinkPoolAsyncCallCount);
+		}
+
 		private async Task<List<PoolSizeData>> GetPoolSizeData()
 		{
 			IPool Pool1 = await PoolService.CreatePoolAsync("bogusPool1", null, true, 0, 0);
diff --git a/Engine/Source/Programs/Horde/Horde.Build.Tests/Fleet/JobQueueStrategyTest.cs b/Engine/Source/Programs/Horde/Horde.Build.Tests/Fleet/JobQueueStrategyTest.cs
--- a/Engine/Source/Programs/Horde/Horde.Build.Tests/Fleet/JobQueueStrategyTest.cs
+++ b/Engine/Source/Programs/Horde/Horde.Build.Tests/Fleet/JobQueueStrategyTest.cs
@@ -37,9 +37,9 @@ namespace Horde.Build.Tests.Fleet
 		}
 		
 		[TestMethod]
-		public async Task BelowMinQueueSizeForScaleOut()
+		public async Task EmptyJobQueueWithLargePool()
 		{
-			await AssertAgentCount(2, 0);
+			await AssertAgentCount(0, -2, false, 20);
 		}
 		
 		[TestMethod]
@@ -54,6 +54,12 @@ namespace Horde.Build.Tests.Fleet
 			await AssertAgentCount(3, -1, false);
 		}
 		
+		[TestMethod]
+		public async Task NumQueuedJobs1()
+		{
+			await AssertAgentCount(1, 1, true);
+		}
+		
 		[TestMethod]
 		public async Task NumQueuedJobs3()
 		{
@@ -69,10 +75,10 @@ namespace Horde.Build.Tests.Fleet
 		[TestMethod]
 		public async Task NumQueuedJobs25()
 		{
-			await AssertAgentCount(25, 6);
+			await AssertAgentCount(25, 7);
 		}
 
-		public async Task AssertAgentCount(int NumBatchesReady, int ExpectedAgentDelta, bool WaitedBeyondThreshold = true, int NumAgents = 10)
+		public async Task AssertAgentCount(int NumBatchesReady, int ExpectedAgentDelta, bool WaitedBeyondThreshold = true, int NumAgents = 8)
 		{
 			(JobQueueStrategy Strategy, PoolSizeData PoolSizeData) = await SetUpJobsAsync(1, NumBatchesReady, NumAgents);
 			TimeSpan TimeToWait = WaitedBeyondThreshold
@@ -91,7 +97,7 @@ namespace Horde.Build.Tests.Fleet
 		/// </summary>
 		/// <param name="NumBatchesRunning">Num of job batches that should be in state running</param>
 		/// <param name="NumBatchesReady">Num of job batches that should be in state waiting</param>
-		private async Task<(JobQueueStrategy, PoolSizeData)> SetUpJobsAsync(int NumBatchesRunning, int NumBatchesReady, int NumAgents = 10)
+		private async Task<(JobQueueStrategy, PoolSizeData)> SetUpJobsAsync(int NumBatchesRunning, int NumBatchesReady, int NumAgents = 8)
 		{
 			IPool Pool1 = await PoolService.CreatePoolAsync("bogusPool1", null, true, 0, 0);
 			List<IAgent> Agents = new();
diff --git a/Engine/Source/Programs/Horde/Horde.Build/Api/Pools.cs b/Engine/Source/Programs/Horde/Horde.Build/Api/Pools.cs
--- a/Engine/Source/Programs/Horde/Horde.Build/Api/Pools.cs
+++ b/Engine/Source/Programs/Horde/Horde.Build/Api/Pools.cs
@@ -6,7 +6,6 @@ using System;
 using System.Collections.Generic;
 using System.ComponentModel.DataAnnotations;
 using System.Linq;
-using System.Threading.Tasks;
 using Horde.Build.Fleet.Autoscale;
 
 namespace HordeServer.Api
@@ -35,9 +34,6 @@ namespace HordeServer.Api
 	/// <see cref="Horde.Build.Fleet.Autoscale.JobQueueSettings" />
 	public class JobQueueSettings
 	{
-		/// <see cref="Horde.Build.Fleet.Autoscale.JobQueueSettings.MinQueueSizeForScaleOut" />
-		private int MinQueueSizeForScaleOut { get; set; }
-
 		/// <see cref="Horde.Build.Fleet.Autoscale.JobQueueSettings.ScaleOutFactor" />
 		private double ScaleOutFactor { get; set; }
 
@@ -50,7 +46,6 @@ namespace HordeServer.Api
 		/// <param name="Settings"></param>
 		public JobQueueSettings(Horde.Build.Fleet.Autoscale.JobQueueSettings Settings)
 		{
-			MinQueueSizeForScaleOut = Settings.MinQueueSizeForScaleOut;
 			ScaleInFactor = Settings.ScaleInFactor;
 			ScaleOutFactor = Settings.ScaleOutFactor;
 		}
@@ -61,7 +56,7 @@ namespace HordeServer.Api
 		/// <returns></returns>
 		public Horde.Build.Fleet.Autoscale.JobQueueSettings Convert()
 		{
-			return new Horde.Build.Fleet.Autoscale.JobQueueSettings(MinQueueSizeForScaleOut, ScaleOutFactor, ScaleInFactor);
+			return new Horde.Build.Fleet.Autoscale.JobQueueSettings(ScaleOutFactor, ScaleInFactor);
 		}
 	}
 	
@@ -86,6 +81,16 @@ namespace HordeServer.Api
 		/// </summary>
 		public bool? EnableAutoscaling { get; set; }
 		
+		/// <summary>
+		/// Cooldown time between scale-out events
+		/// </summary>
+		public TimeSpan? ScaleOutCooldown { get; }
+		
+		/// <summary>
+		/// Cooldown time between scale-in events
+		/// </summary>
+		public TimeSpan? ScaleInCooldown { get; }
+		
 		/// <summary>
 		/// Pool sizing strategy
 		/// </summary>
@@ -157,6 +162,16 @@ namespace HordeServer.Api
 		/// </summary>
 		public bool? EnableAutoscaling { get; set; }
 		
+		/// <summary>
+		/// Cooldown time between scale-out events
+		/// </summary>
+		public TimeSpan? ScaleOutCooldown { get; }
+		
+		/// <summary>
+		/// Cooldown time between scale-in events
+		/// </summary>
+		public TimeSpan? ScaleInCooldown { get; }
+		
 		/// <summary>
 		/// Pool sizing strategy
 		/// </summary>
@@ -235,6 +250,16 @@ namespace HordeServer.Api
 		/// </summary>
 		public bool EnableAutoscaling { get; set; }
 		
+		/// <summary>
+		/// Cooldown time between scale-out events
+		/// </summary>
+		public TimeSpan? ScaleOutCooldown { get; }
+		
+		/// <summary>
+		/// Cooldown time between scale-in events
+		/// </summary>
+		public TimeSpan? ScaleInCooldown { get; }
+		
 		/// <summary>
 		/// Pool sizing strategy to be used for this pool
 		/// </summary>
@@ -280,6 +305,8 @@ namespace HordeServer.Api
 			this.Name = Pool.Name;
 			this.Condition = Pool.Condition;
 			this.EnableAutoscaling = Pool.EnableAutoscaling;
+			this.ScaleOutCooldown = Pool.ScaleOutCooldown;
+			this.ScaleInCooldown = Pool.ScaleInCooldown;
 			this.SizeStrategy = Pool.SizeStrategy;
 			this.LeaseUtilizationSettings = Pool.LeaseUtilizationSettings == null ? null : new LeaseUtilizationSettings(Pool.LeaseUtilizationSettings);
 			this.JobQueueSettings = Pool.JobQueueSettings == null ? null : new JobQueueSettings(Pool.JobQueueSettings);
diff --git a/Engine/Source/Programs/Horde/Horde.Build/Controllers/PoolsController.cs b/Engine/Source/Programs/Horde/Horde.Build/Controllers/PoolsController.cs
--- a/Engine/Source/Programs/Horde/Horde.Build/Controllers/PoolsController.cs
+++ b/Engine/Source/Programs/Horde/Horde.Build/Controllers/PoolsController.cs
@@ -63,7 +63,9 @@ namespace HordeServer.Controllers
 			Horde.Build.Fleet.Autoscale.LeaseUtilizationSettings? LuSettings = Create.LeaseUtilizationSettings?.Convert();
 			Horde.Build.Fleet.Autoscale.JobQueueSettings? JqSettings = Create.JobQueueSettings?.Convert();
 
-			IPool NewPool = await PoolService.CreatePoolAsync(Create.Name, Create.Condition, Create.EnableAutoscaling, Create.MinAgents, Create.NumReserveAgents, Create.SizeStrategy, LuSettings, JqSettings, Create.Properties);
+			IPool NewPool = await PoolService.CreatePoolAsync(
+				Create.Name, Create.Condition, Create.EnableAutoscaling, Create.MinAgents, Create.NumReserveAgents,
+				Create.ScaleOutCooldown, Create.ScaleInCooldown, Create.SizeStrategy, LuSettings, JqSettings, Create.Properties);
 			return new CreatePoolResponse(NewPool.Id.ToString());
 		}
 
@@ -142,7 +144,8 @@ namespace HordeServer.Controllers
 				return NotFound(PoolIdValue);
 			}
 
-			await PoolService.UpdatePoolAsync(Pool, Update.Name, Update.Condition, Update.EnableAutoscaling, Update.MinAgents, Update.NumReserveAgents, Update.Properties, Update.SizeStrategy);
+			await PoolService.UpdatePoolAsync(Pool, Update.Name, Update.Condition, Update.EnableAutoscaling,
+				Update.MinAgents, Update.NumReserveAgents, Update.Properties, Update.ScaleOutCooldown, Update.ScaleInCooldown, Update.SizeStrategy);
 			return new OkResult();
 		}
 
diff --git a/Engine/Source/Programs/Horde/Horde.Build/Fleet/Autoscale/AutoscaleServiceV2.cs b/Engine/Source/Programs/Horde/Horde.Build/Fleet/Autoscale/AutoscaleServiceV2.cs
--- a/Engine/Source/Programs/Horde/Horde.Build/Fleet/Autoscale/AutoscaleServiceV2.cs
+++ b/Engine/Source/Programs/Horde/Horde.Build/Fleet/Autoscale/AutoscaleServiceV2.cs
@@ -28,7 +28,8 @@ namespace Horde.Build.Fleet.Autoscale
 	/// </summary>
 	public sealed class AutoscaleServiceV2 : IHostedService, IDisposable
 	{
-		static readonly TimeSpan ShrinkPoolCoolDown = TimeSpan.FromMinutes(20.0);
+		static readonly TimeSpan DefaultScaleOutCooldown = TimeSpan.FromMinutes(5);
+		static readonly TimeSpan DefaultScaleInCooldown = TimeSpan.FromMinutes(20);
 
 		IPoolSizeStrategy LeaseUtilizationStrategy;
 		IPoolSizeStrategy JobQueueStrategy;
@@ -121,7 +122,7 @@ namespace Horde.Build.Fleet.Autoscale
 				
 				try
 				{
-					using IScope Scope = GlobalTracer.Instance.BuildSpan("Scaling pool").StartActive();
+					using IScope Scope = GlobalTracer.Instance.BuildSpan("ScalingPool").StartActive();
 					Scope.Span.SetTag("poolName", Pool.Name);
 					Scope.Span.SetTag("current", CurrentAgentCount);
 					Scope.Span.SetTag("desired", DesiredAgentCount);
@@ -129,21 +130,35 @@ namespace Horde.Build.Fleet.Autoscale
 
 					if (DeltaAgentCount > 0)
 					{
-						await FleetManager.ExpandPoolAsync(Pool, PoolSizeData.Agents, DeltaAgentCount);
-						await PoolCollection.TryUpdateAsync(Pool, LastScaleUpTime: DateTime.UtcNow);
+						TimeSpan ScaleOutCooldown = Pool.ScaleOutCooldown ?? DefaultScaleOutCooldown;
+						bool IsCoolingDown = Pool.LastScaleUpTime != null && Pool.LastScaleUpTime + ScaleOutCooldown > Clock.UtcNow;
+						Scope.Span.SetTag("isCoolingDown", IsCoolingDown);
+						if (!IsCoolingDown)
+						{
+							await FleetManager.ExpandPoolAsync(Pool, PoolSizeData.Agents, DeltaAgentCount);
+							await PoolCollection.TryUpdateAsync(Pool, LastScaleUpTime: Clock.UtcNow);
+						}
+						else
+						{
+							TimeSpan? CooldownTimeLeft = Pool.LastScaleUpTime + DefaultScaleOutCooldown - Clock.UtcNow;
+							Logger.LogDebug("Cannot scale out {PoolName}, it's cooling down for another {TimeLeft} secs", Pool.Name, CooldownTimeLeft?.TotalSeconds);
+						}
 					}
 
 					if (DeltaAgentCount < 0)
 					{
-						bool bShrinkIsOnCoolDown = Pool.LastScaleDownTime != null && Pool.LastScaleDownTime + ShrinkPoolCoolDown > DateTime.UtcNow;
-						if (!bShrinkIsOnCoolDown)
+						TimeSpan ScaleInCooldown = Pool.ScaleInCooldown ?? DefaultScaleInCooldown;
+						bool IsCoolingDown = Pool.LastScaleDownTime != null && Pool.LastScaleDownTime + ScaleInCooldown > Clock.UtcNow;
+						Scope.Span.SetTag("isCoolingDown", IsCoolingDown);
+						if (!IsCoolingDown)
 						{
 							await FleetManager.ShrinkPoolAsync(Pool, PoolSizeData.Agents, -DeltaAgentCount);
-							await PoolCollection.TryUpdateAsync(Pool, LastScaleDownTime: DateTime.UtcNow);
+							await PoolCollection.TryUpdateAsync(Pool, LastScaleDownTime: Clock.UtcNow);
 						}
 						else
 						{
-							Logger.LogDebug("Cannot shrink {PoolName} right now, it's on cool-down until {CoolDownTimeEnds}", Pool.Name, Pool.LastScaleDownTime + ShrinkPoolCoolDown);
+							TimeSpan? CooldownTimeLeft = Pool.LastScaleDownTime + DefaultScaleInCooldown - Clock.UtcNow;
+							Logger.LogDebug("Cannot scale in {PoolName}, it's cooling down for another {TimeLeft} secs", Pool.Name, CooldownTimeLeft?.TotalSeconds);
 						}
 					}
 				}
diff --git a/Engine/Source/Programs/Horde/Horde.Build/Fleet/Autoscale/JobQueueStrategy.cs b/Engine/Source/Programs/Horde/Horde.Build/Fleet/Autoscale/JobQueueStrategy.cs
--- a/Engine/Source/Programs/Horde/Horde.Build/Fleet/Autoscale/JobQueueStrategy.cs
+++ b/Engine/Source/Programs/Horde/Horde.Build/Fleet/Autoscale/JobQueueStrategy.cs
@@ -21,23 +21,16 @@ namespace Horde.Build.Fleet.Autoscale
 	/// </summary>
 	public class JobQueueSettings
 	{
-		/// <summary>
-		/// Minimum number of jobs in queue for scale-out logic to activate
-		///
-		/// Useful to avoid a very small queue size triggering any scaling. 
-		/// </summary>
-		public int MinQueueSizeForScaleOut { get; set; }
-
 		/// <summary>
 		/// Factor translating queue size to additional agents to grow the pool with
-		///
+		/// The result is always rounded up to nearest integer. 
 		/// Example: if there are 20 jobs in queue, a factor 0.25 will result in 5 new agents being added (20 * 0.25)
 		/// </summary>
 		public double ScaleOutFactor { get; set;  } = 0.25;
 		
 		/// <summary>
 		/// Factor by which to shrink the pool size with when queue is empty
-		///
+		/// The result is always rounded up to nearest integer.
 		/// Example: when the queue size is zero, a default value of 0.9 will shrink the pool by 10% (current agent count * 0.9)
 		/// </summary>
 		public double ScaleInFactor { get; set; } = 0.9;
@@ -45,12 +38,10 @@ namespace Horde.Build.Fleet.Autoscale
 		/// <summary>
 		/// Constructor
 		/// </summary>
-		/// <param name="MinQueueSizeForScaleOut"></param>
 		/// <param name="ScaleOutFactor"></param>
 		/// <param name="ScaleInFactor"></param>
-		public JobQueueSettings(int? MinQueueSizeForScaleOut = null, double? ScaleOutFactor = null, double? ScaleInFactor = null)
+		public JobQueueSettings(double? ScaleOutFactor = null, double? ScaleInFactor = null)
 		{
-			this.MinQueueSizeForScaleOut = MinQueueSizeForScaleOut.GetValueOrDefault(this.MinQueueSizeForScaleOut);
 			this.ScaleOutFactor = ScaleOutFactor.GetValueOrDefault(this.ScaleOutFactor);
 			this.ScaleInFactor = ScaleInFactor.GetValueOrDefault(this.ScaleInFactor);
 		}
@@ -168,20 +159,7 @@ namespace Horde.Build.Fleet.Autoscale
 				PoolQueueSizes.TryGetValue(Current.Pool.Id, out var QueueSize);
 				if (QueueSize > 0)
 				{
-					int AdditionalAgentCount = (int)Math.Round(QueueSize * Settings.ScaleOutFactor);
-
-					int NumAgentsEnabled = Current.Agents.Count(x => x.Enabled);
-					if (NumAgentsEnabled == 0)
-					{
-						// Ensure pool can grow even if no agents are online combined with low queue count
-						// A low queue count with default scale out factor will yield an increase below 0.5,
-						// which is rounded down (see above).
-						AdditionalAgentCount = Math.Max(AdditionalAgentCount, 1);	
-					}
-					
-					if (QueueSize < Settings.MinQueueSizeForScaleOut)
-						AdditionalAgentCount = 0;
-
+					int AdditionalAgentCount = (int)Math.Ceiling(QueueSize * Settings.ScaleOutFactor);
 					int DesiredAgentCount = Current.Agents.Count + AdditionalAgentCount;
 					return new PoolSizeData(Current.Pool, Current.Agents, DesiredAgentCount, $"QueueSize={QueueSize}");
 				}
diff --git a/Engine/Source/Programs/Horde/Horde.Build/Fleet/IPool.cs b/Engine/Source/Programs/Horde/Horde.Build/Fleet/IPool.cs
--- a/Engine/Source/Programs/Horde/Horde.Build/Fleet/IPool.cs
+++ b/Engine/Source/Programs/Horde/Horde.Build/Fleet/IPool.cs
@@ -70,6 +70,16 @@ namespace HordeServer.Models
 		/// </summary>
 		public DateTime? LastScaleDownTime { get; }
 		
+		/// <summary>
+		/// Cooldown time between scale-out events
+		/// </summary>
+		public TimeSpan? ScaleOutCooldown { get; }
+		
+		/// <summary>
+		/// Cooldown time between scale-in events
+		/// </summary>
+		public TimeSpan? ScaleInCooldown { get; }
+		
 		/// <summary>
 		/// Pool sizing strategy to be used for this pool
 		/// </summary>
diff --git a/Engine/Source/Programs/Horde/Horde.Build/Fleet/IPoolCollection.cs b/Engine/Source/Programs/Horde/Horde.Build/Fleet/IPoolCollection.cs
--- a/Engine/Source/Programs/Horde/Horde.Build/Fleet/IPoolCollection.cs
+++ b/Engine/Source/Programs/Horde/Horde.Build/Fleet/IPoolCollection.cs
@@ -33,12 +33,26 @@ namespace HordeServer.Collections
 		/// <param name="EnableAutoscaling">Whether to enable autoscaling for this pool</param>
 		/// <param name="MinAgents">Minimum number of agents in the pool</param>
 		/// <param name="NumReserveAgents">Minimum number of idle agents to maintain</param>
+		/// <param name="ScaleOutCooldown">Cooldown time between scale-out events</param>
+		/// <param name="ScaleInCooldown">Cooldown time between scale-in events</param>
 		/// <param name="SizeStrategy">Pool sizing strategy</param>
 		/// <param name="LeaseUtilizationSettings">Settings for lease utilization strategy</param>
 		/// <param name="JobQueueSettings">Settings for job queue strategy</param>
 		/// <param name="Properties">Properties for the pool</param>
 		/// <returns>The new pool document</returns>
-		Task<IPool> AddAsync(PoolId Id, string Name, Condition? Condition = null, bool? EnableAutoscaling = null, int? MinAgents = null, int? NumReserveAgents = null, PoolSizeStrategy? SizeStrategy = null, LeaseUtilizationSettings? LeaseUtilizationSettings = null, JobQueueSettings? JobQueueSettings = null, IEnumerable<KeyValuePair<string, string>>? Properties = null);
+		Task<IPool> AddAsync(
+			PoolId Id,
+			string Name,
+			Condition? Condition = null,
+			bool? EnableAutoscaling = null,
+			int? MinAgents = null,
+			int? NumReserveAgents = null,
+			TimeSpan? ScaleOutCooldown = null,
+			TimeSpan? ScaleInCooldown = null,
+			PoolSizeStrategy? SizeStrategy = null,
+			LeaseUtilizationSettings? LeaseUtilizationSettings = null,
+			JobQueueSettings? JobQueueSettings = null,
+			IEnumerable<KeyValuePair<string, string>>? Properties = null);
 
 		/// <summary>
 		/// Enumerates all the pools
@@ -80,10 +94,28 @@ namespace HordeServer.Collections
 		/// <param name="NewProperties">New properties for the pool</param>
 		/// <param name="LastScaleUpTime">New time for last (auto) scale up</param>
 		/// <param name="LastScaleDownTime">New time for last (auto) scale down</param>
+		/// <param name="ScaleOutCooldown">Cooldown time between scale-out events</param>
+		/// <param name="ScaleInCooldown">Cooldown time between scale-in events</param>
 		/// <param name="SizeStrategy">Pool sizing strategy</param>
 		/// <param name="LeaseUtilizationSettings">Settings for lease utilization strategy</param>
 		/// <param name="JobQueueSettings">Settings for job queue strategy</param>
 		/// <returns>Async task</returns>
-		Task<IPool?> TryUpdateAsync(IPool Pool, string? NewName = null, Condition? NewCondition = null, bool? NewEnableAutoscaling = null, int? NewMinAgents = null, int? NewNumReserveAgents = null, List<AgentWorkspace>? NewWorkspaces = null, bool? NewUseAutoSdk = null, Dictionary<string, string?>? NewProperties = null, DateTime? LastScaleUpTime = null, DateTime? LastScaleDownTime = null, PoolSizeStrategy? SizeStrategy = null, LeaseUtilizationSettings? LeaseUtilizationSettings = null, JobQueueSettings? JobQueueSettings = null);
+		Task<IPool?> TryUpdateAsync(
+			IPool Pool,
+			string? NewName = null,
+			Condition? NewCondition = null,
+			bool? NewEnableAutoscaling = null,
+			int? NewMinAgents = null,
+			int? NewNumReserveAgents = null,
+			List<AgentWorkspace>? NewWorkspaces = null,
+			bool? NewUseAutoSdk = null,
+			Dictionary<string, string?>? NewProperties = null,
+			DateTime? LastScaleUpTime = null,
+			DateTime? LastScaleDownTime = null,
+			TimeSpan? ScaleOutCooldown = null,
+			TimeSpan? ScaleInCooldown = null,
+			PoolSizeStrategy? SizeStrategy = null,
+			LeaseUtilizationSettings? LeaseUtilizationSettings = null,
+			JobQueueSettings? JobQueueSettings = null);
 	}
 }
diff --git a/Engine/Source/Programs/Horde/Horde.Build/Fleet/Impl/PoolCollection.cs b/Engine/Source/Programs/Horde/Horde.Build/Fleet/Impl/PoolCollection.cs
--- a/Engine/Source/Programs/Horde/Horde.Build/Fleet/Impl/PoolCollection.cs
+++ b/Engine/Source/Programs/Horde/Horde.Build/Fleet/Impl/PoolCollection.cs
@@ -41,6 +41,8 @@ namespace HordeServer.Collections.Impl
 			public int? NumReserveAgents { get; set; }
 			public DateTime? LastScaleUpTime { get; set; }
 			public DateTime? LastScaleDownTime { get; set; }
+			public TimeSpan? ScaleOutCooldown { get; set; }
+			public TimeSpan? ScaleInCooldown { get; set; }
 			public PoolSizeStrategy? SizeStrategy { get; set; }
 			public LeaseUtilizationSettings? LeaseUtilizationSettings { get; set; }
 			public JobQueueSettings? JobQueueSettings { get; set; }
@@ -67,6 +69,8 @@ namespace HordeServer.Collections.Impl
 				NumReserveAgents = Other.NumReserveAgents;
 				LastScaleUpTime = Other.LastScaleUpTime;
 				LastScaleDownTime = Other.LastScaleDownTime;
+				ScaleOutCooldown = Other.ScaleOutCooldown;
+				ScaleInCooldown = Other.ScaleInCooldown;
 				SizeStrategy = Other.SizeStrategy;
 				LeaseUtilizationSettings = Other.LeaseUtilizationSettings;
 				JobQueueSettings = Other.JobQueueSettings;
@@ -89,7 +93,19 @@ namespace HordeServer.Collections.Impl
 		}
 
 		/// <inheritdoc/>
-		public async Task<IPool> AddAsync(PoolId Id, string Name, Condition? Condition, bool? EnableAutoscaling, int? MinAgents, int? NumReserveAgents, PoolSizeStrategy? SizeStrategy, LeaseUtilizationSettings? LeaseUtilizationSettings, JobQueueSettings? JobQueueSettings, IEnumerable<KeyValuePair<string, string>>? Properties)
+		public async Task<IPool> AddAsync(
+			PoolId Id,
+			string Name,
+			Condition? Condition,
+			bool? EnableAutoscaling,
+			int? MinAgents,
+			int? NumReserveAgents,
+			TimeSpan? ScaleOutCooldown,
+			TimeSpan? ScaleInCooldown,
+			PoolSizeStrategy? SizeStrategy,
+			LeaseUtilizationSettings? LeaseUtilizationSettings,
+			JobQueueSettings? JobQueueSettings,
+			IEnumerable<KeyValuePair<string, string>>? Properties)
 		{
 			PoolDocument Pool = new PoolDocument();
 			Pool.Id = Id;
@@ -106,6 +122,8 @@ namespace HordeServer.Collections.Impl
 				Pool.Properties = new Dictionary<string, string>(Properties);
 			}
 
+			Pool.ScaleOutCooldown = ScaleOutCooldown;
+			Pool.ScaleInCooldown = ScaleInCooldown;
 			Pool.SizeStrategy = SizeStrategy;
 			Pool.LeaseUtilizationSettings = LeaseUtilizationSettings;
 			Pool.JobQueueSettings = JobQueueSettings;
@@ -183,7 +201,23 @@ namespace HordeServer.Collections.Impl
 		}
 
 		/// <inheritdoc/>
-		public Task<IPool?> TryUpdateAsync(IPool Pool, string? NewName, Condition? NewCondition, bool? NewEnableAutoscaling, int? NewMinAgents, int? NewNumReserveAgents, List<AgentWorkspace>? NewWorkspaces, bool? NewUseAutoSdk, Dictionary<string, string?>? NewProperties, DateTime? LastScaleUpTime, DateTime? LastScaleDownTime, PoolSizeStrategy? SizeStrategy, LeaseUtilizationSettings? LeaseUtilizationSettings, JobQueueSettings? JobQueueSettings)
+		public Task<IPool?> TryUpdateAsync(
+			IPool Pool,
+			string? NewName,
+			Condition? NewCondition,
+			bool? NewEnableAutoscaling,
+			int? NewMinAgents,
+			int? NewNumReserveAgents,
+			List<AgentWorkspace>? NewWorkspaces,
+			bool? NewUseAutoSdk,
+			Dictionary<string, string?>? NewProperties,
+			DateTime? LastScaleUpTime,
+			DateTime? LastScaleDownTime,
+			TimeSpan? ScaleOutCooldown,
+			TimeSpan? ScaleInCooldown, 
+			PoolSizeStrategy? SizeStrategy,
+			LeaseUtilizationSettings? LeaseUtilizationSettings,
+			JobQueueSettings? JobQueueSettings)
 		{
 			TransactionBuilder<PoolDocument> Transaction = new TransactionBuilder<PoolDocument>();
 			if (NewName != null)
@@ -254,6 +288,14 @@ namespace HordeServer.Collections.Impl
 			{
 				Transaction.Set(x => x.LastScaleDownTime, LastScaleDownTime);
 			}
+			if (ScaleOutCooldown != null)
+			{
+				Transaction.Set(x => x.ScaleOutCooldown, ScaleOutCooldown);
+			}
+			if (ScaleInCooldown != null)
+			{
+				Transaction.Set(x => x.ScaleInCooldown, ScaleInCooldown);
+			}
 			if (SizeStrategy != null)
 			{
 				Transaction.Set(x => x.SizeStrategy, SizeStrategy);
diff --git a/Engine/Source/Programs/Horde/Horde.Build/Fleet/PoolService.cs b/Engine/Source/Programs/Horde/Horde.Build/Fleet/PoolService.cs
--- a/Engine/Source/Programs/Horde/Horde.Build/Fleet/PoolService.cs
+++ b/Engine/Source/Programs/Horde/Horde.Build/Fleet/PoolService.cs
@@ -62,14 +62,39 @@ namespace HordeServer.Services
 		/// <param name="EnableAutoscaling">Whether to enable autoscaling for this pool</param>
 		/// <param name="MinAgents">Minimum number of agents in the pool</param>
 		/// <param name="NumReserveAgents">Minimum number of idle agents to maintain</param>
+		/// <param name="ScaleOutCooldown">Cooldown time between scale-out events</param>
+		/// <param name="ScaleInCooldown">Cooldown time between scale-in events</param>
 		/// <param name="SizeStrategy">Pool sizing strategy</param>
 		/// <param name="LeaseUtilizationSettings">Settings for lease utilization strategy</param>
 		/// <param name="JobQueueSettings">Settings for job queue strategy</param>
 		/// <param name="Properties">Properties for the new pool</param>
 		/// <returns>The new pool document</returns>
-		public Task<IPool> CreatePoolAsync(string Name, Condition? Condition = null, bool? EnableAutoscaling = null, int? MinAgents = null, int? NumReserveAgents = null, PoolSizeStrategy? SizeStrategy = null, LeaseUtilizationSettings? LeaseUtilizationSettings = null, JobQueueSettings? JobQueueSettings = null, Dictionary<string, string>? Properties = null)
+		public Task<IPool> CreatePoolAsync(
+			string Name,
+			Condition? Condition = null,
+			bool? EnableAutoscaling = null,
+			int? MinAgents = null,
+			int? NumReserveAgents = null,
+			TimeSpan? ScaleOutCooldown = null,
+			TimeSpan? ScaleInCooldown = null,
+			PoolSizeStrategy? SizeStrategy = null,
+			LeaseUtilizationSettings? LeaseUtilizationSettings = null,
+			JobQueueSettings? JobQueueSettings = null,
+			Dictionary<string, string>? Properties = null)
 		{
-			return Pools.AddAsync(PoolId.Sanitize(Name), Name, Condition, EnableAutoscaling, MinAgents, NumReserveAgents, SizeStrategy, LeaseUtilizationSettings, JobQueueSettings, Properties);
+			return Pools.AddAsync(
+				PoolId.Sanitize(Name),
+				Name,
+				Condition,
+				EnableAutoscaling,
+				MinAgents,
+				NumReserveAgents,
+				ScaleOutCooldown,
+				ScaleInCooldown,
+				SizeStrategy,
+				LeaseUtilizationSettings,
+				JobQueueSettings,
+				Properties);
 		}
 
 		/// <summary>
@@ -92,13 +117,36 @@ namespace HordeServer.Services
 		/// <param name="NewMinAgents">Minimum number of agents in the pool</param>
 		/// <param name="NewNumReserveAgents">Minimum number of idle agents to maintain</param>
 		/// <param name="NewProperties">Properties on the pool to update. Any properties with a value of null will be removed.</param>
+		/// <param name="ScaleOutCooldown">Cooldown time between scale-out events</param>
+		/// <param name="ScaleInCooldown">Cooldown time between scale-in events</param>
 		/// <param name="SizeStrategy">New pool sizing strategy for the pool</param>
 		/// <returns>Async task object</returns>
-		public async Task<IPool?> UpdatePoolAsync(IPool? Pool, string? NewName = null, Condition? NewCondition = null, bool? NewEnableAutoscaling = null, int? NewMinAgents = null, int? NewNumReserveAgents = null, Dictionary<string, string?>? NewProperties = null, PoolSizeStrategy? SizeStrategy = null)
+		public async Task<IPool?> UpdatePoolAsync(
+			IPool? Pool,
+			string? NewName = null,
+			Condition? NewCondition = null,
+			bool? NewEnableAutoscaling = null,
+			int? NewMinAgents = null,
+			int? NewNumReserveAgents = null,
+			Dictionary<string, string?>? NewProperties = null,
+			TimeSpan? ScaleOutCooldown = null,
+			TimeSpan? ScaleInCooldown = null,
+			PoolSizeStrategy? SizeStrategy = null)
 		{
 			for (; Pool != null; Pool = await Pools.GetAsync(Pool.Id))
 			{
-				IPool? NewPool = await Pools.TryUpdateAsync(Pool, NewName, NewCondition, NewEnableAutoscaling, NewMinAgents, NewNumReserveAgents, NewProperties: NewProperties, SizeStrategy: SizeStrategy);
+				IPool? NewPool = await Pools.TryUpdateAsync(
+					Pool,
+					NewName,
+					NewCondition,
+					NewEnableAutoscaling,
+					NewMinAgents,
+					NewNumReserveAgents,
+					NewProperties: NewProperties,
+					ScaleOutCooldown: ScaleOutCooldown,
+					ScaleInCooldown: ScaleInCooldown,
+					SizeStrategy: SizeStrategy);
+				
 				if (NewPool != null)
 				{
 					return NewPool;
