Horde: Add support for reading CPU and memory info for Linux agents, and add a standalone command "caps" which will output all detected capabilities for the current agent.

diff --git a/Engine/Source/Programs/Horde/Horde.Agent/Commands/Utilities/CapsCommand.cs b/Engine/Source/Programs/Horde/Horde.Agent/Commands/Utilities/CapsCommand.cs
new file mode 100644
--- /dev/null
+++ b/Engine/Source/Programs/Horde/Horde.Agent/Commands/Utilities/CapsCommand.cs
@@ -0,0 +1,45 @@
+// Copyright Epic Games, Inc. All Rights Reserved.
+
+using Grpc.Net.Client;
+using HordeAgent.Utility;
+using HordeCommon.Rpc;
+using Microsoft.Extensions.Configuration;
+using Microsoft.Extensions.Logging;
+using System;
+using System.Collections.Generic;
+using System.IO;
+using System.IO.Compression;
+using System.Net;
+using System.Net.Http;
+using System.Net.Http.Headers;
+using System.Reflection;
+using System.Text;
+using System.Threading.Tasks;
+using EpicGames.Core;
+using HordeCommon.Rpc.Messages;
+using HordeAgent.Services;
+
+namespace HordeAgent.Commands.Utilities
+{
+	/// <summary>
+	/// Shows capabilities of this agent
+	/// </summary>
+	[Command("Caps", "Lists detected capabilities of this agent")]
+	class CapsCommand : Command
+	{
+		/// <inheritdoc/>
+		public override async Task<int> ExecuteAsync(ILogger Logger)
+		{
+			AgentCapabilities Capabilities = await WorkerService.GetAgentCapabilities(DirectoryReference.GetCurrentDirectory(), Logger);
+			foreach (DeviceCapabilities Device in Capabilities.Devices)
+			{
+				Logger.LogInformation("Device: {Name}", Device.Handle);
+				foreach (string Property in Device.Properties)
+				{
+					Logger.LogInformation("  {Property}", Property);
+				}
+			}
+			return 0;
+		}
+	}
+}
diff --git a/Engine/Source/Programs/Horde/Horde.Agent/Services/WorkerService.cs b/Engine/Source/Programs/Horde/Horde.Agent/Services/WorkerService.cs
--- a/Engine/Source/Programs/Horde/Horde.Agent/Services/WorkerService.cs
+++ b/Engine/Source/Programs/Horde/Horde.Agent/Services/WorkerService.cs
@@ -42,6 +42,7 @@ using EpicGames.Horde.Storage;
 using EpicGames.Horde.Compute;
 using System.Net;
 using EpicGames.Horde.Storage.Impl;
+using System.Text.RegularExpressions;
 
 namespace HordeAgent.Services
 {
@@ -1482,7 +1483,7 @@ namespace HordeAgent.Services
 		/// Gets the hardware capabilities of this worker
 		/// </summary>
 		/// <returns>Worker object for advertising to the server</returns>
-		static async Task<AgentCapabilities> GetAgentCapabilities(DirectoryReference WorkingDir, ILogger Logger)
+		public static async Task<AgentCapabilities> GetAgentCapabilities(DirectoryReference WorkingDir, ILogger Logger)
 		{
 			// Create the primary device
 			DeviceCapabilities PrimaryDevice = new DeviceCapabilities();
@@ -1558,20 +1559,7 @@ namespace HordeAgent.Services
 						}
 					}
 
-					if (NameToCount.Count > 0)
-					{
-						PrimaryDevice.Properties.Add("CPU=" + String.Join(", ", NameToCount.Select(x => (x.Value > 1) ? $"{x.Key} x {x.Value}" : x.Key)));
-					}
-
-					if (TotalLogicalCores > 0)
-					{
-						PrimaryDevice.Properties.Add($"LogicalCores={TotalLogicalCores}");
-					}
-
-					if (TotalPhysicalCores > 0)
-					{
-						PrimaryDevice.Properties.Add($"PhysicalCores={TotalPhysicalCores}");
-					}
+					AddCpuInfo(PrimaryDevice, NameToCount, TotalLogicalCores, TotalPhysicalCores);
 				}
 
 				// Add RAM info
@@ -1637,6 +1625,55 @@ namespace HordeAgent.Services
 				// Add EC2 properties if needed
 				await AddAwsProperties(PrimaryDevice.Properties, Logger);
 
+				// Parse the CPU info
+				List<Dictionary<string, string>>? CpuRecords = await ReadLinuxHwPropsAsync("/proc/cpuinfo", Logger);
+				if (CpuRecords != null)
+				{
+					Dictionary<string, string> CpuNames = new Dictionary<string, string>(StringComparer.Ordinal);
+					foreach (Dictionary<string, string> CpuRecord in CpuRecords)
+					{
+						if (CpuRecord.TryGetValue("physical id", out string? PhysicalId) && CpuRecord.TryGetValue("model name", out string? ModelName))
+						{
+							CpuNames[PhysicalId] = ModelName;
+						}
+					}
+
+					Dictionary<string, int> NameToCount = new Dictionary<string, int>(StringComparer.Ordinal);
+					foreach (string CpuName in CpuNames.Values)
+					{
+						NameToCount.TryGetValue(CpuName, out int Count);
+						NameToCount[CpuName] = Count + 1;
+					}
+
+					HashSet<string> LogicalCores = new HashSet<string>();
+					HashSet<string> PhysicalCores = new HashSet<string>();
+					foreach (Dictionary<string, string> CpuRecord in CpuRecords)
+					{
+						if (CpuRecord.TryGetValue("processor", out string? LogicalCoreId))
+						{
+							LogicalCores.Add(LogicalCoreId);
+						}
+						if (CpuRecord.TryGetValue("core id", out string? PhysicalCoreId))
+						{
+							PhysicalCores.Add(PhysicalCoreId);
+						}
+					}
+
+					AddCpuInfo(PrimaryDevice, NameToCount, LogicalCores.Count, PhysicalCores.Count);
+				}
+
+				// Parse the RAM info
+				List<Dictionary<string, string>>? MemRecords = await ReadLinuxHwPropsAsync("/proc/meminfo", Logger);
+				if (MemRecords != null && MemRecords.Count > 0 && MemRecords[0].TryGetValue("MemTotal", out string? MemTotal))
+				{
+					Match Match = Regex.Match(MemTotal, @"(\d+)\s+kB");
+					if (Match.Success)
+					{
+						long TotalCapacity = long.Parse(Match.Groups[1].Value) * 1024;
+						PrimaryDevice.Properties.Add($"RAM={TotalCapacity / (1024 * 1024 * 1024)}");
+					}
+				}
+
 				// Add session information
 				PrimaryDevice.Properties.Add($"User={Environment.UserName}");
 				PrimaryDevice.Properties.Add($"Domain={Environment.UserDomainName}");
@@ -1787,6 +1824,71 @@ namespace HordeAgent.Services
 			return Agent;
 		}
 
+		static void AddCpuInfo(DeviceCapabilities PrimaryDevice, Dictionary<string, int> NameToCount, int NumLogicalCores, int NumPhysicalCores)
+		{
+			if (NameToCount.Count > 0)
+			{
+				PrimaryDevice.Properties.Add("CPU=" + String.Join(", ", NameToCount.Select(x => (x.Value > 1) ? $"{x.Key} x {x.Value}" : x.Key)));
+			}
+
+			if (NumLogicalCores > 0)
+			{
+				PrimaryDevice.Properties.Add($"LogicalCores={NumLogicalCores}");
+			}
+
+			if (NumPhysicalCores > 0)
+			{
+				PrimaryDevice.Properties.Add($"PhysicalCores={NumPhysicalCores}");
+			}
+		}
+
+		static async Task<List<Dictionary<string, string>>?> ReadLinuxHwPropsAsync(string FileName, ILogger Logger)
+		{
+			List<Dictionary<string, string>>? Records = null;
+			if (File.Exists(FileName))
+			{
+				Records = new List<Dictionary<string, string>>();
+				using (StreamReader Reader = new StreamReader(FileName))
+				{
+					Dictionary<string, string> Record = new Dictionary<string, string>(StringComparer.Ordinal);
+
+					string? Line;
+					while ((Line = await Reader.ReadLineAsync()) != null)
+					{
+						int Idx = Line.IndexOf(':');
+						if (Idx == -1)
+						{
+							if (Record.Count > 0)
+							{
+								Records.Add(Record);
+								Record = new Dictionary<string, string>(StringComparer.Ordinal);
+							}
+						}
+						else
+						{
+							string Key = Line.Substring(0, Idx).Trim();
+							string Value = Line.Substring(Idx + 1).Trim();
+
+							if (Record.TryGetValue(Key, out string? PrevValue))
+							{
+								Logger.LogWarning("Multiple entries for {Key} in {File} (was '{Prev}', now '{Next}')", Key, FileName, PrevValue, Value);
+							}
+							else
+							{
+								Record.Add(Key, Value);
+							}
+						}
+					}
+
+					if (Record.Count > 0)
+					{
+						Records.Add(Record);
+					}
+				}
+			}
+			return Records;
+		}
+
 		class WmiProperties
 		{
 			Dictionary<string, object> Properties = new Dictionary<string, object>(StringComparer.Ordinal);
