Horde: Always log in to the Perforce server when creating a new service connection.

diff --git a/Engine/Source/Programs/Horde/Horde.Build/Perforce/Impl/PerforceService.cs b/Engine/Source/Programs/Horde/Horde.Build/Perforce/Impl/PerforceService.cs
--- a/Engine/Source/Programs/Horde/Horde.Build/Perforce/Impl/PerforceService.cs
+++ b/Engine/Source/Programs/Horde/Horde.Build/Perforce/Impl/PerforceService.cs
@@ -229,12 +229,24 @@ namespace HordeServer.Services
 			IPerforceServer Server = await SelectServer(Cluster);
 
 			PerforceSettings Settings = new PerforceSettings(Server.ServerAndPort, UserName);
-			Settings.Password = Password;
 			Settings.AppName = "Horde.Build";
 			Settings.ClientName = "__DOES_NOT_EXIST__";
 			Settings.PreferNativeClient = true;
 
-			return await PerforceConnection.CreateAsync(Settings, Logger);
+			IPerforceConnection Connection = await PerforceConnection.CreateAsync(Settings, Logger);
+			try
+			{
+				if (Password != null)
+				{
+					await Connection.LoginAsync(Password);
+				}
+				return Connection;
+			}
+			catch
+			{
+				Connection.Dispose();
+				throw;
+			}
 		}
 
 		async Task<P4.Repository> GetServiceUserConnection(PerforceCluster Cluster)
