Horde: Additional logging for not being able to do incremental P4 stream updates.

diff --git a/Engine/Source/Programs/Horde/Horde.Build/Perforce/Impl/CommitService.cs b/Engine/Source/Programs/Horde/Horde.Build/Perforce/Impl/CommitService.cs
--- a/Engine/Source/Programs/Horde/Horde.Build/Perforce/Impl/CommitService.cs
+++ b/Engine/Source/Programs/Horde/Horde.Build/Perforce/Impl/CommitService.cs
@@ -818,8 +818,11 @@ namespace HordeServer.Commits.Impl
 				{
 					break;
 				}
+
+				// Check that the previous commit is still valid
 				if (PrevCommit != null && PrevCommit.Change != Values[0])
 				{
+					Logger.LogInformation("Invalidating previous commit; expected {Change}, actually {ActualChange}", PrevCommit.Change, Values[0]);
 					PrevCommit = null;
 				}
 
@@ -858,16 +861,17 @@ namespace HordeServer.Commits.Impl
 		{
 			// Get the ref corresponding to BaseChange
 			CommitTree? BaseTree = null;
-			if (BaseCommit != null)
+			if (BaseCommit == null)
 			{
-				if (BaseCommit.TreeRefId == null)
-				{
-					Logger.LogWarning("Base commit for stream {StreamId} change {Change} does not have a mirrored tree", Stream.Id, BaseCommit.Change);
-				}
-				else
-				{
-					BaseTree = new CommitTree(BaseCommit.StreamId, BaseCommit.Change, BaseCommit.TreeRefId.Value);
-				}
+				Logger.LogInformation("No base commit for mirroring {StreamId} change {Change}", Stream.Id, Commit.Change);
+			}
+			else if (BaseCommit.TreeRefId == null)
+			{
+				Logger.LogWarning("Base commit for stream {StreamId} change {Change} does not have a mirrored tree", Stream.Id, BaseCommit.Change);
+			}
+			else
+			{
+				BaseTree = new CommitTree(BaseCommit.StreamId, BaseCommit.Change, BaseCommit.TreeRefId.Value);
 			}
 
 			// Write the new tree
