HordeAgent Compute: Check for missing blobs before uploading & upload in parallel

#rnx

diff --git a/Engine/Source/Programs/Horde/Horde.Agent/Commands/Compute/ComputeCommand.cs b/Engine/Source/Programs/Horde/Horde.Agent/Commands/Compute/ComputeCommand.cs
--- a/Engine/Source/Programs/Horde/Horde.Agent/Commands/Compute/ComputeCommand.cs
+++ b/Engine/Source/Programs/Horde/Horde.Agent/Commands/Compute/ComputeCommand.cs
@@ -3,18 +3,13 @@
 using Google.Protobuf;
 using Microsoft.Extensions.Logging;
 using System;
-using System.Collections.Concurrent;
 using System.Collections.Generic;
 using System.IO;
 using System.Linq;
 using System.Text;
+using System.Threading;
 using System.Threading.Tasks;
 using EpicGames.Core;
-using System.Text.Json;
-using HordeAgent.Services;
-using Grpc.Net.Client;
-using System.Net.Http.Headers;
-using Grpc.Core;
 using Microsoft.Extensions.Hosting;
 using Microsoft.Extensions.DependencyInjection;
 using Microsoft.Extensions.Configuration;
@@ -22,9 +17,6 @@ using Serilog.Events;
 using EpicGames.Serialization;
 using EpicGames.Horde.Compute;
 using EpicGames.Horde.Storage;
-using System.Net.Http;
-using EpicGames.Horde.Compute.Impl;
-using EpicGames.Horde.Storage.Impl;
 
 namespace HordeAgent.Commands
 {
@@ -212,12 +204,26 @@ namespace HordeAgent.Commands
 		{
 			IComputeClusterInfo Cluster = await ComputeClient.GetClusterInfoAsync(ClusterId);
 
-			List<KeyValuePair<IoHash, byte[]>> UploadBlobs = UploadList.ToList();
-			for (int Idx = 0; Idx < UploadBlobs.Count; Idx++)
 			{
-				KeyValuePair<IoHash, byte[]> Pair = UploadBlobs[Idx];
-				Logger.LogInformation("Uploading blob {Idx}/{Count}: {Hash}", Idx + 1, UploadBlobs.Count, Pair.Key);
-				await StorageClient.WriteBlobFromMemoryAsync(Cluster.NamespaceId, Pair.Key, Pair.Value);
+				int Index = 0, TotalUploaded = 0, TotalSkipped = 0, UploadedBytes = 0, SkippedBytes = 0;
+				var Tasks = UploadList.Select(async (Pair) =>
+				{
+					if (await StorageClient.HasBlobAsync(Cluster.NamespaceId, Pair.Key))
+					{
+						Logger.LogInformation("Skipped blob {Idx}/{Count}: {Hash} ({Length} bytes)", Interlocked.Increment(ref Index), UploadList.Count, Pair.Key, Pair.Value.Length);
+						Interlocked.Increment(ref TotalSkipped);
+						Interlocked.Add(ref SkippedBytes, Pair.Value.Length);
+						return;
+					}
+
+					await StorageClient.WriteBlobFromMemoryAsync(Cluster.NamespaceId, Pair.Key, Pair.Value);
+					Logger.LogInformation("Uploaded blob {Idx}/{Count}: {Hash} ({Bytes} bytes)", Interlocked.Increment(ref Index), UploadList.Count, Pair.Key, Pair.Value.Length);
+					Interlocked.Increment(ref TotalUploaded);
+					Interlocked.Add(ref UploadedBytes, Pair.Value.Length);
+				});
+				await System.Threading.Tasks.Task.WhenAll(Tasks);
+
+				Logger.LogInformation("Uploaded {UploadedCount} blobs ({UploadedBytes} bytes), Skipped {SkippedCount} blobs ({SkippedBytes} bytes)", TotalUploaded, UploadedBytes, TotalSkipped, SkippedBytes);
 			}
 
 			CbObject TaskObject = CbSerializer.Serialize(Task);
