Horde: Fix commit service always performing a full snapshot.

diff --git a/Engine/Source/Programs/Horde/Horde.Build/Perforce/Impl/CommitService.cs b/Engine/Source/Programs/Horde/Horde.Build/Perforce/Impl/CommitService.cs
--- a/Engine/Source/Programs/Horde/Horde.Build/Perforce/Impl/CommitService.cs
+++ b/Engine/Source/Programs/Horde/Horde.Build/Perforce/Impl/CommitService.cs
@@ -826,27 +826,15 @@ namespace HordeServer.Commits.Impl
 					PrevCommit = null;
 				}
 
-				// Get the commit we want to mirror
-				int Change = Values[(PrevCommit == null) ? 0 : 1];
-
-				// Check if the tree already exists. If it does, we don't need to run replication again.
-				ICommit? Commit = await CommitCollection.GetCommitAsync(StreamId, Change);
-				if (Commit == null)
-				{
-					Logger.LogWarning("Missing commit for {Stream} change {Change}", StreamId, Change);
-				}
-				else if (Commit.TreeRefId != null)
-				{
-					Logger.LogInformation("Skipping replication for {Stream} change {Change}; tree already exists", StreamId, Change);
-				}
-				else
+				// If we don't have a previous commit tree yet, perform a full snapshot of that change
+				if (PrevCommit == null)
 				{
-					Commit = await UpdateCommitTreeAsync(Stream, Commit, PrevCommit);
+					PrevCommit = await UpdateCommitTreeAsync(Stream, Values[0], null);
 				}
 
-				// Remove the first change number from this queue
+				// Perform a snapshot of the new change, then remove it from the list
+				PrevCommit = await UpdateCommitTreeAsync(Stream, Values[1], PrevCommit);
 				await Changes.LeftPopAsync();
-				PrevCommit = Commit;
 			}
 		}
 
@@ -854,24 +842,37 @@ namespace HordeServer.Commits.Impl
 		/// Replicates the contents of a stream to Horde storage, optionally using the given change as a starting point
 		/// </summary>
 		/// <param name="Stream">The stream to replicate</param>
-		/// <param name="Commit">Commit to store the tree ref</param>
-		/// <param name="BaseCommit">Previous commit to use as a base</param>
+		/// <param name="Change">Commit to store the tree ref</param>
+		/// <param name="PrevCommit">Previous commit to use as a base</param>
 		/// <returns>New ref id for the given commit</returns>
-		public async Task<ICommit> UpdateCommitTreeAsync(IStream Stream, ICommit Commit, ICommit? BaseCommit = null)
+		public async Task<ICommit?> UpdateCommitTreeAsync(IStream Stream, int Change, ICommit? PrevCommit)
 		{
+			// Get the commit for the requested change
+			ICommit? Commit = await CommitCollection.GetCommitAsync(Stream.Id, Change);
+			if (Commit == null)
+			{
+				Logger.LogWarning("Missing commit for {Stream} change {Change}", Stream.Id, Change);
+				return Commit;
+			}
+			else if (Commit.TreeRefId != null)
+			{
+				Logger.LogInformation("Skipping replication for {Stream} change {Change}; tree already exists", Stream.Id, Change);
+				return Commit;
+			}
+
 			// Get the ref corresponding to BaseChange
 			CommitTree? BaseTree = null;
-			if (BaseCommit == null)
+			if (PrevCommit == null)
 			{
 				Logger.LogInformation("No base commit for mirroring {StreamId} change {Change}", Stream.Id, Commit.Change);
 			}
-			else if (BaseCommit.TreeRefId == null)
+			else if (PrevCommit.TreeRefId == null)
 			{
-				Logger.LogWarning("Base commit for stream {StreamId} change {Change} does not have a mirrored tree", Stream.Id, BaseCommit.Change);
+				Logger.LogWarning("Base commit for stream {StreamId} change {Change} does not have a mirrored tree", Stream.Id, PrevCommit.Change);
 			}
 			else
 			{
-				BaseTree = new CommitTree(BaseCommit.StreamId, BaseCommit.Change, BaseCommit.TreeRefId.Value);
+				BaseTree = new CommitTree(PrevCommit.StreamId, PrevCommit.Change, PrevCommit.TreeRefId.Value);
 			}
 
 			// Write the new tree
