Horde Storage - Debug output to track down issue with some replicators not running as frequently as they should.

diff --git a/Engine/Source/Programs/Horde/HordeStorage/Horde.Storage/Implementation/Replication/ReplicationService.cs b/Engine/Source/Programs/Horde/HordeStorage/Horde.Storage/Implementation/Replication/ReplicationService.cs
--- a/Engine/Source/Programs/Horde/HordeStorage/Horde.Storage/Implementation/Replication/ReplicationService.cs
+++ b/Engine/Source/Programs/Horde/HordeStorage/Horde.Storage/Implementation/Replication/ReplicationService.cs
@@ -97,6 +97,8 @@ namespace Horde.Storage.Implementation
                 return false;
             }
 
+            _logger.Information("Polling for new replication to start");
+
             foreach (IReplicator replicator in state.Replicators)
             {
                 if (_currentReplications.TryGetValue(replicator.Info.ReplicatorName, out Task<bool>? replicationTask))
@@ -112,7 +114,10 @@ namespace Horde.Storage.Implementation
                             continue;
                         }
                      
+                        DateTime time = DateTime.Now;
+                        _logger.Information("Joining replication task for replicator {Name}", replicator.Info.ReplicatorName);
                         await replicationTask;
+                        _logger.Information("Waited for replication task {Name} for {Duration} seconds", replicator.Info.ReplicatorName, (DateTime.Now - time).TotalSeconds);
                     }
                     else
                     {
