Horde Storage - When requesting a blob via the compressed-blobs endpoint, you can now indicate if you accept uncompressed (without the compressed buffer header) content. Use application/octet-stream accept type for that.
This endpoint will now return different mime types based on the content it returns (x-ue-comp for compressed buffers which should be the norm, application/octet-stream if its only available as uncompressed). If the accept type requires a content type we can not provide a 415 (unsupported media type) is returned.

diff --git a/Engine/Source/Programs/Horde/HordeStorage/Horde.Storage/Controllers/CompressedBlobController.cs b/Engine/Source/Programs/Horde/HordeStorage/Horde.Storage/Controllers/CompressedBlobController.cs
--- a/Engine/Source/Programs/Horde/HordeStorage/Horde.Storage/Controllers/CompressedBlobController.cs
+++ b/Engine/Source/Programs/Horde/HordeStorage/Horde.Storage/Controllers/CompressedBlobController.cs
@@ -17,6 +17,7 @@ using Jupiter.Implementation;
 using Jupiter.Utils;
 using Microsoft.AspNetCore.Authorization;
 using Microsoft.AspNetCore.Mvc;
+using Microsoft.Extensions.Primitives;
 using Serilog;
 
 namespace Horde.Storage.Controllers
@@ -33,10 +34,11 @@ namespace Horde.Storage.Controllers
         private readonly IAuthorizationService _authorizationService;
         private readonly CompressedBufferUtils _compressedBufferUtils;
         private readonly BufferedPayloadFactory _bufferedPayloadFactory;
+        private readonly FormatResolver _formatResolver;
 
         private readonly ILogger _logger = Log.ForContext<CompressedBlobController>();
 
-        public CompressedBlobController(IBlobService storage, IContentIdStore contentIdStore, IDiagnosticContext diagnosticContext, IAuthorizationService authorizationService, CompressedBufferUtils compressedBufferUtils, BufferedPayloadFactory bufferedPayloadFactory)
+        public CompressedBlobController(IBlobService storage, IContentIdStore contentIdStore, IDiagnosticContext diagnosticContext, IAuthorizationService authorizationService, CompressedBufferUtils compressedBufferUtils, BufferedPayloadFactory bufferedPayloadFactory, FormatResolver formatResolver)
         {
             _storage = storage;
             _contentIdStore = contentIdStore;
@@ -44,6 +46,7 @@ namespace Horde.Storage.Controllers
             _authorizationService = authorizationService;
             _compressedBufferUtils = compressedBufferUtils;
             _bufferedPayloadFactory = bufferedPayloadFactory;
+            _formatResolver = formatResolver;
         }
 
 
@@ -51,7 +54,7 @@ namespace Horde.Storage.Controllers
         [Authorize("Storage.read")]
         [ProducesResponseType(type: typeof(byte[]), 200)]
         [ProducesResponseType(type: typeof(ValidationProblemDetails), 400)]
-        [Produces(CustomMediaTypeNames.UnrealCompressedBuffer)]
+        [Produces(CustomMediaTypeNames.UnrealCompressedBuffer, MediaTypeNames.Application.Octet)]
 
         public async Task<IActionResult> Get(
             [Required] NamespaceId ns,
@@ -66,9 +69,13 @@ namespace Horde.Storage.Controllers
 
             try
             {
-                BlobContents blobContents = await GetImpl(ns, id);
+                (BlobContents blobContents, string mediaType) = await GetImpl(ns, id);
 
-                return File(blobContents.Stream, CustomMediaTypeNames.UnrealCompressedBuffer, enableRangeProcessing: true);
+                StringValues acceptHeader = Request.Headers["Accept"];
+                if (acceptHeader.Count != 0 && !acceptHeader.Contains(mediaType))
+                    return new UnsupportedMediaTypeResult();
+
+                return File(blobContents.Stream, mediaType, enableRangeProcessing: true);
             }
             catch (BlobNotFoundException e)
             {
@@ -93,7 +100,7 @@ namespace Horde.Storage.Controllers
             {
                 return Forbid();
             }
-            BlobIdentifier[]? chunks = await _contentIdStore.Resolve(ns, id, mustBeContentId: true);
+            BlobIdentifier[]? chunks = await _contentIdStore.Resolve(ns, id, mustBeContentId: false);
             if (chunks == null || chunks.Length == 0)
             {
                 return NotFound();
@@ -136,7 +143,7 @@ namespace Horde.Storage.Controllers
 
             IEnumerable<Task> tasks = id.Select(async blob =>
             {
-                BlobIdentifier[]? chunks = await _contentIdStore.Resolve(ns, blob, mustBeContentId: true);
+                BlobIdentifier[]? chunks = await _contentIdStore.Resolve(ns, blob, mustBeContentId: false);
 
                 if (chunks == null)
                 {
@@ -161,9 +168,9 @@ namespace Horde.Storage.Controllers
             return Ok(new ExistCheckMultipleContentIdResponse { Needs = partialContentIds.ToArray()});
         }
 
-        private async Task<BlobContents> GetImpl(NamespaceId ns, ContentId contentId)
+        private async Task<(BlobContents, string)> GetImpl(NamespaceId ns, ContentId contentId)
         {
-            BlobIdentifier[]? chunks = await _contentIdStore.Resolve(ns, contentId, mustBeContentId: true);
+            BlobIdentifier[]? chunks = await _contentIdStore.Resolve(ns, contentId, mustBeContentId: false);
             if (chunks == null || chunks.Length == 0)
             {
                 throw new ContentIdResolveException(contentId);
@@ -172,7 +179,15 @@ namespace Horde.Storage.Controllers
             // single chunk, we just return that chunk
             if (chunks.Length == 1)
             {
-                return await _storage.GetObject(ns, chunks[0]);
+                BlobIdentifier blobToReturn = chunks[0];
+                string mimeType = CustomMediaTypeNames.UnrealCompressedBuffer;
+                if (contentId.Equals(blobToReturn))
+                {
+                    // this was actually the unmapped blob, meaning its not a compressed buffer
+                    mimeType = MediaTypeNames.Application.Octet;
+                }
+
+                return (await _storage.GetObject(ns, blobToReturn), mimeType);
             }
 
             // chunked content, combine the chunks into a single stream
@@ -193,7 +208,8 @@ namespace Horde.Storage.Controllers
 
             ms.Seek(0, SeekOrigin.Begin);
 
-            return new BlobContents(ms, ms.Length);
+            // chunking could not have happened for a non compressed buffer so assume it is compressed
+            return (new BlobContents(ms, ms.Length), CustomMediaTypeNames.UnrealCompressedBuffer);
         }
 
         [HttpPut("{ns}/{id}")]
diff --git a/Engine/Source/Programs/Horde/HordeStorage/Horde.Storage/Tests/Functional/ContentId/CompressedBlobTests.cs b/Engine/Source/Programs/Horde/HordeStorage/Horde.Storage/Tests/Functional/ContentId/CompressedBlobTests.cs
--- a/Engine/Source/Programs/Horde/HordeStorage/Horde.Storage/Tests/Functional/ContentId/CompressedBlobTests.cs
+++ b/Engine/Source/Programs/Horde/HordeStorage/Horde.Storage/Tests/Functional/ContentId/CompressedBlobTests.cs
@@ -7,6 +7,7 @@ using System.Net;
 using System.Net.Http;
 using System.Net.Http.Headers;
 using System.Net.Mime;
+using System.Text;
 using System.Threading.Tasks;
 using Horde.Storage.FunctionalTests.Storage;
 using Horde.Storage.Implementation;
@@ -171,6 +172,7 @@ namespace Horde.Storage.FunctionalTests.CompressedBlobs
             {
                 HttpResponseMessage result = await _httpClient!.GetAsync($"api/v1/compressed-blobs/{TestNamespace}/{uncompressedPayloadIdentifier}");
                 result.EnsureSuccessStatusCode();
+                Assert.AreEqual(CustomMediaTypeNames.UnrealCompressedBuffer, result.Content.Headers.ContentType!.MediaType);
 
                 byte[] blobContent = await result.Content.ReadAsByteArrayAsync();
                 CollectionAssert.AreEqual(texturePayload, blobContent);
@@ -180,6 +182,7 @@ namespace Horde.Storage.FunctionalTests.CompressedBlobs
                 // verify the compressed blob can be retrieved in the blob store
                 HttpResponseMessage result = await _httpClient!.GetAsync($"api/v1/blobs/{TestNamespace}/{compressedPayloadIdentifier}");
                 result.EnsureSuccessStatusCode();
+                Assert.AreEqual(MediaTypeNames.Application.Octet, result.Content.Headers.ContentType!.MediaType);
 
                 byte[] blobContent = await result.Content.ReadAsByteArrayAsync();
                 CollectionAssert.AreEqual(texturePayload, blobContent);
@@ -208,5 +211,60 @@ namespace Horde.Storage.FunctionalTests.CompressedBlobs
                 Assert.AreEqual(HttpStatusCode.BadRequest, result.StatusCode);
             }
         }
+
+
+        [TestMethod]
+        public async Task GetUncompressedContent()
+        {
+            string stringContent = "this is just a random string";
+            byte[] payload = Encoding.ASCII.GetBytes(stringContent);
+            BlobIdentifier blobIdentifier = BlobIdentifier.FromBlob(payload);
+
+            // upload a uncompressed blob
+            {
+                ByteArrayContent content = new ByteArrayContent(payload);
+                content.Headers.ContentType = new MediaTypeHeaderValue(MediaTypeNames.Application.Octet);
+                HttpResponseMessage result = await _httpClient!.PutAsync($"api/v1/blobs/{TestNamespace}/{blobIdentifier}", content);
+                result.EnsureSuccessStatusCode();
+            }
+
+
+            {
+                HttpResponseMessage result = await _httpClient!.GetAsync($"api/v1/compressed-blobs/{TestNamespace}/{blobIdentifier}");
+                result.EnsureSuccessStatusCode();
+
+                // verify that we return the uncompressed content but also that the media type indicates that
+                byte[] blobContent = await result.Content.ReadAsByteArrayAsync();
+                Assert.AreEqual(MediaTypeNames.Application.Octet, result.Content.Headers.ContentType!.MediaType);
+                CollectionAssert.AreEqual(payload, blobContent);
+            }
+        }
+
+        
+        [TestMethod]
+        public async Task GetUncompressedContentAsCompressedBuffer()
+        {
+            string stringContent = "this is just a random string";
+            byte[] payload = Encoding.ASCII.GetBytes(stringContent);
+            BlobIdentifier blobIdentifier = BlobIdentifier.FromBlob(payload);
+
+            // upload a uncompressed blob
+            {
+                ByteArrayContent content = new ByteArrayContent(payload);
+                content.Headers.ContentType = new MediaTypeHeaderValue(MediaTypeNames.Application.Octet);
+                HttpResponseMessage result = await _httpClient!.PutAsync($"api/v1/blobs/{TestNamespace}/{blobIdentifier}", content);
+                result.EnsureSuccessStatusCode();
+            }
+
+
+            {
+                HttpRequestMessage request = new HttpRequestMessage(HttpMethod.Get, $"api/v1/compressed-blobs/{TestNamespace}/{blobIdentifier}");
+                request.Headers.Add("Accept", CustomMediaTypeNames.UnrealCompactBinary);
+
+                // asking for a compressed buffer when we only have the uncompressed content results in a 415 (unsupported media type) as transcoding this is to complicated when the compressed buffer header requires the full blake3 hash of the content
+                HttpResponseMessage result = await _httpClient!.SendAsync(request);
+                Assert.AreEqual(HttpStatusCode.UnsupportedMediaType, result.StatusCode);
+            }
+        }
     }
 }
diff --git a/Engine/Source/Programs/Horde/HordeStorage/Horde.Storage/Tests/Functional/Objects/ReferencesTests.cs b/Engine/Source/Programs/Horde/HordeStorage/Horde.Storage/Tests/Functional/Objects/ReferencesTests.cs
--- a/Engine/Source/Programs/Horde/HordeStorage/Horde.Storage/Tests/Functional/Objects/ReferencesTests.cs
+++ b/Engine/Source/Programs/Horde/HordeStorage/Horde.Storage/Tests/Functional/Objects/ReferencesTests.cs
@@ -71,6 +71,7 @@ namespace Horde.Storage.FunctionalTests.References
             ISession replicatedKeyspace = scyllaSessionManager.GetSessionForReplicatedKeyspace();
             await replicatedKeyspace.ExecuteAsync(new SimpleStatement("DROP TABLE IF EXISTS objects"));
             await replicatedKeyspace.ExecuteAsync(new SimpleStatement("DROP TABLE IF EXISTS object_last_access"));
+            await replicatedKeyspace.ExecuteAsync(new SimpleStatement("DROP TABLE IF EXISTS content_id"));
 
             // we need to clear out the state we modified in the replication log, otherwise the replication log tests will fail
             ISession localKeyspace = scyllaSessionManager.GetSessionForLocalKeyspace();
