Horde Storage - Fixed issue were Accept: */* didnt work for compressed-blobs GET

diff --git a/Engine/Source/Programs/Horde/HordeStorage/Horde.Storage/Controllers/CompressedBlobController.cs b/Engine/Source/Programs/Horde/HordeStorage/Horde.Storage/Controllers/CompressedBlobController.cs
--- a/Engine/Source/Programs/Horde/HordeStorage/Horde.Storage/Controllers/CompressedBlobController.cs
+++ b/Engine/Source/Programs/Horde/HordeStorage/Horde.Storage/Controllers/CompressedBlobController.cs
@@ -72,7 +72,7 @@ namespace Horde.Storage.Controllers
                 (BlobContents blobContents, string mediaType) = await GetImpl(ns, id);
 
                 StringValues acceptHeader = Request.Headers["Accept"];
-                if (acceptHeader.Count != 0 && !acceptHeader.Contains(mediaType))
+                if (!acceptHeader.Contains("*/*") && acceptHeader.Count != 0 && !acceptHeader.Contains(mediaType))
                     return new UnsupportedMediaTypeResult();
 
                 return File(blobContents.Stream, mediaType, enableRangeProcessing: true);
diff --git a/Engine/Source/Programs/Horde/HordeStorage/Horde.Storage/Tests/Functional/ContentId/CompressedBlobTests.cs b/Engine/Source/Programs/Horde/HordeStorage/Horde.Storage/Tests/Functional/ContentId/CompressedBlobTests.cs
--- a/Engine/Source/Programs/Horde/HordeStorage/Horde.Storage/Tests/Functional/ContentId/CompressedBlobTests.cs
+++ b/Engine/Source/Programs/Horde/HordeStorage/Horde.Storage/Tests/Functional/ContentId/CompressedBlobTests.cs
@@ -265,6 +265,19 @@ namespace Horde.Storage.FunctionalTests.CompressedBlobs
                 HttpResponseMessage result = await _httpClient!.SendAsync(request);
                 Assert.AreEqual(HttpStatusCode.UnsupportedMediaType, result.StatusCode);
             }
+
+            {
+                // ask for any accept type
+                HttpRequestMessage request = new HttpRequestMessage(HttpMethod.Get, $"api/v1/compressed-blobs/{TestNamespace}/{blobIdentifier}");
+                request.Headers.Add("Accept", "*/*");
+
+                HttpResponseMessage result = await _httpClient!.SendAsync(request);
+                result.EnsureSuccessStatusCode();
+                Assert.AreEqual(MediaTypeNames.Application.Octet, result.Content.Headers.ContentType!.MediaType);
+
+                byte[] blobContent = await result.Content.ReadAsByteArrayAsync();
+                CollectionAssert.AreEqual(payload, blobContent);
+            }
         }
     }
 }
