Horde Storage - Improvements to Horde Storage to handle Azure deployments better.
Changed contact points into connection string for scylla.

diff --git a/Engine/Source/Programs/Horde/HordeStorage/Horde.Storage/Horde.Storage.Startup.cs b/Engine/Source/Programs/Horde/HordeStorage/Horde.Storage/Horde.Storage.Startup.cs
--- a/Engine/Source/Programs/Horde/HordeStorage/Horde.Storage/Horde.Storage.Startup.cs
+++ b/Engine/Source/Programs/Horde/HordeStorage/Horde.Storage/Horde.Storage.Startup.cs
@@ -4,6 +4,9 @@ using System;
 using System.Collections.Generic;
 using System.IO;
 using System.Linq;
+using System.Net;
+using System.Net.Security;
+using System.Security.Authentication;
 using Amazon.DAX;
 using Amazon.DynamoDBv2;
 using Amazon.Extensions.NETCore.Setup;
@@ -54,74 +57,110 @@ namespace Horde.Storage
             authorizationOptions.AddPolicy("Cache.read", policy =>
             {
                 policy.AuthenticationSchemes = defaultSchemes;
-                policy.RequireClaim("Cache", "read", "readwrite", "full");
+                policy.BeginAnyContext()
+                    .IncludeClaim("Cache", "read", "readwrite", "full")
+                    .IncludeRole("Cache.read", "Cache.readwrite", "Cache.full")
+                    .AddAny();
             });
 
             authorizationOptions.AddPolicy("Cache.write", policy =>
             {
                 policy.AuthenticationSchemes = defaultSchemes;
-                policy.RequireClaim("Cache", "write", "readwrite", "full");
+                policy.BeginAnyContext()
+                    .IncludeClaim("Cache", "write", "readwrite", "full")
+                    .IncludeRole("Cache.write", "Cache.readwrite", "Cache.full")
+                    .AddAny();
             });
 
             authorizationOptions.AddPolicy("Cache.delete", policy =>
             {
                 policy.AuthenticationSchemes = defaultSchemes;
-                policy.RequireClaim("Cache", "delete", "full");
+                policy.BeginAnyContext()
+                    .IncludeClaim("Cache", "delete", "full")
+                    .IncludeRole("Cache.delete", "Cache.full")
+                    .AddAny();
             });
 
             authorizationOptions.AddPolicy("Object.read", policy =>
             {
                 policy.AuthenticationSchemes = defaultSchemes;
-                policy.RequireClaim("Cache", "read", "readwrite", "full");
+                policy.BeginAnyContext()
+                    .IncludeClaim("Cache", "read", "readwrite", "full")
+                    .IncludeRole("Cache.read", "Cache.readwrite", "Cache.full")
+                    .AddAny();
             });
 
             authorizationOptions.AddPolicy("Object.write", policy =>
             {
                 policy.AuthenticationSchemes = defaultSchemes;
-                policy.RequireClaim("Cache", "write", "readwrite", "full");
+                policy.BeginAnyContext()
+                    .IncludeClaim("Cache", "write", "readwrite", "full")
+                    .IncludeRole("Cache.write", "Cache.readwrite", "Cache.full")
+                    .AddAny();
             });
 
             authorizationOptions.AddPolicy("Object.delete", policy =>
             {
                 policy.AuthenticationSchemes = defaultSchemes;
-                policy.RequireClaim("Cache", "delete", "full");
+                policy.BeginAnyContext()
+                    .IncludeClaim("Cache", "delete", "full")
+                    .IncludeRole("Cache.delete", "Cache.full")
+                    .AddAny();
             });
 
             authorizationOptions.AddPolicy("Storage.read", policy =>
             {
                 policy.AuthenticationSchemes = defaultSchemes;
-                policy.RequireClaim("Storage", "read", "readwrite", "full");
+                policy.BeginAnyContext()
+                    .IncludeClaim("Storage", "read", "readwrite", "full")
+                    .IncludeRole("Storage.read", "Storage.readwrite", "Storage.full")
+                    .AddAny();
             });
 
             authorizationOptions.AddPolicy("Storage.write", policy =>
             {
                 policy.AuthenticationSchemes = defaultSchemes;
-                policy.RequireClaim("Storage", "write", "readwrite", "full");
+                policy.BeginAnyContext()
+                    .IncludeClaim("Storage", "write", "readwrite", "full")
+                    .IncludeRole("Storage.write", "Storage.readwrite", "Storage.full")
+                    .AddAny();
             });
 
             authorizationOptions.AddPolicy("Storage.delete", policy =>
             {
                 policy.AuthenticationSchemes = defaultSchemes;
-                policy.RequireClaim("Storage", "delete", "full");
+                policy.BeginAnyContext()
+                    .IncludeClaim("Storage", "delete", "full")
+                    .IncludeRole("Storage.delete", "Storage.full")
+                    .AddAny();
             });
 
             authorizationOptions.AddPolicy("replication-log.read", policy =>
             {
                 policy.AuthenticationSchemes = defaultSchemes;
-                policy.RequireClaim("transactionLog", "read", "readwrite", "full");
+                policy.BeginAnyContext()
+                    .IncludeClaim("transactionLog", "read", "readwrite", "full")
+                    .IncludeRole("transactionLog.read", "transactionLog.readwrite", "transactionLog.full")
+                    .AddAny();
             });
 
             authorizationOptions.AddPolicy("replication-log.write", policy =>
             {
                 policy.AuthenticationSchemes = defaultSchemes;
-                policy.RequireClaim("transactionLog", "write", "readwrite", "full");
+                policy.BeginAnyContext()
+                    .IncludeClaim("transactionLog", "write", "readwrite", "full")
+                    .IncludeRole("transactionLog.write", "transactionLog.readwrite", "transactionLog.full")
+                    .AddAny();
             });
 
 
             authorizationOptions.AddPolicy("Admin", policy =>
             {
                 policy.AuthenticationSchemes = defaultSchemes;
-                policy.RequireClaim("Admin");
+                policy.BeginAnyContext()
+                    .IncludeClaim("Admin")
+                    .IncludeRole("Admin")
+                    .AddAny();
             });
 
             // A policy that grants any authenticated user access
@@ -334,20 +373,70 @@ namespace Horde.Storage
         private IScyllaSessionManager ScyllaFactory(IServiceProvider provider)
         {
             ScyllaSettings settings = provider.GetService<IOptionsMonitor<ScyllaSettings>>()!.CurrentValue!;
+            ISecretResolver secretResolver = provider.GetService<ISecretResolver>()!;
 
             Serilog.Extensions.Logging.SerilogLoggerProvider serilogLoggerProvider = new();
             Diagnostics.AddLoggerProvider(serilogLoggerProvider);
-
             const string DefaultKeyspaceName = "jupiter";
+
+            string? connectionString = secretResolver.Resolve(settings.ConnectionString);
+            if (string.IsNullOrEmpty(connectionString))
+            {
+                throw new Exception("Connection string has to be specified");
+            }
+
             // Configure the builder with your cluster's contact points
-            var cluster = Cluster.Builder()
-                .AddContactPoints(settings.ContactPoints)
+            Builder clusterBuilder = Cluster.Builder()
+                .WithDefaultKeyspace(DefaultKeyspaceName)
+                .WithConnectionString(connectionString)
                 .WithLoadBalancingPolicy(new DefaultLoadBalancingPolicy(settings.LocalDatacenterName))
                 .WithPoolingOptions(PoolingOptions.Create().SetMaxConnectionsPerHost(HostDistance.Local, settings.MaxConnectionForLocalHost))
-                .WithDefaultKeyspace(DefaultKeyspaceName)
-                .WithExecutionProfiles(options => 
-                    options.WithProfile("default", builder => builder.WithConsistencyLevel(ConsistencyLevel.LocalOne)))
-                .Build();
+                .WithExecutionProfiles(options =>
+                    options.WithProfile("default", builder => builder.WithConsistencyLevel(ConsistencyLevel.LocalOne)));
+
+
+            if (settings.UseAzureCosmosDB)
+            {
+                CassandraConnectionStringBuilder connectionStringBuilder = new CassandraConnectionStringBuilder(connectionString);
+                connectionStringBuilder.DefaultKeyspace = DefaultKeyspaceName;
+                string[] contactPoints = connectionStringBuilder.ContactPoints;
+
+                // Connect to cassandra cluster using TLSv1.2.
+                var sslOptions = new SSLOptions(SslProtocols.Tls12, false,
+                    (sender, certificate, chain, sslPolicyErrors) =>
+                    {
+                        if (sslPolicyErrors == SslPolicyErrors.None)
+                            return true;
+
+                        _logger.Error("Certificate error: {0}", sslPolicyErrors);
+                        // Do not allow this client to communicate with unauthenticated servers.
+                        return false;
+                    });
+
+                // Prepare a map to resolve the host name from the IP address.
+                var hostNameByIp = new Dictionary<IPAddress, string>();
+                foreach (string contactPoint in contactPoints)
+                {
+                    IPAddress[] resolvedIps = Dns.GetHostAddresses(contactPoint);
+                    foreach (IPAddress resolvedIp in resolvedIps)
+                    {
+                        hostNameByIp[resolvedIp] = contactPoint;
+                    }
+                }
+
+                sslOptions.SetHostNameResolver((ipAddress) =>
+                {
+                    if (hostNameByIp.TryGetValue(ipAddress, out string? resolvedName))
+                    {
+                        return resolvedName;
+                    }
+                    var hostEntry = Dns.GetHostEntry(ipAddress.ToString());
+                    return hostEntry.HostName;
+                });
+
+                clusterBuilder = clusterBuilder.WithSSL(sslOptions);
+            }
+            Cluster cluster = clusterBuilder.Build();
 
             Dictionary<string, string> replicationStrategy = ReplicationStrategies.CreateSimpleStrategyReplicationProperty(2);
             if (settings.KeyspaceReplicationStrategy != null)
@@ -603,7 +692,7 @@ namespace Horde.Storage
                         break;
                     case HordeStorageSettings.StorageBackendImplementations.Azure:
                         AzureSettings azureSettings = provider.GetService<IOptionsMonitor<AzureSettings>>()!.CurrentValue;
-                        healthChecks.AddAzureBlobStorage(azureSettings.ConnectionString, tags: new[] {"services"});
+                        healthChecks.AddAzureBlobStorage(AzureBlobStore.GetConnectionString(azureSettings, provider), tags: new[] {"services"});
                         break;
 
                     case HordeStorageSettings.StorageBackendImplementations.FileSystem:
diff --git a/Engine/Source/Programs/Horde/HordeStorage/Horde.Storage/Horde.Storage.csproj b/Engine/Source/Programs/Horde/HordeStorage/Horde.Storage/Horde.Storage.csproj
--- a/Engine/Source/Programs/Horde/HordeStorage/Horde.Storage/Horde.Storage.csproj
+++ b/Engine/Source/Programs/Horde/HordeStorage/Horde.Storage/Horde.Storage.csproj
@@ -77,6 +77,8 @@
     <PackageReference Include="AWSSDK.DynamoDBv2" Version="3.7.0.16" />
     <PackageReference Include="AWSSDK.Extensions.NETCore.Setup" Version="3.7.0.1" />
     <PackageReference Include="AWSSDK.S3" Version="3.7.0.17" />
+    <PackageReference Include="Azure.Identity" Version="1.5.0" />
+    <PackageReference Include="Azure.Security.KeyVault.Secrets" Version="4.2.0" />
     <PackageReference Include="Azure.Storage.Blobs" Version="12.8.0" />
     <PackageReference Include="Destructurama.Attributed" Version="2.0.0" />
     <PackageReference Include="KubernetesClient" Version="4.0.10" />
diff --git a/Engine/Source/Programs/Horde/HordeStorage/Horde.Storage/Implementation/AnyRequirement.cs b/Engine/Source/Programs/Horde/HordeStorage/Horde.Storage/Implementation/AnyRequirement.cs
new file mode 100644
--- /dev/null
+++ b/Engine/Source/Programs/Horde/HordeStorage/Horde.Storage/Implementation/AnyRequirement.cs
@@ -0,0 +1,103 @@
+// Copyright Epic Games, Inc. All Rights Reserved.
+
+using System;
+using System.Collections.Generic;
+using System.Threading.Tasks;
+using Microsoft.AspNetCore.Authorization;
+using Microsoft.AspNetCore.Authorization.Infrastructure;
+
+namespace Horde.Storage
+{
+    /// <summary>
+    /// An implementation of <see cref="IAuthorizationRequirement"/> that combines multiple
+    /// requirements such that if any of those inner requirements are met, this requirement is met.
+    /// This only supports requirements that handle themselves, i.e. implement both <see cref="IAuthorizationRequirement"/>
+    /// and the corresponding <see cref="AuthorizationHandler{TRequirement}"/>.
+    /// </summary>
+    public class AnyRequirement : AuthorizationHandler<AnyRequirement>, IAuthorizationRequirement
+    {
+        public AnyRequirement(params IAuthorizationRequirement[] innerRequirements)
+        {
+            InnerRequirements = innerRequirements ?? throw new ArgumentNullException(nameof(innerRequirements));
+        }
+
+        public IEnumerable<IAuthorizationRequirement> InnerRequirements { get; }
+
+        protected override async Task HandleRequirementAsync(AuthorizationHandlerContext context, AnyRequirement requirement)
+        {
+            foreach (IAuthorizationRequirement innerRequirement in requirement.InnerRequirements)
+            {
+                // This only supports requirements that "handle themselves".
+                IAuthorizationHandler? handler = innerRequirement as IAuthorizationHandler;
+                if (handler != null)
+                {
+                    var innerContext = new AuthorizationHandlerContext(new[] { innerRequirement }, context.User, context.Resource);
+                    await handler.HandleAsync(innerContext);
+                    if (innerContext.HasSucceeded)
+                    {
+                        context.Succeed(this);
+                        break;
+                    }
+                }
+            }
+        }
+    }
+
+    /// <summary>
+    /// A helper class to build an instance of <see cref="AnyRequirement"/>.
+    /// </summary>
+    public class AnyRequirementBuilder
+    {
+        private readonly List<IAuthorizationRequirement> _innerRequirements = new List<IAuthorizationRequirement>();
+        private readonly AuthorizationPolicyBuilder _policyBuilder;
+
+        public AnyRequirementBuilder(AuthorizationPolicyBuilder policyBuilder)
+        {
+            _policyBuilder = policyBuilder;
+        }
+
+        public AnyRequirementBuilder AddRequirement(IAuthorizationRequirement requirement)
+        {
+            _innerRequirements.Add(requirement);
+            return this;
+        }
+
+        public AuthorizationPolicyBuilder AddAny()
+        {
+            var anyRequirement = new AnyRequirement(_innerRequirements.ToArray());
+            return _policyBuilder.AddRequirements(anyRequirement);
+        }
+    }
+
+    /// <summary>
+    /// Helpful extensions related to <see cref="AnyRequirement"/>
+    /// </summary>
+    public static class AnyRequirementExtensions
+    {
+        /// <summary>
+        /// Starts a building context to add requirements that will be combined as an <see cref="AnyRequirement"/>
+        /// and added to the policy builder when <see cref="AnyRequirementBuilder.AddAny"/> is called later.
+        /// </summary>
+        /// <param name="policyBuilder"></param>
+        /// <returns></returns>
+        public static AnyRequirementBuilder BeginAnyContext(this AuthorizationPolicyBuilder policyBuilder)
+        {
+            return new AnyRequirementBuilder(policyBuilder);
+        }
+
+        public static AnyRequirementBuilder IncludeClaim(this AnyRequirementBuilder builder, string claimType, params string[] allowedValues)
+        {
+            return builder.AddRequirement(new ClaimsAuthorizationRequirement(claimType, allowedValues));
+        }
+
+        public static AnyRequirementBuilder IncludeClaim(this AnyRequirementBuilder builder, string claimType)
+        {
+            return builder.AddRequirement(new ClaimsAuthorizationRequirement(claimType, allowedValues: null));
+        }
+
+        public static AnyRequirementBuilder IncludeRole(this AnyRequirementBuilder builder, params string[] roles)
+        {
+            return builder.AddRequirement(new RolesAuthorizationRequirement(roles));
+        }
+    }
+}
diff --git a/Engine/Source/Programs/Horde/HordeStorage/Horde.Storage/Implementation/Blob/AzureBlobStore.cs b/Engine/Source/Programs/Horde/HordeStorage/Horde.Storage/Implementation/Blob/AzureBlobStore.cs
--- a/Engine/Source/Programs/Horde/HordeStorage/Horde.Storage/Implementation/Blob/AzureBlobStore.cs
+++ b/Engine/Source/Programs/Horde/HordeStorage/Horde.Storage/Implementation/Blob/AzureBlobStore.cs
@@ -6,12 +6,14 @@ using System.Globalization;
 using System.IO;
 using System.Threading.Tasks;
 using Azure;
+using Microsoft.Extensions.DependencyInjection;
 using Microsoft.Extensions.Options;
 using Azure.Storage.Blobs;
 using Azure.Storage.Blobs.Models;
 using Jupiter.Implementation;
 using Serilog;
 using EpicGames.Horde.Storage;
+using Jupiter.Common.Implementation;
 
 namespace Horde.Storage.Implementation
 {
@@ -23,9 +25,32 @@ namespace Horde.Storage.Implementation
         private const string LastTouchedKey = "Io_LastTouched";
         private const string NamespaceKey = "Io_Namespace";
 
-        public AzureBlobStore(IOptionsMonitor<AzureSettings> settings)
+        public AzureBlobStore(IOptionsMonitor<AzureSettings> settings, IServiceProvider provider)
         {
-            _connectionString = settings.CurrentValue.ConnectionString;
+            _connectionString = GetConnectionString(settings.CurrentValue, provider);
+        }
+
+        /// <summary>
+        /// Gets the connection string for Azure storage.
+        /// If a key vault secret is used, the value is cached in <see cref="AzureSettings.ConnectionString"/>
+        /// for next time.
+        /// </summary>
+        /// <param name="settings"></param>
+        /// <param name="provider"></param>
+        /// <returns></returns>
+        public static string GetConnectionString(AzureSettings settings, IServiceProvider provider)
+        {
+            if (!string.IsNullOrEmpty(settings.ConnectionString))
+            {
+                return settings.ConnectionString;
+            }
+
+            // Cache the connection string in the settings for next time.
+            ISecretResolver secretResolver = provider.GetService<ISecretResolver>()!;
+            string connectionString = secretResolver.Resolve(settings.ConnectionString)!;
+            settings.ConnectionString = connectionString;
+
+            return connectionString;
         }
 
         public async Task<BlobIdentifier> PutObject(NamespaceId ns, ReadOnlyMemory<byte> content, BlobIdentifier blobIdentifier)
diff --git a/Engine/Source/Programs/Horde/HordeStorage/Horde.Storage/Implementation/Objects/ScyllaReferencesStore.cs b/Engine/Source/Programs/Horde/HordeStorage/Horde.Storage/Implementation/Objects/ScyllaReferencesStore.cs
--- a/Engine/Source/Programs/Horde/HordeStorage/Horde.Storage/Implementation/Objects/ScyllaReferencesStore.cs
+++ b/Engine/Source/Programs/Horde/HordeStorage/Horde.Storage/Implementation/Objects/ScyllaReferencesStore.cs
@@ -378,6 +378,7 @@ namespace Horde.Storage.Implementation
 
     public class ScyllaSettings
     {
+        public string ConnectionString { get; set; } = "Contact Points=localhost,scylla;Default Keyspace=jupiter";
         public long InlineBlobMaxSize { get; set; } = 32 * 1024; // default to 32 kb blobs max
         public string[] ContactPoints { get; set; } = new string[] {"localhost", "scylla"};
         public int MaxSnapshotsPerNamespace { get; set; } = 10;
@@ -415,6 +416,11 @@ namespace Horde.Storage.Implementation
         /// The time for a replication log event to live in the incremental state before being deleted, assumption is that the snapshot will have processed the event within this time
         /// </summary>
         public TimeSpan ReplicationLogTimeToLive { get; set; } = TimeSpan.FromDays(7);
+
+        /// <summary>
+        /// Whether a Cosmos DB Cassandra layer from Azure should be used.
+        /// </summary>
+        public bool UseAzureCosmosDB { get; set; } = false;
     }
 
     [Cassandra.Mapping.Attributes.Table("objects")]
diff --git a/Engine/Source/Programs/Horde/HordeStorage/Horde.Storage/Settings.cs b/Engine/Source/Programs/Horde/HordeStorage/Horde.Storage/Settings.cs
--- a/Engine/Source/Programs/Horde/HordeStorage/Horde.Storage/Settings.cs
+++ b/Engine/Source/Programs/Horde/HordeStorage/Horde.Storage/Settings.cs
@@ -217,7 +217,7 @@ namespace Horde.Storage
 
     public class AzureSettings
     {
-        [Required] public string ConnectionString { get; set; } = "";
+        [Required] public string ConnectionString { get; set; } = string.Empty;
     }
 
     public class FilesystemSettings
diff --git a/Engine/Source/Programs/Horde/HordeStorage/Horde.Storage/Tests/Functional/StorageTests.cs b/Engine/Source/Programs/Horde/HordeStorage/Horde.Storage/Tests/Functional/StorageTests.cs
--- a/Engine/Source/Programs/Horde/HordeStorage/Horde.Storage/Tests/Functional/StorageTests.cs
+++ b/Engine/Source/Programs/Horde/HordeStorage/Horde.Storage/Tests/Functional/StorageTests.cs
@@ -61,7 +61,7 @@ namespace Horde.Storage.FunctionalTests.Storage
             if (await _s3!.DoesS3BucketExistAsync(S3BucketName))
             {
                 // if we have failed to run the cleanup for some reason we run it now
-                await Teardown();
+                await Teardown(provider);
             }
 
             await _s3.PutBucketAsync(S3BucketName);
@@ -71,7 +71,7 @@ namespace Horde.Storage.FunctionalTests.Storage
             await _s3.PutObjectAsync(new PutObjectRequest {BucketName = S3BucketName, Key = _oldBlobFileHash.AsS3Key(), ContentBody = OldFileContents});
         }
 
-        protected override async Task Teardown()
+        protected override async Task Teardown(IServiceProvider provider)
         {
             string S3BucketName = $"tests-{TestNamespaceName}";
             ListObjectsResponse response = await _s3!.ListObjectsAsync(S3BucketName);
@@ -96,13 +96,13 @@ namespace Horde.Storage.FunctionalTests.Storage
         {
             _settings = provider.GetService<IOptionsMonitor<AzureSettings>>()!.CurrentValue;
 
-            string connectionString = _settings.ConnectionString;
+            string connectionString = AzureBlobStore.GetConnectionString(_settings, provider);
             BlobContainerClient container = new BlobContainerClient(connectionString, TestNamespaceName.ToString());
 
             if (await container.ExistsAsync())
             {
                 // if we have failed to run the cleanup for some reason we run it now
-                await Teardown();
+                await Teardown(provider);
             }
 
             await container.CreateAsync();
@@ -127,9 +127,9 @@ namespace Horde.Storage.FunctionalTests.Storage
             await oldBlob.UploadAsync(oldBlobContentsSteam);
         }
 
-        protected override async Task Teardown()
+        protected override async Task Teardown(IServiceProvider provider)
         {
-            string connectionString = _settings!.ConnectionString;
+            string connectionString = AzureBlobStore.GetConnectionString( _settings!, provider);
             BlobContainerClient container = new BlobContainerClient(connectionString, TestNamespaceName.ToString());
 
             await container.DeleteAsync();
@@ -193,7 +193,7 @@ namespace Horde.Storage.FunctionalTests.Storage
             File.SetLastWriteTimeUtc(oldFileInfo.FullName, DateTime.Now.AddDays(-7));
         }
 
-        protected override Task Teardown()
+        protected override Task Teardown(IServiceProvider provider)
         {
             Directory.Delete(_localTestDir, true);
 
@@ -361,13 +361,13 @@ namespace Horde.Storage.FunctionalTests.Storage
         protected abstract IEnumerable<KeyValuePair<string, string>> GetSettings();
 
         protected abstract Task Seed(IServiceProvider serverServices);
-        protected abstract Task Teardown();
+        protected abstract Task Teardown(IServiceProvider serverServices);
 
 
         [TestCleanup]
         public async Task MyTeardown()
         {
-            await Teardown();
+            await Teardown(_server!.Services);
         }
 
 
diff --git a/Engine/Source/Programs/Horde/HordeStorage/Jupiter.Common/Implementation/SecretResolver.cs b/Engine/Source/Programs/Horde/HordeStorage/Jupiter.Common/Implementation/SecretResolver.cs
--- a/Engine/Source/Programs/Horde/HordeStorage/Jupiter.Common/Implementation/SecretResolver.cs
+++ b/Engine/Source/Programs/Horde/HordeStorage/Jupiter.Common/Implementation/SecretResolver.cs
@@ -2,10 +2,14 @@
 
 using System;
 using System.Collections.Generic;
+using System.IO;
 using System.Text.Json;
 using System.Threading.Tasks;
 using Amazon.SecretsManager;
 using Amazon.SecretsManager.Model;
+using Azure;
+using Azure.Identity;
+using Azure.Security.KeyVault.Secrets;
 using Microsoft.Extensions.DependencyInjection;
 
 namespace Jupiter.Common.Implementation
@@ -43,6 +47,8 @@ namespace Jupiter.Common.Implementation
             {
                 case "aws":
                     return ResolveAWSSecret(providerPath);
+                case "akv":
+                    return ResolveAKVSecret(providerPath);
                 default:
                     // no provider matches so just return the original value
                     return originalValue;
@@ -51,14 +57,7 @@ namespace Jupiter.Common.Implementation
 
         private string? ResolveAWSSecret(string providerPath)
         {
-            int keySeparator = providerPath.IndexOf("|");
-            string? key = null;
-            string arn = providerPath;
-            if (keySeparator != -1)
-            {
-                arn = providerPath.Substring(0, keySeparator);
-                key = providerPath.Substring(keySeparator + 1);
-            }
+            SplitByFirstSeparator(providerPath, '|', out string? key, out string arn);
 
             IAmazonSecretsManager? secretsManager = _serviceProvider.GetService<IAmazonSecretsManager>();
             if (secretsManager == null)
@@ -82,5 +81,42 @@ namespace Jupiter.Common.Implementation
 
             throw new Exception($"Unable to find key {key} in blob returned for secret {arn}");
         }
+
+        /// <summary>
+        /// AKV = Azure Key Vault.
+        /// Uses <see cref="DefaultAzureCredential"/> to get a secret from Azure Key Vault.
+        /// Using <see cref="DefaultAzureCredential"/> lets us debug using <see cref="VisualStudioCredential"/>
+        /// and deploy to an app service that uses <see cref="ManagedIdentityCredential"/>.
+        /// </summary>
+        /// <param name="providerPath"></param>
+        /// <returns></returns>
+        /// <exception cref="Exception"></exception>
+        private string? ResolveAKVSecret(string providerPath)
+        {
+            if (!SplitByFirstSeparator(providerPath, '|', out string? vaultName, out string secretName))
+            {
+                throw new InvalidDataException("Azure Key Vault secret path must be of the form vaultName|secretName");
+            }
+
+            string vaultUrl = $"https://{vaultName}.vault.azure.net/";
+
+            SecretClient client = new SecretClient(new Uri(vaultUrl), new DefaultAzureCredential());
+            Response<KeyVaultSecret> response = client.GetSecret(secretName);
+            return response.Value.Value;
+        }
+
+        private static bool SplitByFirstSeparator(string fullPath, char sep, out string? left, out string right)
+        {
+            int keySeparator = fullPath.IndexOf(sep);
+            left = null;
+            right = fullPath;
+            if (keySeparator != -1)
+            {
+                left = fullPath.Substring(0, keySeparator);
+                right = fullPath.Substring(keySeparator + 1);
+                return true;
+            }
+            return false;
+        }
     }
 }
diff --git a/Engine/Source/Programs/Horde/HordeStorage/Jupiter.Common/Jupiter.Common.csproj b/Engine/Source/Programs/Horde/HordeStorage/Jupiter.Common/Jupiter.Common.csproj
--- a/Engine/Source/Programs/Horde/HordeStorage/Jupiter.Common/Jupiter.Common.csproj
+++ b/Engine/Source/Programs/Horde/HordeStorage/Jupiter.Common/Jupiter.Common.csproj
@@ -17,6 +17,8 @@
     <PackageReference Include="AWSSDK.Extensions.NETCore.Setup" Version="3.7.0.1" />
     <PackageReference Include="AWSSDK.SecretsManager" Version="3.7.1.39" />
     <PackageReference Include="AWSSDK.SecurityToken" Version="3.7.1.4" />
+    <PackageReference Include="Azure.Identity" Version="1.5.0" />
+    <PackageReference Include="Azure.Security.KeyVault.Secrets" Version="4.2.0" />
     <PackageReference Include="Blake3" Version="0.3.0" />
     <PackageReference Include="Datadog.Trace" Version="2.1.0" />
     <PackageReference Include="IdentityServer4.AccessTokenValidation" Version="3.0.1" />
