Horde Storage - Fixed issue in azure blob storage connection string resolve (always pass the string to the secret resolver)

diff --git a/Engine/Source/Programs/Horde/HordeStorage/Horde.Storage/Implementation/Blob/AzureBlobStore.cs b/Engine/Source/Programs/Horde/HordeStorage/Horde.Storage/Implementation/Blob/AzureBlobStore.cs
--- a/Engine/Source/Programs/Horde/HordeStorage/Horde.Storage/Implementation/Blob/AzureBlobStore.cs
+++ b/Engine/Source/Programs/Horde/HordeStorage/Horde.Storage/Implementation/Blob/AzureBlobStore.cs
@@ -40,11 +40,6 @@ namespace Horde.Storage.Implementation
         /// <returns></returns>
         public static string GetConnectionString(AzureSettings settings, IServiceProvider provider)
         {
-            if (!string.IsNullOrEmpty(settings.ConnectionString))
-            {
-                return settings.ConnectionString;
-            }
-
             // Cache the connection string in the settings for next time.
             ISecretResolver secretResolver = provider.GetService<ISecretResolver>()!;
             string connectionString = secretResolver.Resolve(settings.ConnectionString)!;
