Horde Storage - Fixed issue resolving AWS secrets and improved logigng when AWS secret can not be found.

diff --git a/Engine/Source/Programs/Horde/HordeStorage/Jupiter.Common/Implementation/SecretResolver.cs b/Engine/Source/Programs/Horde/HordeStorage/Jupiter.Common/Implementation/SecretResolver.cs
--- a/Engine/Source/Programs/Horde/HordeStorage/Jupiter.Common/Implementation/SecretResolver.cs
+++ b/Engine/Source/Programs/Horde/HordeStorage/Jupiter.Common/Implementation/SecretResolver.cs
@@ -11,6 +11,7 @@ using Azure;
 using Azure.Identity;
 using Azure.Security.KeyVault.Secrets;
 using Microsoft.Extensions.DependencyInjection;
+using Serilog;
 
 namespace Jupiter.Common.Implementation
 {
@@ -22,6 +23,7 @@ namespace Jupiter.Common.Implementation
     public class SecretResolver : ISecretResolver
     {
         private readonly IServiceProvider _serviceProvider;
+        private readonly ILogger _logger = Log.ForContext<SecretResolver>();
 
         public SecretResolver(IServiceProvider serviceProvider)
         {
@@ -57,18 +59,27 @@ namespace Jupiter.Common.Implementation
 
         private string? ResolveAWSSecret(string providerPath)
         {
-            SplitByFirstSeparator(providerPath, '|', out string? key, out string arn);
+            SplitByFirstSeparator(providerPath, '|', out string arn, out string? key);
 
             IAmazonSecretsManager? secretsManager = _serviceProvider.GetService<IAmazonSecretsManager>();
             if (secretsManager == null)
                 throw new Exception($"Unable to get AWSSecretsManager when resolving aws secret resource: {providerPath}");
 
-            Task<GetSecretValueResponse>? response = secretsManager.GetSecretValueAsync(new GetSecretValueRequest
+            string secretValue;
+            try
             {
-                SecretId = arn,
-            });
+                Task<GetSecretValueResponse>? response = secretsManager.GetSecretValueAsync(new GetSecretValueRequest
+                {
+                    SecretId = arn,
+                });
+                secretValue = response.Result.SecretString;
+            }
+            catch (ResourceNotFoundException e)
+            {
+                _logger.Error(e, "Failed to find AWS secret: {Arn}", arn);
+                throw;
+            }
 
-            string secretValue = response.Result.SecretString;
             if (key == null)
                 return secretValue;
 
@@ -93,7 +104,7 @@ namespace Jupiter.Common.Implementation
         /// <exception cref="Exception"></exception>
         private string? ResolveAKVSecret(string providerPath)
         {
-            if (!SplitByFirstSeparator(providerPath, '|', out string? vaultName, out string secretName))
+            if (!SplitByFirstSeparator(providerPath, '|', out string vaultName, out string? secretName))
             {
                 throw new InvalidDataException("Azure Key Vault secret path must be of the form vaultName|secretName");
             }
@@ -105,11 +116,11 @@ namespace Jupiter.Common.Implementation
             return response.Value.Value;
         }
 
-        private static bool SplitByFirstSeparator(string fullPath, char sep, out string? left, out string right)
+        private static bool SplitByFirstSeparator(string fullPath, char sep, out string left, out string? right)
         {
             int keySeparator = fullPath.IndexOf(sep);
-            left = null;
-            right = fullPath;
+            left = fullPath;
+            right = null;
             if (keySeparator != -1)
             {
                 left = fullPath.Substring(0, keySeparator);
