Horde: Fix Horde.Agent.Tests always being built.

diff --git a/Engine/Source/Programs/Horde/Horde.Agent.Tests/Horde.Agent.Tests.csproj b/Engine/Source/Programs/Horde/Horde.Agent.Tests/Horde.Agent.Tests.csproj
--- a/Engine/Source/Programs/Horde/Horde.Agent.Tests/Horde.Agent.Tests.csproj
+++ b/Engine/Source/Programs/Horde/Horde.Agent.Tests/Horde.Agent.Tests.csproj
@@ -24,7 +24,7 @@
 
   <ItemGroup>
     <Content Include="remote-exec-test-bin\*.exe">
-      <CopyToOutputDirectory>Always</CopyToOutputDirectory>
+      <CopyToOutputDirectory>PreserveNewest</CopyToOutputDirectory>
     </Content>
   </ItemGroup>
 
