Horde: Fix paths to stage when building Docker image.

diff --git a/Engine/Source/Programs/Horde/BuildHorde.xml b/Engine/Source/Programs/Horde/BuildHorde.xml
--- a/Engine/Source/Programs/Horde/BuildHorde.xml
+++ b/Engine/Source/Programs/Horde/BuildHorde.xml
@@ -17,12 +17,12 @@
 		<Node Name="Build HordeServer">
 			<!-- Tag all the files that need to be staged to build -->
 			<Property Name="StagePaths">
+				Engine/Binaries/DotNET/EpicGames.Perforce.Native/...
 				Engine/Source/Programs/Shared/...
 				Engine/Source/Programs/Horde/...
 				Engine/Source/ThirdParty/Perforce/...
 			</Property>
 			<Tag Files="$(StagePaths)" Except=".../.vs/...;.../.git/...;.../bin/...;.../obj/..." With="#InputFiles"/>
-			<Tag Files="Engine/Source/Programs/Shared/EpicGames.Perforce.Native/bin/..." Filter="*.pdb;*.dll;*.so;*.dylib" With="#InputFiles"/>
 
 			<!-- Build the Docker image and publish it -->
 			<Docker-Build BaseDir="Engine/Source" Files="#InputFiles" Tag="hordeserver-public" DockerFile="Engine/Source/Programs/Horde/Horde.Build/Dockerfile" Arguments="--build-arg msbuild_args=&quot;$(VersionArguments)&quot;"/>
