Horde Storage - Avoid busy waiting if no references are resolved

diff --git a/Engine/Source/Programs/Horde/HordeStorage/Horde.Storage/Implementation/IReferenceResolver.cs b/Engine/Source/Programs/Horde/HordeStorage/Horde.Storage/Implementation/IReferenceResolver.cs
--- a/Engine/Source/Programs/Horde/HordeStorage/Horde.Storage/Implementation/IReferenceResolver.cs
+++ b/Engine/Source/Programs/Horde/HordeStorage/Horde.Storage/Implementation/IReferenceResolver.cs
@@ -2,6 +2,7 @@
 
 using System;
 using System.Collections.Generic;
+using System.Linq;
 using System.Threading.Tasks;
 using Datadog.Trace;
 using EpicGames.Core;
@@ -125,6 +126,12 @@ namespace Horde.Storage.Implementation
                 {
                     pendingCompactBinaryAttachments.Remove(finishedTask);
                 }
+
+                // if there are pending resolves left, wait for one of them to finish to avoid busy waiting
+                if (pendingContentIdResolves.Any() || pendingCompactBinaryAttachments.Any())
+                {
+                    await Task.WhenAny(pendingContentIdResolves.Concat<Task>(pendingCompactBinaryAttachments));
+                }
             }
 
             if (unresolvedContentIdReferences.Count != 0)
