Horde: Fix issue with aborted jobs not updating step refs with outcome
#skipci
#rnx

diff --git a/Engine/Source/Programs/Horde/Horde.Build/Jobs/IJobStepRefCollection.cs b/Engine/Source/Programs/Horde/Horde.Build/Jobs/IJobStepRefCollection.cs
--- a/Engine/Source/Programs/Horde/Horde.Build/Jobs/IJobStepRefCollection.cs
+++ b/Engine/Source/Programs/Horde/Horde.Build/Jobs/IJobStepRefCollection.cs
@@ -105,7 +105,15 @@ namespace HordeServer.Collections
 					}
 				}
 
-				await JobStepRefs.InsertOrReplaceAsync(new JobStepRefId(Job.Id, Batch.Id, Step.Id), Job.Name, NodeName, Job.StreamId, Job.TemplateId, Job.Change, Step.LogId, Batch.PoolId, Batch.AgentId, Outcome, LastSuccess, LastWarning, WaitTime, InitTime, Step.StartTimeUtc ?? DateTime.UtcNow, Step.FinishTimeUtc);
+				// If the job has been aborted, mark the step as finished and set outcome to unspecified
+				DateTime? FinishTime = Step.FinishTimeUtc;
+				if (Job.AbortedByUserId != null && Outcome == null)
+				{
+					Outcome = JobStepOutcome.Unspecified;
+					FinishTime = DateTime.UtcNow;
+				}
+
+				await JobStepRefs.InsertOrReplaceAsync(new JobStepRefId(Job.Id, Batch.Id, Step.Id), Job.Name, NodeName, Job.StreamId, Job.TemplateId, Job.Change, Step.LogId, Batch.PoolId, Batch.AgentId, Outcome, LastSuccess, LastWarning, WaitTime, InitTime, Step.StartTimeUtc ?? DateTime.UtcNow, FinishTime);
 			}
 		}
 	}
