Horde: Fix for aborted steps to avoid skewing timing info
#rnx

diff --git a/Engine/Source/Programs/Horde/Horde.Build/Jobs/IJobStepRefCollection.cs b/Engine/Source/Programs/Horde/Horde.Build/Jobs/IJobStepRefCollection.cs
--- a/Engine/Source/Programs/Horde/Horde.Build/Jobs/IJobStepRefCollection.cs
+++ b/Engine/Source/Programs/Horde/Horde.Build/Jobs/IJobStepRefCollection.cs
@@ -105,15 +105,15 @@ namespace HordeServer.Collections
 					}
 				}
 
-				// If the job has been aborted, mark the step as finished and set outcome to unspecified
+				DateTime StartTime = Step.StartTimeUtc ?? DateTime.UtcNow;
 				DateTime? FinishTime = Step.FinishTimeUtc;
-				if (Job.AbortedByUserId != null && Outcome == null)
+				if (Job.AbortedByUserId != null && FinishTime == null)
 				{
+					FinishTime = StartTime;
 					Outcome = JobStepOutcome.Unspecified;
-					FinishTime = DateTime.UtcNow;
 				}
 
-				await JobStepRefs.InsertOrReplaceAsync(new JobStepRefId(Job.Id, Batch.Id, Step.Id), Job.Name, NodeName, Job.StreamId, Job.TemplateId, Job.Change, Step.LogId, Batch.PoolId, Batch.AgentId, Outcome, LastSuccess, LastWarning, WaitTime, InitTime, Step.StartTimeUtc ?? DateTime.UtcNow, FinishTime);
+				await JobStepRefs.InsertOrReplaceAsync(new JobStepRefId(Job.Id, Batch.Id, Step.Id), Job.Name, NodeName, Job.StreamId, Job.TemplateId, Job.Change, Step.LogId, Batch.PoolId, Batch.AgentId, Outcome, LastSuccess, LastWarning, WaitTime, InitTime, StartTime, FinishTime);
 			}
 		}
 	}
diff --git a/Engine/Source/Programs/Horde/HordeDashboard/src/backend/Api.ts b/Engine/Source/Programs/Horde/HordeDashboard/src/backend/Api.ts
--- a/Engine/Source/Programs/Horde/HordeDashboard/src/backend/Api.ts
+++ b/Engine/Source/Programs/Horde/HordeDashboard/src/backend/Api.ts
@@ -66,7 +66,10 @@ export enum JobStepOutcome {
 	Warnings = "Warnings",
 
 	// Finished Succesfully
-	Success = "Success"
+	Success = "Success",
+
+	// Unspecified, skipped, aborted, etc
+	Unspecified = "Unspecified"
 }
 
 // The state of a particular run
