Horde Storage - Only use the first 20 bytes of the hash in the compressed buffer, its a IoHash and not a full Blake3 hash.

diff --git a/Engine/Source/Programs/Horde/HordeStorage/Horde.Storage/Implementation/ContentId/CompressedBufferUtils.cs b/Engine/Source/Programs/Horde/HordeStorage/Horde.Storage/Implementation/ContentId/CompressedBufferUtils.cs
--- a/Engine/Source/Programs/Horde/HordeStorage/Horde.Storage/Implementation/ContentId/CompressedBufferUtils.cs
+++ b/Engine/Source/Programs/Horde/HordeStorage/Horde.Storage/Implementation/ContentId/CompressedBufferUtils.cs
@@ -7,6 +7,7 @@ using Blake3;
 using Datadog.Trace;
 using Force.Crc32;
 using Jupiter;
+using Jupiter.Implementation;
 using Jupiter.Utils;
 using K4os.Compression.LZ4;
 
@@ -178,14 +179,15 @@ namespace Horde.Storage.Implementation
                 throw new Exception("Did not decompress the full payload");
 
             {
-                using IScope _ = Tracer.Instance.StartActive("cb.hash");
-                using Hasher hasher = Hasher.New();
-                hasher.UpdateWithJoin(decompressedPayload);
-                Hash blake3Hash = hasher.Finalize();
-
-                Span<byte> hash = blake3Hash.AsSpan();
-                if (!new ByteArrayComparer().Equals(hash.ToArray(), header.RawHash))
-                    throw new Exception($"Payload was expected to be {StringUtils.FormatAsHexString(header.RawHash)} but was {StringUtils.FormatAsHexString(hash.ToArray())}");
+                // only read the first 20 bytes of the hash field as IoHashes are 20 bytes and not 32 bytes
+                byte[] slicedHash = new byte[20];
+                Array.Copy(header.RawHash, 0, slicedHash, 0, 20);
+
+                BlobIdentifier headerIdentifier = new BlobIdentifier(slicedHash);
+                BlobIdentifier contentHash = BlobIdentifier.FromBlob(decompressedPayload);
+               
+                if (!headerIdentifier.Equals(contentHash))
+                    throw new Exception($"Payload was expected to be {headerIdentifier} but was {contentHash}");
             }
 
 
