Horde Storage - Fixed issue when default keyspace was not specified in the Scylla connection string. No longer specifying default keyspace by default as we already default to the right one if you do not.

diff --git a/Engine/Source/Programs/Horde/HordeStorage/Horde.Storage/Implementation/Objects/ScyllaReferencesStore.cs b/Engine/Source/Programs/Horde/HordeStorage/Horde.Storage/Implementation/Objects/ScyllaReferencesStore.cs
--- a/Engine/Source/Programs/Horde/HordeStorage/Horde.Storage/Implementation/Objects/ScyllaReferencesStore.cs
+++ b/Engine/Source/Programs/Horde/HordeStorage/Horde.Storage/Implementation/Objects/ScyllaReferencesStore.cs
@@ -378,7 +378,7 @@ namespace Horde.Storage.Implementation
 
     public class ScyllaSettings
     {
-        public string ConnectionString { get; set; } = "Contact Points=localhost,scylla;Default Keyspace=jupiter";
+        public string ConnectionString { get; set; } = "Contact Points=localhost,scylla";
         public long InlineBlobMaxSize { get; set; } = 32 * 1024; // default to 32 kb blobs max
         public string[] ContactPoints { get; set; } = new string[] {"localhost", "scylla"};
         public int MaxSnapshotsPerNamespace { get; set; } = 10;
