Horde: Set unspecified outcome on step refs when job is aborted
#rnx
#skipci

diff --git a/Engine/Source/Programs/Horde/Horde.Build/Jobs/IJobStepRefCollection.cs b/Engine/Source/Programs/Horde/Horde.Build/Jobs/IJobStepRefCollection.cs
--- a/Engine/Source/Programs/Horde/Horde.Build/Jobs/IJobStepRefCollection.cs
+++ b/Engine/Source/Programs/Horde/Horde.Build/Jobs/IJobStepRefCollection.cs
@@ -105,15 +105,12 @@ namespace HordeServer.Collections
 					}
 				}
 
-				DateTime StartTime = Step.StartTimeUtc ?? DateTime.UtcNow;
-				DateTime? FinishTime = Step.FinishTimeUtc;
-				if (Job.AbortedByUserId != null && FinishTime == null)
+				if (Job.AbortedByUserId != null && Outcome == null)
 				{
-					FinishTime = StartTime;
 					Outcome = JobStepOutcome.Unspecified;
 				}
 
-				await JobStepRefs.InsertOrReplaceAsync(new JobStepRefId(Job.Id, Batch.Id, Step.Id), Job.Name, NodeName, Job.StreamId, Job.TemplateId, Job.Change, Step.LogId, Batch.PoolId, Batch.AgentId, Outcome, LastSuccess, LastWarning, WaitTime, InitTime, StartTime, FinishTime);
+				await JobStepRefs.InsertOrReplaceAsync(new JobStepRefId(Job.Id, Batch.Id, Step.Id), Job.Name, NodeName, Job.StreamId, Job.TemplateId, Job.Change, Step.LogId, Batch.PoolId, Batch.AgentId, Outcome, LastSuccess, LastWarning, WaitTime, InitTime, Step.StartTimeUtc ?? DateTime.UtcNow, Step.FinishTimeUtc);
 			}
 		}
 	}
