Automated Test Results visibility improvements
* Rewrite testdata handler to be able to handle more test data and support partial load (through filtering)
* Introduce AutomatedTestSession plugin to handle new test data model from UE
* The new model involves a test data that hold test results per session and another test data that hold test events per session.
* Introduced BuildHealthSummary hook for global view of the test results for an entire stream.
* test report global view for now displays latest results from 30'000 testdata, history is accessible through the detailed report

diff --git a/Engine/Source/Programs/Horde/HordeDashboard/src/Plugins/AutomatedTestSession/models/AutomatedTestSession.ts b/Engine/Source/Programs/Horde/HordeDashboard/src/Plugins/AutomatedTestSession/models/AutomatedTestSession.ts
new file mode 100644
--- /dev/null
+++ b/Engine/Source/Programs/Horde/HordeDashboard/src/Plugins/AutomatedTestSession/models/AutomatedTestSession.ts
@@ -0,0 +1,911 @@
+// Copyright Epic Games, Inc. All Rights Reserved.
+
+import { observable, action, when, reaction } from "mobx";
+import { DataWrapper, TestDataWrapper, TestDataCollection } from '../../../backend/TestDataHandler'
+
+export enum TestState {
+    Failed = "Fail",
+    Success = "Success",
+    SuccessWithWarnings = "SuccessWithWarnings",
+    NotRun = "NotRun",
+    InProcess = "InProcess",
+    Skipped = "Skipped",
+    Unknown = "Unknown",
+}
+
+export enum EventType {
+    Info = "Info",
+    Error = "Error",
+    Warning = "Warning",
+}
+
+export enum ArtifactType {
+    ImageCompare = "image comparison",
+    Approved = "approved",
+    Unapproved = "unapproved",
+    Difference = "difference",
+}
+
+export type TestArtifact = {
+    Tag: string;
+    ReferencePath: string;
+}
+
+export type TestEvent = {
+    Message: string;
+    Context: string;
+    Type: string;
+    Tag: string;
+    Hash: string;
+    DateTime: string;
+    Artifacts: TestArtifact[];
+}
+
+export type TestResult = {
+    Name: string;
+    TestUID: string;
+    Suite: string;
+    State: string;
+    DeviceAppInstanceName: string;
+    ErrorCount: number;
+    WarningCount: number;
+    ErrorHashAggregate: string;
+    DateTime: string;
+    TimeElapseSec: number;
+    // references that need to be initiated post fetch
+    Details: TestResultDetails;
+    Session: TestSessionWrapper;
+}
+
+export type TestResultDetails = {
+    Events: TestEvent[];
+}
+
+export type TestDevice = MetadataHolder& {
+    Name: string;
+    AppInstanceName: string;
+}
+
+export type TestSessionInfo = {
+    DateTime: string;
+    TimeElapseSec: number;
+    Tests: TestResultsMap;
+    TestResultsTestDataUID: string;
+}
+
+export type IndexedError = {
+    Message: string;
+    Tag: string;
+    TestUIDs: string[];
+}
+
+export type Metadata = {[Key in string]: string|undefined}
+
+type MetadataHolder = {
+    Metadata?: Metadata
+}
+
+const generatetHashFromString = (aggregate: string) => {
+    let hash = 0;
+    for (let i = 0; i < aggregate.length; i++) {
+        const char = aggregate.charCodeAt(i);
+        hash = ((hash << 5) - hash) + char;
+        hash = hash & hash;
+    }
+    return hash.toString(16);
+}
+
+export const generateHashFromMetadata = (obj: Metadata) => {
+    const keys = Object.keys(obj).sort();
+    const aggregate = keys.map((key) => key + (obj[key]??'')).join();
+    return generatetHashFromString(aggregate);
+}
+
+export class MetaWrapper {
+    private collection?: Metadata
+    private _hash: string
+
+    constructor(collection?: Metadata) {
+        this.collection = collection;
+        this._hash = "";
+    }
+
+    get(key: string) : string | undefined {
+        if(this.collection) {
+            return this.collection[key];
+        }
+
+        return undefined;
+    }
+
+    map(mapFunc: (key: string) => any) : any[] {
+        if (this.collection) {
+            return Object.keys(this.collection).map(mapFunc);
+        }
+        return [];
+    }
+
+    private generateHashCode(): string {
+        if(this.collection) {
+            return generateHashFromMetadata(this.collection);
+        }
+        return "";
+    }
+
+    get hash(): string {
+        if(!this._hash) {
+            this._hash = this.generateHashCode();
+        }
+        return this._hash;
+    }
+}
+
+export type IndexedErrorMap = {[Key in string]: IndexedError}
+
+export type TestResultDetailsMap = {[Key in string]: TestResultDetails}
+
+export type TestResultsMap = {[Key in string]: TestResult}
+
+export type TestSession = MetadataHolder& {
+    Name: string;
+    Type: string;
+    PreFlightChange: string;
+    TestSessionInfo: TestSessionInfo;
+    IndexedErrors: IndexedErrorMap;
+    Devices: TestDevice[];
+}
+
+type CheckList = {[Key in string]: boolean}
+
+export enum SessionFields {
+    initiate = "data.Metadata,data.Name,data.PreFlightChange",
+    remaining = "data.TestSessionInfo,data.IndexedErrors,data.Devices,data.Type",
+}
+
+export class TestSessionWrapper implements TestSession, DataWrapper {
+    get Name(): string { return this.getOwnProperty('Name'); }
+    get Type(): string { return this.getOwnProperty('Type'); }
+    get PreFlightChange(): string { return this.getOwnProperty('PreFlightChange'); }
+    get TestSessionInfo(): TestSessionInfo { return this.getOwnProperty('TestSessionInfo'); }
+    get IndexedErrors(): IndexedErrorMap { return this.getOwnProperty('IndexedErrors'); }
+    get Devices(): TestDevice[] { return this.getOwnProperty('Devices'); }
+    get Metadata(): Metadata { return this.getOwnProperty('Metadata'); }
+
+    Testdata: TestDataWrapper;
+    MetaHandler?: MetaWrapper;
+
+    @observable
+    propertiesLoaded: CheckList = {
+        TestSessionInfo: false
+    };
+    @action
+    setPropertiesLoaded(...properties: string[]) {
+        properties.forEach((key) => {
+            this.propertiesLoaded[key] = true;
+        })
+    }
+    isPropertyLoaded(key: string) {
+        return this.propertiesLoaded[key];
+    }
+    isAllPropertiesLoaded() {
+        return this.propertiesLoaded['TestSessionInfo']; // check TestSessionInfo only for now
+    }
+    getMissingProperties() {
+        return SessionFields.remaining.split(',').map((item) => item.slice(5));
+    }
+    whenPropertyLoaded(key: string) {
+        return when(() => this.propertiesLoaded[key]);
+    }
+
+    private _testsession: TestSession;
+
+    constructor(testdata: TestDataWrapper) {
+        this._testsession = testdata.data;
+        this.Testdata = testdata;
+        // first initialization, we should call the post Properties Update
+        this.postPropertiesUpdate();
+    }
+
+    getOwnProperty<T>(key: string): T {
+        return Object.getOwnPropertyDescriptor(this._testsession, key)?.value as T;
+    }
+
+    getUID(): string {
+        return this.Name + (this.MetaHandler?.hash??'');
+    }
+
+    postPropertiesUpdate() {
+        this.capitalizeKeys()
+        const data = this._testsession;
+        if(!this.MetaHandler && data.Metadata) {
+            this.MetaHandler = new MetaWrapper(data.Metadata);
+        }
+        if(data.TestSessionInfo && !this.isPropertyLoaded('TestSessionInfo')) {
+            for(const key in data.TestSessionInfo.Tests) {
+                const test = data.TestSessionInfo.Tests[key];
+                test.Session = this;
+            }
+            this.setPropertiesLoaded('TestSessionInfo');
+        }
+    }
+
+    getTestResultByUID(uid: string) {
+        if(this.TestSessionInfo) {
+            return this.TestSessionInfo.Tests[uid];
+        }
+
+        return undefined
+    }
+
+    getDevice(AppInstanceName: string) {
+        if(this.Devices) {
+            return this.Devices.find((item) => item.AppInstanceName === AppInstanceName);
+        }
+
+        return undefined
+    }
+
+    private capitalizeKeys() {
+        // Fix inconsistances in object keys fetched from Horde backend.
+        const data = this._testsession as any;
+        for (var key in data) {
+            const first = key.charAt(0);
+            const firstUpper = first.toUpperCase();
+            if(first !== firstUpper) {
+                const value = data[key];
+                delete data[key];
+                data[firstUpper + key.slice(1)] = value;
+            }
+        }
+    }
+}
+
+export class TestStats {
+    Failed: number = 0;
+    Passed: number = 0;
+    Skipped: number = 0;
+    Unexecuted: number = 0;
+    Incomplete: number = 0;
+}
+
+export class Section {
+    Name: string;
+    FullName: string;
+    Type: FilterType;
+    Stats: TestStats;
+    Results: TestResult[];
+    TestNames: Set<string>;
+
+    constructor(name: string, fullName: string, type: FilterType) {
+        this.Name = name;
+        this.FullName = fullName;
+        this.Type = type;
+        this.Stats = new TestStats();
+        this.Results = [];
+        this.TestNames = new Set();
+    }
+
+    addTest(test: TestResult, testName?: string) {
+        this.Results.push(test);
+        this.TestNames.add(testName??test.Name);
+        // update stats
+        switch(test.State) {
+        case TestState.Failed:
+            this.Stats.Failed++;
+            break;
+        case TestState.NotRun:
+            this.Stats.Unexecuted++;
+            break;
+        case TestState.Success:
+            this.Stats.Passed++;
+            break;
+        case TestState.InProcess:
+            this.Stats.Incomplete++;
+            break;
+        case TestState.Skipped:
+            this.Stats.Skipped++;
+            break;
+        }
+    }
+
+    get isOneTest(): boolean {
+        return this.TestNames.size === 1 && this.TestNames.has(this.FullName)
+    }
+}
+
+export enum FilterType {
+    suite = 'suite',
+    name = 'name',
+    project = 'project',
+    testuid = 'testuid',
+    metauid = 'metauid',
+    testchange = 'testchange',
+    view = 'view',
+    namecontains = 'namecontains',
+}
+
+export type SectionFilter = {
+    suite?: string;
+    name?: string;
+    project?: string;
+    namecontains?: string;
+}
+
+export type ReportFilter = {
+    testuid?: string;
+    metauid?: string;
+    testchange?: number;
+    view?: string;
+}
+
+type GetKeyFunction = (test: TestResult) => [string, string];
+
+export class SectionCollection {
+
+    private fixedTestNameCache: Map<string, string>;
+    private filterCache: Map<string, {hash: string, sections: Section[]}>;
+    private onTheFlyCache: {filter: string, hash: string, sections: Section[]} | undefined;
+
+    constructor() {
+        this.fixedTestNameCache = new Map();
+        this.filterCache = new Map();
+    }
+
+    extractSections(filter: SectionFilter, sessions: TestSessionCollection) {
+        // Get the sections, first look for in the cache, otherwise generate it
+        const filterCacheKey = generateHashFromMetadata(filter);
+        const cache = this.filterCache.get(filterCacheKey);
+        const isSameHash = cache && cache.hash === sessions.hash;
+        const sections = cache && isSameHash? cache.sections : this.filterSections(filter, sessions.getTests());
+        if (!this.filterCache.has(filterCacheKey) || !isSameHash) {
+            this.filterCache.set(filterCacheKey, {hash: sessions.hash, sections:sections})
+        }
+        return sections;
+    }
+
+    reExtractSectionOnTheFly(filter: SectionFilter, testNameFilter: string, sessions: TestSessionCollection) {
+        // to use only when iterating through cached data
+        let sections = this.extractSections({...filter, namecontains: undefined}, sessions);
+        const onTheFlyFilter = {...filter, namecontains: testNameFilter.toLowerCase()}
+        sections = this.filterSections(onTheFlyFilter, sections.flatMap((section) => section.Results));
+        this.onTheFlyCache = {filter: generateHashFromMetadata(onTheFlyFilter), hash: sessions.hash, sections};
+        return sections;
+    }
+
+    promoteOnTheFlyCache() {
+        if(this.onTheFlyCache) {
+            this.filterCache.set(this.onTheFlyCache.filter, this.onTheFlyCache);
+            this.onTheFlyCache = undefined;
+        }
+    }
+
+    private filterSections(filter: SectionFilter, tests: TestResult[]) {
+        const filterFunc = (test: TestResult) => {
+            if(filter.project) {
+                if(test.Session.MetaHandler?.get('Project') !== filter.project) {
+                    return false;
+                }
+            }
+            if(filter.suite) {
+                if(test.Suite !== filter.suite) {
+                    return false;
+                }
+            }
+            if(filter.namecontains) {
+                const substrings = filter.namecontains.split(' ');
+                const testName = test.Name.toLowerCase();
+                if(!substrings.some((value) => testName.indexOf(value) >= 0)) {
+                    return false;
+                }
+            }
+            if(filter.name) {
+                if(this.fixTestName(test.Name).indexOf(filter.name + '.') !== 0) {
+                    return false;
+                }
+            }
+            return true;
+        }
+        const filteredTests = tests.filter(filterFunc);
+
+        const commonSection = filter.name? this.getCommonSection(filteredTests.map((item) => this.fixTestName(item.Name))) : "";
+        const sectionKeyFunc = (key: FilterType): GetKeyFunction | undefined => {
+            switch(key) {
+                case FilterType.suite:
+                    return (test: TestResult) => {
+                        return [test.Suite, test.Suite];
+                    }
+                case FilterType.name:
+                    return (test: TestResult) => {
+                        const suite = this.fixTestName(test.Name).slice(commonSection.length).split('.', 1)[0];
+                        return [suite, commonSection+suite];
+                    }
+            }
+        };
+        const sectionFocus = filter.suite || filter.name? FilterType.name : FilterType.suite;
+        const sections = this.buildSections(filteredTests, sectionFocus, sectionKeyFunc(sectionFocus)!);
+        // sorting
+        sections.forEach((section) => this.sortSectionResults(section));
+        this.sortSections(sections);
+        return sections;
+    }
+
+    private buildSections(tests: TestResult[], type: FilterType, keyFunc: GetKeyFunction): Section[] {
+        const sections : Section[] = [];
+        const sectionsMap = new Map<string, Section>(); 
+
+        tests.forEach((test) => {
+            const [ key, fullKey ] = keyFunc(test);
+            let targetSection = sectionsMap.get(key);
+            if(!targetSection) {
+                targetSection = new Section(key, fullKey, type);
+                sections.push(targetSection);
+                sectionsMap.set(key, targetSection);
+            }
+            targetSection.addTest(test, this.fixTestName(test.Name));
+        });
+
+        return sections;
+    }
+
+    private  fixTestName(name: string): string {
+        // It is necessary to sanitize the test names that contains / instead of . syntax
+        if(this.fixedTestNameCache.has(name)) {
+            return this.fixedTestNameCache.get(name)!;
+        }
+        let fixedName = name
+        if(name.indexOf('/') > -1) {
+            fixedName = name.replaceAll(/\.*\//g, ".");
+        }
+        this.fixedTestNameCache.set(name, fixedName);
+        return fixedName;
+    }
+
+    private getCommonString(list: string[]): string {
+        if(list.length === 0) {
+            return "";
+        }
+        let shortest: string = list[0];
+        let longest: string = "";
+        list.forEach((item) => {
+            if(item > longest) {
+                longest = item;
+                return;
+            }
+            if(item < shortest)
+            {
+                shortest = item;
+            }
+        });
+        while(longest.indexOf(shortest) !== 0) {
+            shortest = shortest.slice(0, -1)
+        }
+
+        return shortest;
+    }
+
+    private getCommonSection(list: string[]): string {
+        let commonKey = this.getCommonString(list);
+        if(commonKey[commonKey.length - 1] !== '.' && commonKey.indexOf('.') >= 0) {
+            commonKey = commonKey.slice(0, commonKey.lastIndexOf('.') + 1);
+        }
+
+        return commonKey;
+    }
+
+    // sort comparing the number of Failed, Incomplete, Unexecuted and Passed
+    private sortSections(sections: Section[]) {
+        sections.sort((a, b) => {
+            if (a.Stats.Failed > b.Stats.Failed) {
+                return -1;
+            } else if (a.Stats.Failed < b.Stats.Failed) {
+                return 1;
+            } else if (a.Stats.Incomplete > b.Stats.Incomplete) {
+                return -1;
+            } else if (a.Stats.Incomplete < b.Stats.Incomplete) {
+                return 1;
+            } else if (a.Stats.Unexecuted > b.Stats.Unexecuted) {
+                return -1;
+            } else if (a.Stats.Unexecuted < b.Stats.Unexecuted) {
+                return 1;
+            } else if (a.Stats.Passed > b.Stats.Passed) {
+                return -1;
+            } else if (a.Stats.Passed < b.Stats.Passed) {
+                return 1;
+            } else if (a.isOneTest && !b.isOneTest) {
+                return 1;
+            } else if (b.isOneTest && !a.isOneTest) {
+                return -1;
+            }
+            return 0;
+        });
+    }
+
+    // sort in that order Failed, Incomplete, Unexecuted, Passed, Skipped
+    private sortSectionResults(section: Section) {
+        section.Results.sort((a, b) => {
+            if(a.State === TestState.Failed) {
+                if(b.State !== TestState.Failed) {
+                    return -1;
+                }
+            } else if(a.State === TestState.InProcess) {
+                if(b.State === TestState.Failed) {
+                    return 1;
+                } else if(b.State !== TestState.InProcess) {
+                    return -1;
+                }
+            } else if(a.State === TestState.NotRun) {
+                if(b.State === TestState.Failed || b.State === TestState.InProcess) {
+                    return 1;
+                } else if(b.State !== TestState.NotRun) {
+                    return -1;
+                }
+            } else if(a.State === TestState.Success) {
+                if(b.State === TestState.Skipped) {
+                    return -1;
+                } else if(b.State !== TestState.Success || b.WarningCount > 0) {
+                    return 1;
+                } else if(a.WarningCount > 0) {
+                    return -1;
+                }
+            } else if(a.State === TestState.Skipped) {
+                if(b.State !== TestState.Skipped) {
+                    return 1;
+                }
+            }
+
+            if (a.Name > b.Name) {
+                return 1;
+            } else if (a.Name < b.Name) {
+                return -1;
+            }
+
+            return 0;
+        })
+    }
+}
+
+type SessionReferences = {
+    SessionId: string;
+    ResultsId: string;
+}
+
+export class TestSessionCollection {
+
+    @observable
+    version: number = 0;
+    @action
+    updateVersion() {
+        this.version++;
+    }
+    reactOnVersionChange(callback: (version: number) => void) {
+        return reaction(
+            () => this.version,
+            (version) => callback(version)
+        );
+    }
+
+    @observable
+    loadingCount: number = 0;
+    @action
+    incrementLoading() {
+        this.loadingCount++;
+    }
+    @action
+    decrementLoading() {
+        this.loadingCount--;
+    }
+
+    get size(): number { return this._collection.size; }
+    get hash(): string {
+        let aggregate = "";
+        this.forEachSession( // only include fully loaded session
+            (session) => aggregate += session.isAllPropertiesLoaded()?session.getUID():""
+        );
+        return generatetHashFromString(aggregate);
+    }
+
+    private _collection: Map<string, TestSessionWrapper>;
+    private _datahandler: TestDataCollection;
+    private _includePreflight: boolean = false;
+
+    constructor(handler: TestDataCollection, includePreflight: boolean = false) {
+        this._collection = new Map();
+        this._datahandler = handler;
+        this._includePreflight = includePreflight;
+    }
+
+    set(uid: string, item: TestSessionWrapper) {
+        this._collection.set(uid, item);
+        if(item.isAllPropertiesLoaded()) { this.updateVersion(); }
+    }
+
+    get(uid: string) {
+        return this._collection.get(uid);
+    }
+
+    has(uid: string) {
+        return this._collection.has(uid);
+    }
+
+    delete(uid: string) {
+        this._collection.delete(uid);
+        this.updateVersion();
+    }
+
+    clear() {
+        this._collection.clear()
+        this.updateVersion();
+    }
+
+    forEachSession(func: (item: TestSessionWrapper) => void) {
+        this._collection.forEach(func);
+    }
+
+    *iterSessions() {
+        const iterator = this._collection.values();
+        let result = iterator.next();
+        while(!result.done) {
+            yield result.value;
+            result = iterator.next();
+        }
+    }
+
+    *iterTests() {
+        const iterSession = this._collection.values();
+        let resultSession = iterSession.next();
+        while(!resultSession.done) {
+            const sessionInfo = resultSession.value.TestSessionInfo;
+            if (sessionInfo) {
+                for (const key in sessionInfo.Tests) {
+                    yield sessionInfo.Tests[key];
+                }
+            }
+            resultSession = iterSession.next();
+        }
+    }
+
+    async getSessionsForStream(streamId: string, maxCount: number = 1000) {
+        this.incrementLoading();
+        const onFetch = (items: (TestDataWrapper | undefined)[]) => this.onInitialFetch(items);
+        try {
+            if(streamId === this._datahandler.streamId
+                && this._datahandler.items
+                && this._datahandler.activated
+                && maxCount > this._datahandler.items.size) {
+                    // Max count resize
+                    await this._datahandler.queueFetchItems(onFetch, SessionFields.initiate, maxCount);
+            } else {
+                // Initiale fetch or reload
+                await this._datahandler.setFromStream(
+                    streamId, "Automated Test Session", onFetch, SessionFields.initiate, maxCount
+                );
+            }
+        } finally { this.decrementLoading() }
+        await when(() => this.loadingCount === 0);
+    }
+
+    private onInitialFetch(items: (TestDataWrapper|undefined)[]) {
+        const sessionsToFetch: string[] = [];
+        items.forEach((item) => {
+            if (item) {
+                let data = item.getDataHandler() as TestSessionWrapper;
+                if(!data) {
+                    data = new TestSessionWrapper(item);
+                    item.setDataHandler(data);
+                }
+                if (!this._includePreflight && data.PreFlightChange) {
+                    return;
+                }
+                const session_uid = data.getUID();
+                if(!this.has(session_uid)) {
+                    this.set(session_uid, data);
+                    if (!data.isAllPropertiesLoaded()) {
+                        sessionsToFetch.push(item.id);
+                    }
+                }
+            }
+        });
+        if(sessionsToFetch.length > 0) {
+            this.incrementLoading();
+            this._datahandler.fetchUpdateItems(
+                sessionsToFetch, (items) => this.onRemainingFetch(items), SessionFields.remaining
+            ).finally(() => this.decrementLoading());
+        }
+    }
+
+    private onRemainingFetch(items: (TestDataWrapper|undefined)[]) {
+        let hasUpdated = false;
+        items.forEach((item) => {
+            if (item) {
+                let data = item.getDataHandler() as TestSessionWrapper;
+                if(data && data.isAllPropertiesLoaded()) {
+                    hasUpdated = true;
+                }
+            }
+        });
+        if(hasUpdated) { this.updateVersion() }
+    }
+
+    async getTestResultByUID(uid: string): Promise<TestResult[]> {
+        const iterator = this._collection.values();
+        const tests: TestResult[] = [];
+        const toFetch: TestSessionWrapper[] = [];
+        let item = iterator.next();
+        while(!item.done) {
+            const session = item.value;
+            if (!session.isAllPropertiesLoaded()) {
+                await session.whenPropertyLoaded('TestSessionInfo');
+            }
+            const test = session.getTestResultByUID(uid);
+            if (test) {
+                if (test.Details) {
+                    tests.push(test);
+                } else {
+                    toFetch.push(session);
+                }
+            }
+            item = iterator.next();
+        }
+
+        if(tests.length === 0 && this.loadingCount !== 0) {
+            // Give a chance to the collection to be fully loaded if nothing was found
+            // It is possible the test is run infrequently and is deep in the result collection
+            await when(() => this.loadingCount === 0);
+            // then retry
+            return await this.getTestResultByUID(uid);
+        }
+
+        if(toFetch.length > 0) {
+            // Fetch missing data from session test results
+            await this.fetchSessionResults(
+                toFetch.map((item) => {
+                    return {SessionId: item.Testdata.id, ResultsId: item.TestSessionInfo.TestResultsTestDataUID}
+                }), 'data.' + uid
+            );
+            toFetch.forEach((item) => {
+                const test = item.getTestResultByUID(uid);
+                if(test?.Details) {
+                    tests.push(test);
+                }
+            });
+        }
+
+        return tests;
+    }
+
+    async getTestResultByQuery(uid: string, metauid: string, change: number): Promise<TestResult | undefined> {
+        const latestTests = await this.getTestResultByUID(uid);
+        let test: TestResult | undefined = undefined;
+
+        if(latestTests.length > 0) {
+            test = latestTests.find((test) => test.Session.Testdata.change === change && test.Session.MetaHandler?.hash === metauid);
+            if(!test) {
+                // Need to get the target change from the history
+                const sessionName = latestTests[0].Session.Name;
+                const changes = await this.getSessionHistory(sessionName);
+                const sessionsChange = changes.find((items) => items[0].Testdata.change === change);
+                const session = sessionsChange?.find((item) => item.MetaHandler?.hash === metauid);
+                test = session?.getTestResultByUID(uid);
+            }
+        }
+
+        if(test && !test.Details) {
+            // Fetch missing data from the session test result
+            await this.fetchSessionResults(
+                [{SessionId: test.Session.Testdata.id, ResultsId: test.Session.TestSessionInfo.TestResultsTestDataUID}]
+                , 'data.' + uid
+            );
+            if(!test.Details) {
+                // if nothing was fetched or failed to fetch, return nothing
+                return undefined;
+            }
+        }
+
+        return test;
+    }
+
+    private async fetchSessionResults(items: SessionReferences[], filter?: string, parallelCount: number = 30) {
+        let count = 0;
+        for (let i = 0; i < items.length; i += parallelCount) {
+            if (!this._datahandler.activated) {
+                break;
+            }
+            const chunk = items.slice(i, i + parallelCount);
+            count += (
+                await Promise.all(chunk.map((item) => this.fetchResultItemFromSession(item, filter)))
+            ).reduce((a, b) => a + b);
+        }
+        return count;
+    }
+
+    private async fetchResultItemFromSession(item: SessionReferences, filter?: string): Promise<number> {
+        const update = await this._datahandler.fetchItemsFromAlternateKey(`${this._datahandler.key} Result Details::${item.ResultsId}`, undefined, 1, 0, filter);
+        // We only expect one returned item
+        if(update.length > 0) {
+            const results = update[0].data as TestResultDetailsMap;
+            const session = this._datahandler.items?.get(item.SessionId)?.getDataHandler() as TestSessionWrapper;
+            if(session) {
+                for(const uid in results ) {
+                    const test = session.getTestResultByUID(uid);
+                    if (test) {
+                        test.Details = results[uid];
+                    }
+                }
+                return 1;
+            }
+        }
+
+        return 0;
+    }
+
+    getTests() {
+        const tests: TestResult[] = [];
+        const iterator = this._collection.values();
+        let result = iterator.next();
+        while(!result.done) {
+            const sessionInfo = result.value.TestSessionInfo;
+            if(sessionInfo) {
+                tests.push(...Object.values(result.value.TestSessionInfo.Tests));
+            }
+            result = iterator.next();
+        }
+        return tests;
+    }
+
+    async getSessionHistory(sessionName: string, maxCount: number = 1000) {
+        if(this._datahandler.cursor) {
+            // for single view mode we need to fetch the whole pool for context
+            const onFetch = (items: (TestDataWrapper|undefined)[]) => this.onInitialHistoryFetch(items)
+            await this._datahandler.getCursorHistory(onFetch, SessionFields.initiate, maxCount)
+        }
+
+        const missingSessions: string[] = [];
+        if(this.loadingCount !== 0) {
+            // wait for the current collection to have loaded
+            await when(() => this.loadingCount === 0);
+        }
+
+        const relatedSessionsByChange: TestSessionWrapper[][] = [];
+        const iterator = this._datahandler.iterItemsByChange();
+        let change = iterator.next();
+        while (!change.done) {
+            const sessions: TestSessionWrapper[] = [];
+            change.value.forEach((item) => {
+                const session = item.getDataHandler() as TestSessionWrapper;
+                if(session.Name === sessionName) {
+                    if (!this._includePreflight && session.PreFlightChange) {
+                        return;
+                    }
+                    sessions.push(session);
+                    if(!session.isAllPropertiesLoaded()) {
+                        missingSessions.push(session.Testdata.id);
+                    }
+                }
+            });
+            if (sessions.length > 0) {
+                relatedSessionsByChange.push(sessions);
+            }
+            change = iterator.next();
+        }
+
+        if(missingSessions.length > 0) {
+            await this._datahandler.fetchUpdateItems(
+                missingSessions, undefined, SessionFields.remaining
+            )
+        }
+        return relatedSessionsByChange;
+    }
+
+    private onInitialHistoryFetch(items: (TestDataWrapper|undefined)[]) {
+        items.forEach((item) => {
+            if (item) {
+                let data = item.getDataHandler() as TestSessionWrapper;
+                if(!data) {
+                    data = new TestSessionWrapper(item);
+                    item.setDataHandler(data);
+                }
+            }
+        });
+    }
+
+}
diff --git a/Engine/Source/Programs/Horde/HordeDashboard/src/backend/Backend.ts b/Engine/Source/Programs/Horde/HordeDashboard/src/backend/Backend.ts
--- a/Engine/Source/Programs/Horde/HordeDashboard/src/backend/Backend.ts
+++ b/Engine/Source/Programs/Horde/HordeDashboard/src/backend/Backend.ts
@@ -1060,19 +1060,15 @@ export class Backend {
 
     }
 
-    getJobStepTestDataByKey(jobId: string, stepId: string, key: string): Promise<TestData[]> {
-        return new Promise<TestData[]>((resolve, reject) => {
-            this.backend.get(`/api/v1/testdata?JobId=${jobId}&JobStepId=${stepId}&Key=${key}`).then((value) => {
-                resolve(value.data as TestData[]);
-            }).catch((reason) => {
-                reject(reason);
-            });
-        });
-    }
-
-    getTestData(id: string): Promise<TestData> {
+    getTestData(id: string, filter?: string): Promise<TestData> {
         return new Promise<TestData>((resolve, reject) => {
-            this.backend.get(`/api/v1/testdata/${id}`).then((value) => {
+            const params: any = {
+                filter: 'id,key,change,jobId,stepId,streamId,' + (filter? filter : 'data')
+            };
+
+            this.backend.get(`/api/v1/testdata/${id}`, {
+                params: params
+            }).then((value) => {
                 resolve(value.data as TestData);
             }).catch((reason) => {
                 reject(reason);
@@ -1084,7 +1080,7 @@ export class Backend {
         return new Promise<TestData[]>((resolve, reject) => {
             const params: any = {
                 jobId: jobId,
-                filter: 'key,id,stepId',
+                filter: 'id,key,change,jobId,stepId,streamId',
                 count: 200,
             };
 
@@ -1102,13 +1098,13 @@ export class Backend {
         });
     }
 
-    getTestDataHistory(streamId: string, key: string, maxChange?: number, count: number = 30, index?: number): Promise<TestData[]> {
+    getTestDataHistory(streamId: string, key: string, maxChange?: number, count: number = 30, index?: number, filter?: string): Promise<TestData[]> {
         return new Promise<TestData[]>((resolve, reject) => {
             const params: any = {
                 streamId: streamId,
                 key: key,
                 count: count,
-                filter: 'id,change,data'
+                filter: 'id,key,change,jobId,stepId,streamId,' + (filter? filter : 'data')
             };
 
             if (maxChange) {
diff --git a/Engine/Source/Programs/Horde/HordeDashboard/src/backend/Config.ts b/Engine/Source/Programs/Horde/HordeDashboard/src/backend/Config.ts
--- a/Engine/Source/Programs/Horde/HordeDashboard/src/backend/Config.ts
+++ b/Engine/Source/Programs/Horde/HordeDashboard/src/backend/Config.ts
@@ -12,6 +12,7 @@ const dev = {
 	"environment": "dev",
 	"datadogClient" : "",
 	"plugins": [
+		"AutomatedTestSession",
 		"UnrealAutomatedTests",
 		"SimpleTestReport"
 	]
@@ -21,6 +22,7 @@ const prod = {
 	"environment": "production",
 	"datadogClient" : "",
 	"plugins": [
+		"AutomatedTestSession",
 		"UnrealAutomatedTests",
 		"SimpleTestReport"
 	]
diff --git a/Engine/Source/Programs/Horde/HordeDashboard/src/backend/TestDataHandler.ts b/Engine/Source/Programs/Horde/HordeDashboard/src/backend/TestDataHandler.ts
--- a/Engine/Source/Programs/Horde/HordeDashboard/src/backend/TestDataHandler.ts
+++ b/Engine/Source/Programs/Horde/HordeDashboard/src/backend/TestDataHandler.ts
@@ -1,252 +1,522 @@
 // Copyright Epic Games, Inc. All Rights Reserved.
 
-import { observable, action } from "mobx";
+import { observable, action, when } from "mobx";
 import backend from '../backend';
-import { StreamData, JobData, ArtifactData, TestData } from './Api';
-import { projectStore } from './ProjectStore';
-import { BreadcrumbItem } from '../components/Breadcrumbs';
-
-class JobInfo {
-    stream?: StreamData;
-    jobdata?: JobData;
-    artifacts?: ArtifactData[];
-    stepName : string = "UnknownJob";
-    id?: string;
-    stepId?: string;
+import { JobData, ArtifactData, TestData } from './Api';
 
-    @action
-    setActive(active: boolean) {
-        this.active = active;
-    }
-
-    @observable
-    active = true;
+export type DataWrapper = {
+    postPropertiesUpdate(): void
+    isAllPropertiesLoaded(): boolean
+    getMissingProperties(): string[]
+}
 
-    async set(targetJobId : string, targetStepId : string) {
-        if (this.id === targetJobId && this.stepId === targetStepId) {
-            return;
+const updateObject = (target: any, version: any, skipKeys?: string[]) => {
+    const properties = Object.getOwnPropertyDescriptors(version);
+    Object.getOwnPropertyNames(version).forEach((key) => {
+        if(!properties[key].value || skipKeys?.includes(key)) {
+            delete properties[key];
         }
+    });
+    return Object.defineProperties(target, properties);
+}
 
-        if (!targetJobId || !targetStepId) {
-            return;
+export class TestDataWrapper implements TestData {
+    get id(): string { return this.getTestDataProperty('id'); }
+    get key(): string { return this.getTestDataProperty('key'); }
+    get change(): number { return this.getTestDataProperty('change'); }
+    get jobId(): string { return this.getTestDataProperty('jobId'); }
+    get stepId(): string { return this.getTestDataProperty('stepId'); }
+    get streamId(): string { return this.getTestDataProperty('streamId'); }
+    get templateRefId(): string { return this.getTestDataProperty('templateRefId'); }
+    get data(): any { return this.getTestDataProperty('data'); }
+
+    get type(): string { return this.key.split('::', 2)[0]; }
+
+    private jobArtifacts?: ArtifactData[];
+    private artifactMap?: Map<string, ArtifactData | undefined>;
+    private jobdata?: JobData;
+    private stepName?: string;
+
+    private _testdata: TestData;
+    private _datahandler?: DataWrapper;
+
+    constructor(item: TestData) {
+        this._testdata = item;
+    }
+
+    setDataHandler(handler: DataWrapper) {
+        if(this.data) {
+            this._datahandler = handler;
         }
+    }
 
-        this.stepId = targetStepId;
-        this.stepName = "UnknownJob";
+    getDataHandler(): DataWrapper | undefined {
+        return this._datahandler ;
+    }
 
-        this.setActive(true);
+    getTestDataProperty<T>(key: string): T {
+        return Object.getOwnPropertyDescriptor(this._testdata, key)?.value as T;
+    }
 
-        try {
-            if (this.jobdata === undefined || this.id !== targetJobId) {
-                this.jobdata = await backend.getJob(targetJobId);
-                this.id = targetJobId;
-                this.stream = projectStore.streamById(this.jobdata.streamId);
+    updateProperties(update: TestData) {
+        updateObject(this._testdata, update, ['data']);
+        if (update.data) {
+            if (!this._testdata.data) {
+                this._testdata.data = update.data;
+            } else {
+                updateObject(this._testdata.data, update.data);
             }
+            this._datahandler?.postPropertiesUpdate();
+        }
+    }
 
-            let stepName = "";
-            const batch = this.jobdata.batches?.find(b => b.steps.find(s => s.id === this.stepId));
-            const stepNode = batch?.steps.find(s => s.id === this.stepId);
-            const groups = this.jobdata?.graphRef?.groups;
-            if (groups && stepNode && batch) {
-                stepName = groups[batch.groupIdx]?.nodes[stepNode.nodeIdx]?.name;
+    isPropertiesDefined(...properties: string[]) {
+        const checks = properties.map((key) => {
+            let obj = this._testdata;
+            let tailKey = key;
+            let headKey = key;
+            while (tailKey.indexOf('.') > 0) {
+                [headKey, tailKey] = tailKey.split('.', 2);
+                const nextProp = Object.getOwnPropertyDescriptor(obj, headKey);
+                if(!nextProp) {
+                    return false;
+                }
+                obj = nextProp?.value;
             }
-            this.stepName = stepName;
+            return !!Object.getOwnPropertyDescriptor(obj, tailKey);
+        });
 
-            this.artifacts = await backend.getJobArtifacts(this.id!, this.stepId);
+        return !checks.find((check) => !check);
+    }
+
+    async getArtifacts() {
+        if (!this.jobArtifacts) {
+            this.jobArtifacts = await backend.getJobArtifacts(this.jobId, this.stepId);
         }
-        catch(reason) {
-            console.error('Failed to load job info!');
-            throw reason;
+        return this.jobArtifacts;
+    }
+
+    async findArtifactData(artifactName: string) {
+        if (!artifactName) {
+            return undefined;
         }
-        finally {
-            this.setActive(false);
+
+        const jobArtifacts = await this.getArtifacts();
+
+        artifactName = artifactName.replace(/\\/g,  '/');
+
+        if (!this.artifactMap) {
+            this.artifactMap = new Map();
         }
-    }   
 
-}
+        if (this.artifactMap.has(artifactName)) {
+            return this.artifactMap.get(artifactName);
+        }
 
-export class TestDataHandler {
-    stream?: StreamData;
-    id?: string;
+        const found = jobArtifacts?.find((value) => value.name.indexOf(artifactName) > -1);
+        this.artifactMap.set(artifactName, found);
 
-    jobInfo: JobInfo = new JobInfo();
-    testdata?: TestData;
-    name?: string;
-    type?: string;
+        return found;
+    }
 
-    history?: TestData[];
+    async getJobStepName() {
+        if(this.stepName) {
+            return this.stepName;
+        }
+        if(!this.jobdata) {
+            this.jobdata = await backend.getJob(this.jobId);
+        }
 
-    artifactMap: Map<string, ArtifactData | undefined> = new Map();
+        let stepName = "";
+        const batch = this.jobdata.batches?.find(b => b.steps.find(s => s.id === this.stepId));
+        const stepNode = batch?.steps.find(s => s.id === this.stepId);
+        const groups = this.jobdata?.graphRef?.groups;
+        if (groups && stepNode && batch) {
+            stepName = groups[batch.groupIdx]?.nodes[stepNode.nodeIdx]?.name;
+        }
+        this.stepName = stepName;
 
-    @action
-    setActive(active: boolean) {
-        this.active = active;
+        return this.stepName;
     }
 
-    @observable
-    active = true;
+    getJobStepLink() {
+        return `/job/${this.jobId}?step=${this.stepId}`;
+    }
 
-    @action
-    setHistoryLoaded(loaded: boolean) {
-        this.historyLoaded = loaded;
+    async getJobName() {
+        if(!this.jobdata) {
+            this.jobdata = await backend.getJob(this.jobId);
+        }
+
+        return this.jobdata.name;
     }
 
+    getJobLink() {
+        return `/job/${this.jobId}`;
+    }
+
+    async preflighChange() {
+        if(!this.jobdata) {
+            this.jobdata = await backend.getJob(this.jobId);
+        }
+
+        return this.jobdata.preflightChange;
+    }
+
+    async getChangeName() {
+        if(!this.jobdata) {
+            this.jobdata = await backend.getJob(this.jobId);
+        }
+
+        const data = this.jobdata;
+
+        let changeName = "";
+        if (data.preflightChange) {
+            changeName = `Preflight ${data.preflightChange} `;
+            changeName += ` - Base ${data.change ? data.change : "Latest"}`;
+
+        } else {
+            changeName = `${data.change ? "CL " + data.change : "Latest CL"}`;
+        }
+        
+        return changeName;
+    }
+}
+
+export type OnFetchTestData = (newChunk: (TestDataWrapper|undefined)[]) => void;
+export type SelectorTestData = (item: TestDataWrapper) => boolean;
+
+export class TestDataCollection {
+
+    activated = false;
+    setActive(activated: boolean) {
+        this.activated = activated;
+    }
+
+    // loadingNotice is used to mark discreet items that are being fetched.
+    // That is to avoid doing multiple fetch requests of the same items.
     @observable
-    historyLoaded = false
+    loadingNotice: Set<string> = new Set();
+    @action
+    acquireLoadingNotice(key: string) {
+        this.loadingNotice.add(key);
+    }
+    @action
+    releaseLoadingNotice(key: string) {
+        this.loadingNotice.delete(key);
+    }
+    @action
+    releaseAllLoadingNotice() {
+        this.loadingNotice.clear();
+    }
+    isLoading(key: string) {
+        return this.loadingNotice.has(key);
+    }
 
-    async set(targetId: string) {
+    key?: string;
+    streamId?: string;
+    items?: Map<string, TestDataWrapper>;
+    itemsByChange?: Map<number, TestDataWrapper[]>;
+    cursor?: TestDataWrapper;
+    private nextPageIndex: number = 0;
+    private fetchCount: number = 100;
+    private fetchMaxCount: number = 1000;
+
+    desactivate() {
+        this.setActive(false);
+        this.releaseAllLoadingNotice();
+    }
 
-        if (this.id === targetId) {
-            return;
+    findById(id: string) {
+        if(!this.items) {
+            return undefined;
         }
+        return this.items.get(id);
+    }
 
-        if (!targetId) {
+    forEach(func: (item: TestDataWrapper) => void) {
+        if(!this.items) {
             return;
         }
+        this.items.forEach(func);
+    }
 
-        this.id = targetId;
-        this.testdata = undefined;
-        this.name = undefined;
-        this.type = undefined;
-        this.history = undefined;
-        if (this.artifactMap.size > 0) {
-            this.artifactMap = new Map();
+    *iterItems() {
+        const items = this.items;
+        if(items) {
+            const iterator = items.values();
+            let result = iterator.next();
+            while(!result.done) {
+                yield result.value;
+                result = iterator.next();
+            }
         }
+    }
 
+    *iterItemsByChange() {
+        if(this.itemsByChange) {
+            const changes = this.getChanges();
+            for (const change of changes) {
+                const items = this.itemsByChange.get(change);
+                if (items) {
+                    yield items;
+                }
+            }
+        }
+    }
+
+    getChanges() {
+        if(this.itemsByChange) {
+            const changes = Array.from(this.itemsByChange.keys());
+            changes.sort((a, b) => a > b? -1 : (a < b? 1 : 0)); // in reversed order
+            return changes;
+        }
+        return [];
+    }
+
+    async setFromStream(streamId: string, key: string, onFetch?: OnFetchTestData, filter?: string, maxCount? : number) {
         this.setActive(true);
 
-        try {
-            this.testdata = await backend.getTestData(targetId);
-            const key = this.testdata.key.split('::');
-            if (key.length > 1) {
-                [ this.type, this.name ] = key;
+        if (key === this.key && streamId === this.streamId && !this.cursor) {
+            // reprocessed what was loaded already, and then continue loading if max count was not reached
+            if(this.items) {
+                if(onFetch) {
+                    onFetch(Array.from(this.items.values()));
+                }
+                if(this.items.size%this.fetchCount !== 0 || this.items.size >= this.fetchMaxCount) {
+                    return;
+                }
             }
-            else {
-                this.name = key[0];
+        } else {
+            this.cleanCache();
+            this.key = key;
+            this.streamId = streamId;
+        }
+
+        await this.queueFetchItems(onFetch, filter, maxCount);
+    }
+
+    async setFromId(id: string, onFetch?: OnFetchTestData, filter?: string) {
+        this.setActive(true);
+
+        if (this.items) {
+            const item = this.items.get(id);
+            if(item) {
+                const handler = item.getDataHandler();
+                if (handler) {
+                    if(!handler.isAllPropertiesLoaded()){
+                        const missing = handler.getMissingProperties();
+                        const filter = missing.map((item) => `data.${item}`).join(',');
+                        const count = await this.fetchUpdateItem(id, undefined, filter);
+                        if (count === 0) {
+                            return undefined;
+                        }
+                    }
+                }
+                this.cursor = item;
+                return item;
             }
+        }
 
-            await this.jobInfo.set(this.testdata.jobId, this.testdata.stepId);
-            this.stream = this.jobInfo.stream;
+        // pre fetch minimal data to get the context and see if we need to clean the cache
+        const item = await backend.getTestData(id, "nodata");
+        if(!item) {
+            return undefined;
         }
-        catch(reason) {
-            console.error('Failed to load test data from job!');
-            throw reason;
+        if(item.key !== this.key || item.streamId !== this.streamId) {
+            this.cleanCache();
+            this.key = item.key;
+            this.streamId = item.streamId;
         }
-        finally {
-            this.setActive(false);
+
+        // actually fetching the item
+        const count = await this.fetchUpdateItem(id, onFetch, filter);
+        if (count === 1) {
+            this.cursor = this.items?.get(id);
+            return this.cursor;
         }
 
+        return undefined;
     }
 
-    async getHistory(fromChange? : number) {
-        if (!this.id || !this.stream || !this.testdata) {
-            return;
+    async getCursorHistory(onFetch?: OnFetchTestData, filter?: string, maxCount?: number, selector?: SelectorTestData) {
+        if (!this.cursor) {
+            return [];
         }
 
-        if (this.history !== undefined) {
-            return;
+        if (!this.itemsByChange || this.itemsByChange.size === 1) {
+            await this.queueFetchItems(onFetch, filter, maxCount??this.fetchCount);
         }
 
-        this.setHistoryLoaded(false);
+        const mapper = selector? (items: TestDataWrapper[]) => items.find(selector) as TestDataWrapper : (items: TestDataWrapper[]) => items[0];
 
-        try {
-            this.history = await backend.getTestDataHistory(this.stream.id, this.testdata.key, fromChange);
-        }
-        catch(reason) {
-            console.error('Failed to load test data history!');
-            throw reason;
+        return Array.from(this.iterItemsByChange()).map(mapper);
+    }
+
+    async queueFetchItems(onFetch?: OnFetchTestData, filter?: string, maxCount? : number) {
+        if (maxCount) {
+            this.fetchMaxCount = maxCount;
         }
-        finally {
-            this.setHistoryLoaded(true);
+
+        let totalCount = await this.fetchItems(onFetch, filter);
+        let count = totalCount;
+        while (this.activated && count > 0 && count%this.fetchCount === 0 && this.items && this.items.size < this.fetchMaxCount) {
+            count = await this.fetchItems(onFetch, filter);
+            totalCount += count;
         }
+
+        return totalCount;
     }
 
-    cleanHistory() {
-        if (this.history !== undefined) {
-            this.history = undefined;
-            this.setHistoryLoaded(false);
+    private cleanCache() {
+        if (this.items) {
+            this.items = undefined;
+            this.itemsByChange = undefined;
+            this.cursor = undefined;
+            this.nextPageIndex = 0;
+            this.fetchMaxCount = 1000; // the default
+            this.releaseAllLoadingNotice();
         }
     }
 
-    get clText(): string {
-
-        const data = this.jobInfo.jobdata;
-        if (!data) {
-            return "";
+    private initiateCache() {
+        if(!this.items) {
+            this.items = new Map();
+            this.itemsByChange = new Map();
         }
+    }
 
-        let clText = "";
-        if (data.preflightChange) {
-            clText = `Preflight ${data.preflightChange} `;
-            clText += ` - Base ${data.change ? data.change : "Latest"}`;
-
+    private updateItemsByChange(item: TestDataWrapper) {
+        if (!this.itemsByChange) {
+            throw new Error("Trying to update itemsByChange before the object was initiated.")
+        }
+        const change = item.change;
+        const changeCollection = this.itemsByChange.get(change);
+        if(changeCollection) {
+            changeCollection.push(item);
         } else {
-            clText = `${data.change ? "CL " + data.change : "Latest CL"}`;
+            this.itemsByChange.set(change, [item])
         }
-        
-        return clText;
     }
 
-    get crumbs(): BreadcrumbItem[] {
+    private async fetchItems(onFetch?: OnFetchTestData, filter?: string) {
+        if (!this.key || !this.streamId) {
+            return 0;
+        }
+        const key = this.key;
+        const streamId = this.streamId;
+        let count = 0;
+        try {
+            const pageData = await backend.getTestDataHistory(streamId, key, undefined, this.fetchCount, this.nextPageIndex++, filter);
+
+            // Are we still fetching for the same collection?
+            if(key === this.key && streamId === this.streamId) {
 
-        const jobdata = this.jobInfo.jobdata!;
+                this.initiateCache();
 
-        if (!this.stream) {
-            return [];
-        }
+                const wrappedData = pageData.map((item) => {
+                    if(!this.items?.has(item.id)) {
+                        const itemWrapped = new TestDataWrapper(item);
+                        this.items?.set(item.id, itemWrapped);
+                        this.updateItemsByChange(itemWrapped);
+                        return itemWrapped;
+                    }
+                    return this.items?.get(item.id);
+                });
+
+                //process data
+                if (onFetch) {
+                    onFetch(wrappedData);
+                }
 
-        let projectName = this.stream.project?.name;
-        if (projectName === "Engine") {
-            projectName = "UE4";
-        }
-
-        const crumbItems: BreadcrumbItem[] = [
-            {
-                text: projectName ?? "Unknown Project",
-                link: `/project/${this.stream?.project?.id}`
-            },
-            {
-                text: this.stream.name,
-                link: `/stream/${this.stream.id}`
-            },
-            {
-                text: `${this.clText}: ${jobdata?.name ?? ""}`,
-                link: `/job/${jobdata?.id}`
-            },
-            {
-                text: this.jobInfo.stepName,
-                link: this.jobInfo.stepId ? `/job/${jobdata?.id}?step=${this.jobInfo.stepId}` : `/job/${jobdata?.id}`
-            },
-            {
-                text: this.name + " (Test Report)" 
+                count = pageData.length;
             }
-        ];
 
-        return crumbItems;
-    }
+        } catch(reason) {
+            console.error(`Failed to fetch TestData history from ${streamId} on key '${key}' !`);
+            console.error(reason);
+        }
 
-    get crumbTitle(): string | undefined {
+        return count;
+    }
 
-        if (!this.stream) {
-            return undefined;
+    async fetchUpdateItems(ids: string[], onFetch?: OnFetchTestData, filter?: string, parallelCount: number = 30) {
+        let count = 0;
+        for (let i = 0; i < ids.length; i += parallelCount) {
+            if (!this.activated) {
+                break;
+            }
+            const chunk = ids.slice(i, i + parallelCount);
+            count += (
+                await Promise.all(chunk.map((id) => this.fetchUpdateItem(id, onFetch, filter)))
+            ).reduce((a, b) => a + b);
         }
-        return `Horde - ${this.clText}: ${this.jobInfo.jobdata?.name ?? ""} - ${this.jobInfo.stepName} (Test Report)`
+        return count;
     }
 
-    findArtifactData(artifactName: string): ArtifactData | undefined {
-        if (!artifactName) {
-            return undefined;
+    async fetchUpdateItem(id: string, onFetch?: OnFetchTestData, filter?: string) {
+        if (!this.key || !this.streamId) {
+            return 0;
         }
+        if(this.loadingNotice.has(id)) {
+            await when(() => !this.loadingNotice.has(id));
+            return this.activated? 1 : 0;
+        }
+        const key = this.key;
+        const streamId = this.streamId;
+        let count = 0;
+        this.acquireLoadingNotice(id);
+        try {
+            const update = await backend.getTestData(id, filter);
 
-        artifactName = artifactName.replace(/\\/g,  '/');
+            // Are we still fetching for the same collection?
+            if(key === this.key && streamId === this.streamId) {
 
-        if (this.artifactMap.has(artifactName)) {
-            return this.artifactMap.get(artifactName);
+                this.initiateCache();
+
+                let item = this.items?.get(id);
+                if(item) {
+                    item.updateProperties(update);
+                } else {
+                    item = new TestDataWrapper(update);
+                    this.items?.set(id, item);
+                    this.updateItemsByChange(item);
+                }
+
+                //process data
+                if (onFetch) {
+                    onFetch([item]);
+                }
+
+                count = 1;
+            }
+        } catch(reason) {
+            console.error(`Failed to fetch updated TestData item ${id} from ${streamId} on key '${key}' !`);
+        } finally { this.releaseLoadingNotice(id) }
+
+        return count;
+    }
+
+    async fetchItemsFromAlternateKey(altKey: string, maxChange?: number, count: number = 1, page: number = 0, filter?: string) {
+        if (!this.key || !this.streamId) {
+            return [];
+        }
+        if(this.loadingNotice.has(altKey)) {
+            await when(() => !this.loadingNotice.has(altKey));
+            return []; // We have no way to know what is the results... :/
         }
 
-        const found = this.jobInfo.artifacts?.find((value) => value.name.indexOf(artifactName) > -1);
-        this.artifactMap.set(artifactName, found);
+        const streamId = this.streamId;
+        let results: TestData[]= [];
+
+        this.acquireLoadingNotice(altKey);
+        try {
+                results = await backend.getTestDataHistory(streamId, altKey, maxChange, count, page, filter);
+        } catch(reason) {
+            console.error(`Failed to fetch Test Data for ${altKey} from ${streamId} !`);
+        } finally { this.releaseLoadingNotice(altKey) }
+
+        return results;
 
-        return found;
     }
-}
+}
\ No newline at end of file
