Horde: Fix conversions for auto-scaling in pools controller

diff --git a/Engine/Source/Programs/Horde/Horde.Build.Tests/PoolsControllerTest.cs b/Engine/Source/Programs/Horde/Horde.Build.Tests/PoolsControllerTest.cs
--- a/Engine/Source/Programs/Horde/Horde.Build.Tests/PoolsControllerTest.cs
+++ b/Engine/Source/Programs/Horde/Horde.Build.Tests/PoolsControllerTest.cs
@@ -5,6 +5,7 @@ using System.Threading.Tasks;
 using Horde.Build.Fleet.Autoscale;
 using HordeServer.Api;
 using HordeServer.Models;
+using HordeServer.Utilities;
 using HordeServerTests;
 using Microsoft.AspNetCore.Mvc;
 using Microsoft.VisualStudio.TestTools.UnitTesting;
@@ -26,6 +27,29 @@ namespace Horde.Build.Tests
 	        Assert.AreEqual(Pool1.SizeStrategy, Response.SizeStrategy);
         }
         
+        [TestMethod]
+        public async Task CreatePoolsTest()
+        {
+	        CreatePoolRequest Request = new CreatePoolRequest
+	        {
+		        Name = "Pool1",
+		        ScaleOutCooldown = 111,
+		        ScaleInCooldown = 222,
+		        SizeStrategy = PoolSizeStrategy.JobQueue,
+		        JobQueueSettings = new HordeServer.Api.JobQueueSettings(new Horde.Build.Fleet.Autoscale.JobQueueSettings(0.35, 0.85))
+	        };
+	        
+	        ActionResult<CreatePoolResponse> Result = await PoolsController.CreatePoolAsync(Request);
+	        CreatePoolResponse Response = Result.Value!;
+		        
+	        IPool Pool = (await PoolService.GetPoolAsync(new StringId<IPool>(Response.Id)))!;
+	        Assert.AreEqual(Request.Name, Pool.Name);
+	        Assert.AreEqual(Request.ScaleOutCooldown, (int)Pool.ScaleOutCooldown!.Value.TotalSeconds);
+	        Assert.AreEqual(Request.ScaleInCooldown, (int)Pool.ScaleInCooldown!.Value.TotalSeconds);
+	        Assert.AreEqual(Request.JobQueueSettings.ScaleOutFactor, Pool.JobQueueSettings!.ScaleOutFactor, 0.0001);
+	        Assert.AreEqual(Request.JobQueueSettings.ScaleInFactor, Pool.JobQueueSettings!.ScaleInFactor, 0.0001);
+        }
+        
         [TestMethod]
         public async Task UpdatePoolTest()
         {
diff --git a/Engine/Source/Programs/Horde/Horde.Build/Api/Pools.cs b/Engine/Source/Programs/Horde/Horde.Build/Api/Pools.cs
--- a/Engine/Source/Programs/Horde/Horde.Build/Api/Pools.cs
+++ b/Engine/Source/Programs/Horde/Horde.Build/Api/Pools.cs
@@ -34,11 +34,19 @@ namespace HordeServer.Api
 	/// <see cref="Horde.Build.Fleet.Autoscale.JobQueueSettings" />
 	public class JobQueueSettings
 	{
-		/// <see cref="Horde.Build.Fleet.Autoscale.JobQueueSettings.ScaleOutFactor" />
-		private double ScaleOutFactor { get; set; }
+		/// <summary>
+		/// Factor translating queue size to additional agents to grow the pool with
+		/// The result is always rounded up to nearest integer. 
+		/// Example: if there are 20 jobs in queue, a factor 0.25 will result in 5 new agents being added (20 * 0.25)
+		/// </summary>
+		public double ScaleOutFactor { get; set; }
 
-		/// <see cref="Horde.Build.Fleet.Autoscale.JobQueueSettings.ScaleInFactor" />
-		private double ScaleInFactor { get; set; }
+		/// <summary>
+		/// Factor by which to shrink the pool size with when queue is empty
+		/// The result is always rounded up to nearest integer.
+		/// Example: when the queue size is zero, a default value of 0.9 will shrink the pool by 10% (current agent count * 0.9)
+		/// </summary>
+		public double ScaleInFactor { get; set; }
 
 		/// <summary>
 		/// Construct a public REST API representation from the internal one
@@ -82,14 +90,14 @@ namespace HordeServer.Api
 		public bool? EnableAutoscaling { get; set; }
 		
 		/// <summary>
-		/// Cooldown time between scale-out events
+		/// Cooldown time between scale-out events in seconds
 		/// </summary>
-		public TimeSpan? ScaleOutCooldown { get; }
+		public int? ScaleOutCooldown { get; set; }
 		
 		/// <summary>
-		/// Cooldown time between scale-in events
+		/// Cooldown time between scale-in events in seconds
 		/// </summary>
-		public TimeSpan? ScaleInCooldown { get; }
+		public int? ScaleInCooldown { get; set; }
 		
 		/// <summary>
 		/// Pool sizing strategy
@@ -163,14 +171,14 @@ namespace HordeServer.Api
 		public bool? EnableAutoscaling { get; set; }
 		
 		/// <summary>
-		/// Cooldown time between scale-out events
+		/// Cooldown time between scale-out events in seconds
 		/// </summary>
-		public TimeSpan? ScaleOutCooldown { get; }
+		public int? ScaleOutCooldown { get; set; }
 		
 		/// <summary>
-		/// Cooldown time between scale-in events
+		/// Cooldown time between scale-in events in seconds
 		/// </summary>
-		public TimeSpan? ScaleInCooldown { get; }
+		public int? ScaleInCooldown { get; set; }
 		
 		/// <summary>
 		/// Pool sizing strategy
@@ -251,14 +259,14 @@ namespace HordeServer.Api
 		public bool EnableAutoscaling { get; set; }
 		
 		/// <summary>
-		/// Cooldown time between scale-out events
+		/// Cooldown time between scale-out events in seconds
 		/// </summary>
-		public TimeSpan? ScaleOutCooldown { get; }
+		public int? ScaleOutCooldown { get; set; }
 		
 		/// <summary>
-		/// Cooldown time between scale-in events
+		/// Cooldown time between scale-in events in seconds
 		/// </summary>
-		public TimeSpan? ScaleInCooldown { get; }
+		public int? ScaleInCooldown { get; set; }
 		
 		/// <summary>
 		/// Pool sizing strategy to be used for this pool
@@ -305,8 +313,8 @@ namespace HordeServer.Api
 			this.Name = Pool.Name;
 			this.Condition = Pool.Condition;
 			this.EnableAutoscaling = Pool.EnableAutoscaling;
-			this.ScaleOutCooldown = Pool.ScaleOutCooldown;
-			this.ScaleInCooldown = Pool.ScaleInCooldown;
+			this.ScaleOutCooldown = Pool.ScaleOutCooldown == null ? null : (int)Pool.ScaleOutCooldown.Value.TotalSeconds;
+			this.ScaleInCooldown = Pool.ScaleInCooldown == null ? null : (int)Pool.ScaleInCooldown.Value.TotalSeconds;
 			this.SizeStrategy = Pool.SizeStrategy;
 			this.LeaseUtilizationSettings = Pool.LeaseUtilizationSettings == null ? null : new LeaseUtilizationSettings(Pool.LeaseUtilizationSettings);
 			this.JobQueueSettings = Pool.JobQueueSettings == null ? null : new JobQueueSettings(Pool.JobQueueSettings);
diff --git a/Engine/Source/Programs/Horde/Horde.Build/Controllers/PoolsController.cs b/Engine/Source/Programs/Horde/Horde.Build/Controllers/PoolsController.cs
--- a/Engine/Source/Programs/Horde/Horde.Build/Controllers/PoolsController.cs
+++ b/Engine/Source/Programs/Horde/Horde.Build/Controllers/PoolsController.cs
@@ -62,10 +62,13 @@ namespace HordeServer.Controllers
 
 			Horde.Build.Fleet.Autoscale.LeaseUtilizationSettings? LuSettings = Create.LeaseUtilizationSettings?.Convert();
 			Horde.Build.Fleet.Autoscale.JobQueueSettings? JqSettings = Create.JobQueueSettings?.Convert();
+			
+			TimeSpan? ScaleOutCooldown = Create.ScaleOutCooldown == null ? null : TimeSpan.FromSeconds(Create.ScaleOutCooldown.Value);
+			TimeSpan? ScaleInCooldown = Create.ScaleInCooldown == null ? null : TimeSpan.FromSeconds(Create.ScaleInCooldown.Value);
 
 			IPool NewPool = await PoolService.CreatePoolAsync(
 				Create.Name, Create.Condition, Create.EnableAutoscaling, Create.MinAgents, Create.NumReserveAgents,
-				Create.ScaleOutCooldown, Create.ScaleInCooldown, Create.SizeStrategy, LuSettings, JqSettings, Create.Properties);
+				ScaleOutCooldown, ScaleInCooldown, Create.SizeStrategy, LuSettings, JqSettings, Create.Properties);
 			return new CreatePoolResponse(NewPool.Id.ToString());
 		}
 
@@ -143,9 +146,12 @@ namespace HordeServer.Controllers
 			{
 				return NotFound(PoolIdValue);
 			}
+			
+			TimeSpan? ScaleOutCooldown = Update.ScaleOutCooldown == null ? null : TimeSpan.FromSeconds(Update.ScaleOutCooldown.Value);
+			TimeSpan? ScaleInCooldown = Update.ScaleInCooldown == null ? null : TimeSpan.FromSeconds(Update.ScaleInCooldown.Value);
 
 			await PoolService.UpdatePoolAsync(Pool, Update.Name, Update.Condition, Update.EnableAutoscaling,
-				Update.MinAgents, Update.NumReserveAgents, Update.Properties, Update.ScaleOutCooldown, Update.ScaleInCooldown, Update.SizeStrategy);
+				Update.MinAgents, Update.NumReserveAgents, Update.Properties, ScaleOutCooldown, ScaleInCooldown, Update.SizeStrategy);
 			return new OkResult();
 		}
 
