Horde Storage - Using the blob index to determine if a blob exists or not - this is opt in per namespace right now to avoid using this for the legacy api (limiting risk) but this is reasonable to do for all namespaces in practise.
The benefit of this approch is that we do not get issues when a blob is present in the filesystem cache but has in practise actually been deleted from bigger blob store (s3). We already used this method for existence checks on s3 requests, so in practise this will only make the cases were the blob existed in the local filesystem cache slightly slower but more accurate.
We could also have made the filesystem cache more accturaley reflect what data we expect to be in it (so we clean it up when we delete a blob) but cleaning up all pods filesystem cache when we decide to delete a blob requires cross pod coordination which we do not do.

Also removed the default keyspace in the connection string as we this is defaulted already.
Lastly I removed the legacy option to run blob store operations in parallel, while that would in theory speed up operations (as we can start the s3 request before doing the filesystem check) it resulted in signficant more load on the latter blob stores which ended up making things slower.

diff --git a/Engine/Source/Programs/Horde/HordeStorage/Horde.Storage/Implementation/Blob/IBlobService.cs b/Engine/Source/Programs/Horde/HordeStorage/Horde.Storage/Implementation/Blob/IBlobService.cs
--- a/Engine/Source/Programs/Horde/HordeStorage/Horde.Storage/Implementation/Blob/IBlobService.cs
+++ b/Engine/Source/Programs/Horde/HordeStorage/Horde.Storage/Implementation/Blob/IBlobService.cs
@@ -15,7 +15,6 @@ using Jupiter.Common.Implementation;
 using Jupiter.Implementation;
 using Microsoft.Extensions.DependencyInjection;
 using Microsoft.Extensions.Options;
-using Nito.AsyncEx;
 using Serilog;
 
 namespace Horde.Storage.Implementation;
@@ -147,29 +146,14 @@ public class BlobService : IBlobService
 
     private async Task<BlobIdentifier> PutObjectToStores(NamespaceId ns, IBufferedPayload bufferedPayload, BlobIdentifier identifier)
     {
-        if (_settings.CurrentValue.BlobStoreParallel)
+        foreach (IBlobStore store in _blobStores)
         {
-            await Parallel.ForEachAsync(_blobStores, async (store, cancellationToken) =>
-            {
-                using IScope scope = Tracer.Instance.StartActive("put_blob_to_store");
-                scope.Span.ResourceName = identifier.ToString();
-                scope.Span.SetTag("store", store.ToString());
+            using IScope scope = Tracer.Instance.StartActive("put_blob_to_store");
+            scope.Span.ResourceName = identifier.ToString();
+            scope.Span.SetTag("store", store.ToString());
 
-                await using Stream s = bufferedPayload.GetStream();
-                await store.PutObject(ns, s, identifier);
-            });
-        }
-        else
-        {
-            foreach (IBlobStore store in _blobStores)
-            {
-                using IScope scope = Tracer.Instance.StartActive("put_blob_to_store");
-                scope.Span.ResourceName = identifier.ToString();
-                scope.Span.SetTag("store", store.ToString());
-
-                await using Stream s = bufferedPayload.GetStream();
-                await store.PutObject(ns, s, identifier);
-            }
+            await using Stream s = bufferedPayload.GetStream();
+            await store.PutObject(ns, s, identifier);
         }
 
 
@@ -178,35 +162,15 @@ public class BlobService : IBlobService
 
     private async Task<BlobIdentifier> PutObjectToStores(NamespaceId ns, byte[] payload, BlobIdentifier identifier)
     {
-        if (_settings.CurrentValue.BlobStoreParallel)
+        foreach (IBlobStore store in _blobStores)
         {
-            await Parallel.ForEachAsync(_blobStores,
-                async (store, cancellationToken) => { await store.PutObject(ns, payload, identifier); });
-        }
-        else
-        {
-            foreach (IBlobStore store in _blobStores)
-            {
-                await store.PutObject(ns, payload, identifier);
-            }
+            await store.PutObject(ns, payload, identifier);
         }
 
         return identifier;
     }
 
-    public Task<BlobContents> GetObject(NamespaceId ns, BlobIdentifier blob)
-    {
-        if (_settings.CurrentValue.BlobStoreParallel)
-        {
-            return GetObjectParallel(ns, blob);
-        }
-        else
-        {
-            return GetObjectSequential(ns, blob);
-        }
-    }
-
-    private async Task<BlobContents> GetObjectSequential(NamespaceId ns, BlobIdentifier blob)
+    public async Task<BlobContents> GetObject(NamespaceId ns, BlobIdentifier blob)
     {
         bool seenBlobNotFound = false;
         bool seenNamespaceNotFound = false;
@@ -273,144 +237,26 @@ public class BlobService : IBlobService
         return blobContents;
     }
 
-
-    private async Task<BlobContents> GetObjectParallel(NamespaceId ns, BlobIdentifier blob)
+    public async Task<bool> Exists(NamespaceId ns, BlobIdentifier blob)
     {
-        bool seenBlobNotFound = false;
-        bool seenNamespaceNotFound = false;
-
-        List<Task<(BlobContents?, IBlobStore?)>> getTasks = new();
-
-        foreach (IBlobStore store in _blobStores)
-        {
-            getTasks.Add(Task.Run( async () =>
-            {
-                using IScope scope = Tracer.Instance.StartActive("HierarchicalStore.GetObject");
-                scope.Span.SetTag("BlobStore", store.GetType().Name);
-                scope.Span.SetTag("ObjectFound", false.ToString());
-                try
-                {
-                    BlobContents? blobContents = await store.GetObject(ns, blob);
-                    scope.Span.SetTag("ObjectFound", true.ToString());
-                    return ((BlobContents?)blobContents, (IBlobStore?)store);
-                }
-                catch (BlobNotFoundException)
-                {
-                    seenBlobNotFound = true;
-                }
-                catch (NamespaceNotFoundException)
-                {
-                    seenNamespaceNotFound = true;
-                }
-
-                return (null, null);
-            }));
-        }
-
-        BlobContents? blobContents = null;
-        IBlobStore? storeWithBlob = null;
-        while (getTasks.Count != 0)
-        {
-            Task<(BlobContents?, IBlobStore?)> completedTask = await getTasks.WhenAny();
-            (BlobContents? contents, IBlobStore? store)= await completedTask;
-            if (contents == null)
-            {
-                // the task did not find a object, we wait for the other tasks to run
-                getTasks.Remove(completedTask);
-                continue;
-            }
-
-            blobContents = contents;
-            storeWithBlob = store;
-            // object found, we can break
-            break;
-        }
-
-        if (seenBlobNotFound && blobContents == null)
-        {
-            throw new BlobNotFoundException(ns, blob);
-        }
-
-        if (seenNamespaceNotFound && blobContents == null)
-        {
-            throw new NamespaceNotFoundException(ns);
-        }
-        
-        if (blobContents == null)
-        {
-            // Should not happen but exists to safeguard against the null pointer
-            throw new Exception("blobContents is null");
-        }
-        if (storeWithBlob == null)
-        {
-            // Should not happen but exists to safeguard against the null pointer
-            throw new Exception("storeWithBlob is null");
-        }
-
-        int indexOfStoreWithBlob = _blobStores.IndexOf(storeWithBlob);
-        // if it wasn't the first store we need to populate the stores before it with the blob
-        if (indexOfStoreWithBlob != 0)
-        {
-            using IScope _ = Tracer.Instance.StartActive("HierarchicalStore.Populate");
-            // buffer the blob contents as it could be to large for us to buffer into memory before forwarding to the blob stores
-            await using Stream s = blobContents.Stream;
-            using FilesystemBufferedPayload payload = await FilesystemBufferedPayload.Create(s);
-
-            // Don't populate the last store, as that is where we got the hit
-            // Populate each store traversed that did not have the content found lower in the hierarchy
-            await Parallel.ForEachAsync(_blobStores.GetRange(0, indexOfStoreWithBlob), async (store, token) =>
-            {
-                using IScope scope = Tracer.Instance.StartActive("HierarchicalStore.PopulateStore");
-                scope.Span.SetTag("BlobStore", store.GetType().Name);
-
-                await using Stream s = payload.GetStream();
-                await store.PutObject(ns, s, blob);
-            });
-
-            // we just pushed the content to all stores, thus consuming it, so we just refetch it from the store that just reported it had it
-            blobContents = await storeWithBlob.GetObject(ns, blob);
-        }
-        
-        return blobContents;
-    }
-
-    public Task<bool> Exists(NamespaceId ns, BlobIdentifier blob)
-    {
-        if (_settings.CurrentValue.BlobStoreParallel)
-        {
-            return ExistsParallel(ns, blob);
-        }
-        else
-        {
-            return ExistsSequential(ns, blob);
-        }
-    }
-
-    private async Task<bool> ExistsSequential(NamespaceId ns, BlobIdentifier blob)
-    {
-        foreach (IBlobStore store in _blobStores)
+        bool useBlobIndex = _settings.CurrentValue.NamespacesThatUseBlobIndexForExistsCheck.Contains(ns.ToString());
+        if (useBlobIndex)
         {
             using IScope scope = Tracer.Instance.StartActive("HierarchicalStore.ObjectExists");
-            scope.Span.SetTag("BlobStore", store.GetType().Name);
-            if (await store.Exists(ns, blob))
+            scope.Span.SetTag("BlobStore", "BlobIndex");
+            bool exists = await _blobIndex.BlobExistsInRegion(ns, blob);
+            if (exists)
             {
                 scope.Span.SetTag("ObjectFound", true.ToString());
                 return true;
             }
+
             scope.Span.SetTag("ObjectFound", false.ToString());
+            return false;
         }
-
-        return false;
-    }
-
-
-    private async Task<bool> ExistsParallel(NamespaceId ns, BlobIdentifier blob)
-    {
-        List<Task<bool>> existsTasks = new();
-
-        foreach (IBlobStore store in _blobStores)
+        else
         {
-            existsTasks.Add(Task.Run( async () =>
+            foreach (IBlobStore store in _blobStores)
             {
                 using IScope scope = Tracer.Instance.StartActive("HierarchicalStore.ObjectExists");
                 scope.Span.SetTag("BlobStore", store.GetType().Name);
@@ -419,29 +265,13 @@ public class BlobService : IBlobService
                     scope.Span.SetTag("ObjectFound", true.ToString());
                     return true;
                 }
-
                 scope.Span.SetTag("ObjectFound", false.ToString());
-                return false;
-            }));
-        }
-
-        while (existsTasks.Count != 0)
-        {
-            Task<bool> completedTask = await existsTasks.WhenAny();
-            bool exist = await completedTask;
-            if (!exist)
-            {
-                // the task did not find a object, we wait for the other tasks to run
-                existsTasks.Remove(completedTask);
-                continue;
             }
 
-            // object found, we can break
-            return true;
+            return false;
         }
-        return false;
     }
-
+    
     public async Task DeleteObject(NamespaceId ns, BlobIdentifier blob)
     {
         bool blobNotFound = false;
diff --git a/Engine/Source/Programs/Horde/HordeStorage/Horde.Storage/Settings.cs b/Engine/Source/Programs/Horde/HordeStorage/Horde.Storage/Settings.cs
--- a/Engine/Source/Programs/Horde/HordeStorage/Horde.Storage/Settings.cs
+++ b/Engine/Source/Programs/Horde/HordeStorage/Horde.Storage/Settings.cs
@@ -132,11 +132,11 @@ namespace Horde.Storage
         public ContentIdStoreImplementations ContentIdStoreImplementation { get; set; } = ContentIdStoreImplementations.Memory;
         public BlobIndexImplementations BlobIndexImplementation { get; set; } = BlobIndexImplementations.Memory;
 
+        public List<string> NamespacesThatUseBlobIndexForExistsCheck { get; set; } = new List<string>();
+
         public int? MaxSingleBlobSize { get; set; } = null; // disable blob partitioning
 
         public int LastAccessRollupFrequencySeconds { get; set; } = 900; // 15 minutes
-
-        public bool BlobStoreParallel { get; set; } = false;
     }
 
     public class MongoSettings
