Horde: Move default values for scale-in/out cooldowns to settings

diff --git a/Engine/Source/Programs/Horde/Horde.Build/Fleet/Autoscale/AutoscaleServiceV2.cs b/Engine/Source/Programs/Horde/Horde.Build/Fleet/Autoscale/AutoscaleServiceV2.cs
--- a/Engine/Source/Programs/Horde/Horde.Build/Fleet/Autoscale/AutoscaleServiceV2.cs
+++ b/Engine/Source/Programs/Horde/Horde.Build/Fleet/Autoscale/AutoscaleServiceV2.cs
@@ -28,20 +28,19 @@ namespace Horde.Build.Fleet.Autoscale
 	/// </summary>
 	public sealed class AutoscaleServiceV2 : IHostedService, IDisposable
 	{
-		static readonly TimeSpan DefaultScaleOutCooldown = TimeSpan.FromMinutes(5);
-		static readonly TimeSpan DefaultScaleInCooldown = TimeSpan.FromMinutes(20);
-
-		IPoolSizeStrategy LeaseUtilizationStrategy;
-		IPoolSizeStrategy JobQueueStrategy;
-		IPoolSizeStrategy NoOpStrategy;
-		IAgentCollection AgentCollection;
-		IPoolCollection PoolCollection;
-		IFleetManager FleetManager;
-		IDogStatsd DogStatsd;
-		IClock Clock;
-		ITicker Ticker;
-		ServerSettings Settings;
-		ILogger<AutoscaleServiceV2> Logger;
+		private IPoolSizeStrategy LeaseUtilizationStrategy;
+		private IPoolSizeStrategy JobQueueStrategy;
+		private IPoolSizeStrategy NoOpStrategy;
+		private readonly IAgentCollection AgentCollection;
+		private readonly IPoolCollection PoolCollection;
+		private readonly IFleetManager FleetManager;
+		private readonly IDogStatsd DogStatsd;
+		private readonly IClock Clock;
+		private readonly ITicker Ticker;
+		private readonly ServerSettings Settings;
+		private readonly TimeSpan DefaultScaleOutCooldown;
+		private readonly TimeSpan DefaultScaleInCooldown;
+		private readonly ILogger<AutoscaleServiceV2> Logger;
 
 		/// <summary>
 		/// Constructor
@@ -59,6 +58,8 @@ namespace Horde.Build.Fleet.Autoscale
 			this.Ticker = Clock.AddSharedTicker<AutoscaleServiceV2>(TimeSpan.FromMinutes(5.0), TickLeaderAsync, Logger);
 			this.Settings = Settings.Value;
 			this.Logger = Logger;
+			this.DefaultScaleOutCooldown = TimeSpan.FromSeconds(Settings.Value.AgentPoolScaleOutCooldownSeconds);
+			this.DefaultScaleInCooldown = TimeSpan.FromSeconds(Settings.Value.AgentPoolScaleInCooldownSeconds);
 		}
 
 		/// <inheritdoc/>
diff --git a/Engine/Source/Programs/Horde/Horde.Build/ServerSettings.cs b/Engine/Source/Programs/Horde/Horde.Build/ServerSettings.cs
--- a/Engine/Source/Programs/Horde/Horde.Build/ServerSettings.cs
+++ b/Engine/Source/Programs/Horde/Horde.Build/ServerSettings.cs
@@ -489,6 +489,16 @@ namespace HordeServer
 		/// </summary>
 		public PoolSizeStrategy DefaultAgentPoolSizeStrategy { get; set; } = PoolSizeStrategy.LeaseUtilization;
 
+		/// <summary>
+		/// Scale-out cooldown for auto-scaling agent pools (in seconds). Can be overridden by per-pool settings.
+		/// </summary>
+		public int AgentPoolScaleOutCooldownSeconds { get; set; } = 60; // 1 min
+		
+		/// <summary>
+		/// Scale-in cooldown for auto-scaling agent pools (in seconds). Can be overridden by per-pool settings.
+		/// </summary>
+		public int AgentPoolScaleInCooldownSeconds { get; set; } = 1200; // 20 mins
+
 		/// <summary>
 		/// Set the minimum size of the global thread pool
 		/// This value has been found in need of tweaking to avoid timeouts with the Redis client during bursts
