Horde: Add generic tool exited event messages to systemic matcher
#skipci
#rnx

diff --git a/Engine/Source/Programs/Horde/Horde.Build/Issues/Impl/Handlers/SystemicIssueHandler.cs b/Engine/Source/Programs/Horde/Horde.Build/Issues/Impl/Handlers/SystemicIssueHandler.cs
--- a/Engine/Source/Programs/Horde/Horde.Build/Issues/Impl/Handlers/SystemicIssueHandler.cs
+++ b/Engine/Source/Programs/Horde/Horde.Build/Issues/Impl/Handlers/SystemicIssueHandler.cs
@@ -67,10 +67,16 @@ namespace HordeServer.IssueHandlers.Impl
 				for (int i = 0; i < EventData.Lines.Count; i++)
 				{
 					string Message = EventData.Lines[i].Message;
-					if (Message.Contains("AutomationTool exiting with ExitCode", StringComparison.InvariantCultureIgnoreCase) || Message.Contains("BUILD FAILED", StringComparison.InvariantCultureIgnoreCase))
+
+					string[] SystemicExitMessages = new string[] { "AutomationTool exiting with ExitCode", "BUILD FAILED", "tool returned code" };
+
+					for (int j = 0; j < SystemicExitMessages.Length; j++)
 					{
-						Fingerprint = new NewIssueFingerprint(Type, new[] { Node.Name }, null);
-						return true;
+						if (Message.Contains(SystemicExitMessages[j], StringComparison.InvariantCultureIgnoreCase))
+						{
+							Fingerprint = new NewIssueFingerprint(Type, new[] { Node.Name }, null);
+							return true;
+						}
 					}
 				}
 			}
