Horde Storage - If a replicator gets multiple use snapshot responses it can just stop running as it will have reached the latest state.

diff --git a/Engine/Source/Programs/Horde/HordeStorage/Horde.Storage/Implementation/Replication/RefsReplicator.cs b/Engine/Source/Programs/Horde/HordeStorage/Horde.Storage/Implementation/Replication/RefsReplicator.cs
--- a/Engine/Source/Programs/Horde/HordeStorage/Horde.Storage/Implementation/Replication/RefsReplicator.cs
+++ b/Engine/Source/Programs/Horde/HordeStorage/Horde.Storage/Implementation/Replication/RefsReplicator.cs
@@ -223,9 +223,11 @@ namespace Horde.Storage.Implementation
 
                     if (useSnapshotException != null)
                     {
-                        // if we have already attempted to recover using a snapshot and we still fail we just give up and throw the exception onwards.
+                        // if we have already attempted to recover using a snapshot and we still fail we just give up as there is no content to fetch
                         if (haveAttemptedSnapshot)
-                            throw new AggregateException(useSnapshotException);
+                        {
+                            break;
+                        }
 
                         // if we fail to replicate incrementally we revert to using a snapshot
                         haveAttemptedSnapshot = true;
