Horde Storage - Cleanup GC settings which was a mess, they were inverted both in how they were configured and what they did (so end result was correct luckily).
Renamed CleanNamespacesV1 setting to CleanNamespacesLegacy as is used to cleanup the legacy (dynamo based) namespaces.

diff --git a/Engine/Source/Programs/Horde/HordeStorage/Horde.Storage/Implementation/GC/CacheCleanup.cs b/Engine/Source/Programs/Horde/HordeStorage/Horde.Storage/Implementation/GC/CacheCleanup.cs
--- a/Engine/Source/Programs/Horde/HordeStorage/Horde.Storage/Implementation/GC/CacheCleanup.cs
+++ b/Engine/Source/Programs/Horde/HordeStorage/Horde.Storage/Implementation/GC/CacheCleanup.cs
@@ -45,14 +45,14 @@ namespace Horde.Storage.Implementation
                 // do not apply our cleanup policies to the internal namespace
                 return Task.FromResult(new List<OldRecord>());
             }
-            else if (_settings.CurrentValue.CleanNamespacesV1.Contains(ns.ToString()))
-            {
-                return CleanNamespaceV1(ns, cancellationToken);
-            }
             else if (_settings.CurrentValue.CleanNamespaces.Contains(ns.ToString()))
             {
                 return CleanNamespace(ns, cancellationToken);
             }
+            else if (_settings.CurrentValue.CleanNamespacesLegacy.Contains(ns.ToString()))
+            {
+                return CleanNamespaceLegacy(ns, cancellationToken);
+            }
             else
                 throw new NotImplementedException(
                     $"Namespace {ns} not present in CleanNamespaces or CleanNamespacesV1 lists, unable to clean it as we do not know which method to use.");
@@ -105,7 +105,7 @@ namespace Horde.Storage.Implementation
             return deletedRecords;
         }
 
-        private async Task<List<OldRecord>> CleanNamespaceV1(NamespaceId ns, CancellationToken cancellationToken)
+        private async Task<List<OldRecord>> CleanNamespaceLegacy(NamespaceId ns, CancellationToken cancellationToken)
         {
             List<OldRecord> deletedRecords = new List<OldRecord>();
             await foreach (OldRecord record in _refs.GetOldRecords(ns, _settings.CurrentValue.LastAccessCutoff).WithCancellation(cancellationToken))
diff --git a/Engine/Source/Programs/Horde/HordeStorage/Horde.Storage/Implementation/GC/OrphanBlobCleanup.cs b/Engine/Source/Programs/Horde/HordeStorage/Horde.Storage/Implementation/GC/OrphanBlobCleanup.cs
--- a/Engine/Source/Programs/Horde/HordeStorage/Horde.Storage/Implementation/GC/OrphanBlobCleanup.cs
+++ b/Engine/Source/Programs/Horde/HordeStorage/Horde.Storage/Implementation/GC/OrphanBlobCleanup.cs
@@ -167,7 +167,7 @@ namespace Horde.Storage.Implementation
 
         private bool NamespaceShouldBeCleaned(NamespaceId ns)
         {
-            return _gcSettings.CurrentValue.CleanNamespacesV1.Contains(ns.ToString());
+            return _gcSettings.CurrentValue.CleanNamespacesLegacy.Contains(ns.ToString());
         }
 
         public virtual async IAsyncEnumerable<NamespaceId> ListNamespaces()
diff --git a/Engine/Source/Programs/Horde/HordeStorage/Horde.Storage/Tests/Functional/GC/GCBlobTests.cs b/Engine/Source/Programs/Horde/HordeStorage/Horde.Storage/Tests/Functional/GC/GCBlobTests.cs
--- a/Engine/Source/Programs/Horde/HordeStorage/Horde.Storage/Tests/Functional/GC/GCBlobTests.cs
+++ b/Engine/Source/Programs/Horde/HordeStorage/Horde.Storage/Tests/Functional/GC/GCBlobTests.cs
@@ -171,7 +171,7 @@ namespace Horde.Storage.FunctionalTests.GC
             {
                 CleanOldBlobs = true,
                 CleanOldBlobsLegacy = true,
-                CleanNamespacesV1 = new List<string> { TestNamespace.ToString() }
+                CleanNamespacesLegacy = new List<string> { TestNamespace.ToString() }
             };
             IOptionsMonitor<GCSettings> gcSettingsMon = Mock.Of<IOptionsMonitor<GCSettings>>(_ => _.CurrentValue == gcSettings);
 
@@ -234,7 +234,7 @@ namespace Horde.Storage.FunctionalTests.GC
                     new KeyValuePair<string, string>("Horde_Storage:StorageImplementations:0", "MemoryBlobStore"),
                     new KeyValuePair<string, string>("GC:CleanNamespaces:0", TestNamespace.ToString()),
                     new KeyValuePair<string, string>("GC:CleanOldBlobs", true.ToString()),
-                    new KeyValuePair<string, string>("GC:CleanOldBlobsLegacy", true.ToString()),
+                    new KeyValuePair<string, string>("GC:CleanOldBlobsLegacy", false.ToString()),
                     new KeyValuePair<string, string>("Horde_Storage:BlobIndexImplementation", GetImplementation()),
                     
                 })
