Horde Storage - Added ability to ondemand replicate blobs from other regions. This is opt in by specifying both a feature toggle as well as enabling it for the given namespace.

diff --git a/Engine/Source/Programs/Horde/HordeStorage/Horde.Storage/Controllers/StorageController.cs b/Engine/Source/Programs/Horde/HordeStorage/Horde.Storage/Controllers/StorageController.cs
--- a/Engine/Source/Programs/Horde/HordeStorage/Horde.Storage/Controllers/StorageController.cs
+++ b/Engine/Source/Programs/Horde/HordeStorage/Horde.Storage/Controllers/StorageController.cs
@@ -151,7 +151,16 @@ namespace Horde.Storage.Controllers
 
         private async Task<BlobContents> GetImpl(NamespaceId ns, BlobIdentifier blob)
         {
-            return await _storage.GetObject(ns, blob);
+            try
+            {
+                return await _storage.GetObject(ns, blob);
+            }
+            catch (BlobNotFoundException)
+            {
+                if (!_storage.ShouldFetchBlobOnDemand(ns))
+                    throw;
+                return await _storage.ReplicateObject(ns, blob);
+            }
         }
 
         [HttpPut("{ns}/{id}")]
diff --git a/Engine/Source/Programs/Horde/HordeStorage/Horde.Storage/Implementation/Blob/CachedBlobIndex.cs b/Engine/Source/Programs/Horde/HordeStorage/Horde.Storage/Implementation/Blob/CachedBlobIndex.cs
--- a/Engine/Source/Programs/Horde/HordeStorage/Horde.Storage/Implementation/Blob/CachedBlobIndex.cs
+++ b/Engine/Source/Programs/Horde/HordeStorage/Horde.Storage/Implementation/Blob/CachedBlobIndex.cs
@@ -36,7 +36,7 @@ public class CachedBlobIndex : IBlobIndex
         return null;
     }
 
-    public async Task AddBlobToIndex(NamespaceId ns, BlobIdentifier id)
+    public async Task AddBlobToIndex(NamespaceId ns, BlobIdentifier id, string? region = null)
     {
         // We do not actually track any blob information when running in cached mode
         await Task.CompletedTask;
diff --git a/Engine/Source/Programs/Horde/HordeStorage/Horde.Storage/Implementation/Blob/IBlobIndex.cs b/Engine/Source/Programs/Horde/HordeStorage/Horde.Storage/Implementation/Blob/IBlobIndex.cs
--- a/Engine/Source/Programs/Horde/HordeStorage/Horde.Storage/Implementation/Blob/IBlobIndex.cs
+++ b/Engine/Source/Programs/Horde/HordeStorage/Horde.Storage/Implementation/Blob/IBlobIndex.cs
@@ -16,7 +16,7 @@ namespace Horde.Storage.Implementation.Blob
             public List<(BucketId, IoHashKey)> References { get; init; } = new List<(BucketId, IoHashKey)>();
         }
 
-        Task AddBlobToIndex(NamespaceId ns, BlobIdentifier id);
+        Task AddBlobToIndex(NamespaceId ns, BlobIdentifier id, string? region = null);
 
         Task<BlobInfo?> GetBlobInfo(NamespaceId ns, BlobIdentifier id);
 
diff --git a/Engine/Source/Programs/Horde/HordeStorage/Horde.Storage/Implementation/Blob/IBlobService.cs b/Engine/Source/Programs/Horde/HordeStorage/Horde.Storage/Implementation/Blob/IBlobService.cs
--- a/Engine/Source/Programs/Horde/HordeStorage/Horde.Storage/Implementation/Blob/IBlobService.cs
+++ b/Engine/Source/Programs/Horde/HordeStorage/Horde.Storage/Implementation/Blob/IBlobService.cs
@@ -5,6 +5,8 @@ using System.Collections.Concurrent;
 using System.Collections.Generic;
 using System.IO;
 using System.Linq;
+using System.Net;
+using System.Net.Http;
 using System.Threading.Tasks;
 using Dasync.Collections;
 using Datadog.Trace;
@@ -27,6 +29,9 @@ public interface IBlobService
     Task<BlobIdentifier> PutObject(NamespaceId ns, byte[] payload, BlobIdentifier identifier);
 
     Task<BlobContents> GetObject(NamespaceId ns, BlobIdentifier blob);
+
+    Task<BlobContents> ReplicateObject(NamespaceId ns, BlobIdentifier blob);
+
     Task<bool> Exists(NamespaceId ns, BlobIdentifier blob);
 
     // Delete a object
@@ -41,6 +46,7 @@ public interface IBlobService
     Task<BlobIdentifier[]> FilterOutKnownBlobs(NamespaceId ns, IAsyncEnumerable<BlobIdentifier> blobs);
     Task<BlobContents> GetObjects(NamespaceId ns, BlobIdentifier[] refRequestBlobReferences);
 
+    bool ShouldFetchBlobOnDemand(NamespaceId ns);
 }
 
 public class BlobService : IBlobService
@@ -48,6 +54,9 @@ public class BlobService : IBlobService
     private List<IBlobStore> _blobStores;
     private readonly IOptionsMonitor<HordeStorageSettings> _settings;
     private readonly IBlobIndex _blobIndex;
+    private readonly IPeerStatusService _peerStatusService;
+    private readonly IHttpClientFactory _httpClientFactory;
+    private readonly IServiceCredentials _serviceCredentials;
     private readonly ILogger _logger = Log.ForContext<BlobService>();
 
     internal IEnumerable<IBlobStore> BlobStore
@@ -56,11 +65,14 @@ public class BlobService : IBlobService
         set { _blobStores = value.ToList(); } 
     }
 
-    public BlobService(IServiceProvider provider, IOptionsMonitor<HordeStorageSettings> settings, IBlobIndex blobIndex)
+    public BlobService(IServiceProvider provider, IOptionsMonitor<HordeStorageSettings> settings, IBlobIndex blobIndex, IPeerStatusService peerStatusService, IHttpClientFactory httpClientFactory, IServiceCredentials serviceCredentials)
     {
         _blobStores = GetBlobStores(provider, settings).ToList();
         _settings = settings;
         _blobIndex = blobIndex;
+        _peerStatusService = peerStatusService;
+        _httpClientFactory = httpClientFactory;
+        _serviceCredentials = serviceCredentials;
     }
 
     private IEnumerable<IBlobStore> GetBlobStores(IServiceProvider provider, IOptionsMonitor<HordeStorageSettings> settings)
@@ -237,6 +249,60 @@ public class BlobService : IBlobService
         return blobContents;
     }
 
+    public async Task<BlobContents> ReplicateObject(NamespaceId ns, BlobIdentifier blob)
+    {
+        if (!ShouldFetchBlobOnDemand(ns))
+            throw new NotSupportedException($"Replication is not allowed in namespace {ns}");
+
+        using IScope scope = Tracer.Instance.StartActive("HierarchicalStore.Replicate");
+
+        IBlobIndex.BlobInfo? blobInfo = await _blobIndex.GetBlobInfo(ns, blob);
+        if (blobInfo == null)
+            throw new BlobNotFoundException(ns, blob);
+
+        _logger.Information("On-demand replicating blob {Blob} in Namespace {Ns}", blob, ns);
+        SortedList<int, string> possiblePeers = _peerStatusService.GetPeersByLatency(blobInfo.Regions.ToList());
+
+        string bestPeer = possiblePeers.First().Value;
+        IPeerStatusService.PeerStatus? peerStatus = _peerStatusService.GetPeerStatus(bestPeer);
+        if (peerStatus == null)
+            throw new Exception($"Failed to find peer {bestPeer}");
+        PeerEndpoints peerEndpoint = peerStatus.Endpoints.First();
+        HttpClient httpClient = _httpClientFactory.CreateClient();
+        HttpRequestMessage blobRequest = BuildHttpRequest(HttpMethod.Get, $"{peerEndpoint.Url}/api/v1/blobs/{ns}/{blob}");
+        HttpResponseMessage blobResponse = await httpClient.SendAsync(blobRequest, HttpCompletionOption.ResponseHeadersRead);
+
+        if (blobResponse.StatusCode == HttpStatusCode.NotFound)
+        {
+            throw new BlobReplicationException(ns, blob, $"Failed to replicate {blob} in {ns} due to it not existing in region {bestPeer}.");
+        }
+
+        if (blobResponse.StatusCode != HttpStatusCode.OK)
+        {
+            throw new BlobReplicationException(ns, blob, $"Failed to replicate {blob} in {ns} from region {bestPeer} due to bad http status code {blobResponse.StatusCode}.");
+        }
+
+        await using Stream s = await blobResponse.Content.ReadAsStreamAsync();
+        using FilesystemBufferedPayload payload = await FilesystemBufferedPayload.Create(s);
+        await PutObject(ns, payload, blob);
+
+        return await GetObject(ns, blob);
+    }
+
+    public bool ShouldFetchBlobOnDemand(NamespaceId ns)
+    {
+        return _settings.CurrentValue.EnableOnDemandReplication && _settings.CurrentValue.OnDemandReplicationNamespaces.Any(s => s.Equals("*") || s.Equals(ns.ToString()));
+    }
+
+    private HttpRequestMessage BuildHttpRequest(HttpMethod httpMethod, string uri)
+    {
+        string? token = _serviceCredentials.GetToken();
+        HttpRequestMessage request = new HttpRequestMessage(httpMethod, uri);
+        if (!string.IsNullOrEmpty(token))
+            request.Headers.Add("Authorization", "Bearer " + token);
+        return request;
+    }
+
     public async Task<bool> Exists(NamespaceId ns, BlobIdentifier blob)
     {
         bool useBlobIndex = _settings.CurrentValue.NamespacesThatUseBlobIndexForExistsCheck.Contains(ns.ToString());
diff --git a/Engine/Source/Programs/Horde/HordeStorage/Horde.Storage/Implementation/Blob/IBlobStore.cs b/Engine/Source/Programs/Horde/HordeStorage/Horde.Storage/Implementation/Blob/IBlobStore.cs
--- a/Engine/Source/Programs/Horde/HordeStorage/Horde.Storage/Implementation/Blob/IBlobStore.cs
+++ b/Engine/Source/Programs/Horde/HordeStorage/Horde.Storage/Implementation/Blob/IBlobStore.cs
@@ -43,6 +43,19 @@ namespace Horde.Storage.Implementation
             Ns = ns;
             Blob = blob;
         }
+
+        public BlobNotFoundException(NamespaceId ns, BlobIdentifier blob, string message) : base(message)
+        {
+            Ns = ns;
+            Blob = blob;
+        }
+    }
+
+    public class BlobReplicationException : BlobNotFoundException
+    {
+        public BlobReplicationException(NamespaceId ns, BlobIdentifier blob, string message) : base(ns, blob, message)
+        {
+        }
     }
 
     public class NamespaceNotFoundException : Exception
diff --git a/Engine/Source/Programs/Horde/HordeStorage/Horde.Storage/Implementation/Blob/MemoryBlobIndex.cs b/Engine/Source/Programs/Horde/HordeStorage/Horde.Storage/Implementation/Blob/MemoryBlobIndex.cs
--- a/Engine/Source/Programs/Horde/HordeStorage/Horde.Storage/Implementation/Blob/MemoryBlobIndex.cs
+++ b/Engine/Source/Programs/Horde/HordeStorage/Horde.Storage/Implementation/Blob/MemoryBlobIndex.cs
@@ -31,10 +31,11 @@ public class MemoryBlobIndex : IBlobIndex
         return _index.GetOrAdd(ns, id => new ConcurrentDictionary<BlobIdentifier, MemoryBlobInfo>());
     }
 
-    public Task AddBlobToIndex(NamespaceId ns, BlobIdentifier id)
+    public Task AddBlobToIndex(NamespaceId ns, BlobIdentifier id, string? region = null)
     {
+        region ??= _jupiterSettings.CurrentValue.CurrentSite;
         ConcurrentDictionary<BlobIdentifier, MemoryBlobInfo> index = GetNamespaceContainer(ns);
-        index[id] = NewBlobInfo(ns);
+        index[id] = NewBlobInfo(ns, region);
         return Task.CompletedTask;
     }
 
@@ -69,7 +70,7 @@ public class MemoryBlobIndex : IBlobIndex
 
             index.AddOrUpdate(id, _ =>
             {
-                MemoryBlobInfo info = NewBlobInfo(ns);
+                MemoryBlobInfo info = NewBlobInfo(ns, _jupiterSettings.CurrentValue.CurrentSite);
                 info.References.Add((bucket, key));
                 return info;
             }, (_, info) =>
@@ -82,11 +83,11 @@ public class MemoryBlobIndex : IBlobIndex
         return Task.CompletedTask;
     }
 
-    private MemoryBlobInfo NewBlobInfo(NamespaceId ns)
+    private MemoryBlobInfo NewBlobInfo(NamespaceId ns, string region)
     {
         MemoryBlobInfo info = new MemoryBlobInfo
         {
-            Regions = new HashSet<string> { _jupiterSettings.CurrentValue.CurrentSite },
+            Regions = new HashSet<string> { region },
             Namespace = ns,
         };
         return info;
diff --git a/Engine/Source/Programs/Horde/HordeStorage/Horde.Storage/Implementation/Blob/MongoBlobIndex.cs b/Engine/Source/Programs/Horde/HordeStorage/Horde.Storage/Implementation/Blob/MongoBlobIndex.cs
--- a/Engine/Source/Programs/Horde/HordeStorage/Horde.Storage/Implementation/Blob/MongoBlobIndex.cs
+++ b/Engine/Source/Programs/Horde/HordeStorage/Horde.Storage/Implementation/Blob/MongoBlobIndex.cs
@@ -40,11 +40,12 @@ public class MongoBlobIndex : MongoStore, IBlobIndex
         });
     }
 
-    public async Task AddBlobToIndex(NamespaceId ns, BlobIdentifier id)
+    public async Task AddBlobToIndex(NamespaceId ns, BlobIdentifier id, string? region = null)
     {
+        region ??= _jupiterSettings.CurrentValue.CurrentSite;
         IMongoCollection<MongoBlobIndexModelV0> collection = GetCollection<MongoBlobIndexModelV0>();
         MongoBlobIndexModelV0 model = new MongoBlobIndexModelV0(ns, id);
-        model.Regions.Add(_jupiterSettings.CurrentValue.CurrentSite);
+        model.Regions.Add(region);
             
         FilterDefinition<MongoBlobIndexModelV0> filter = Builders<MongoBlobIndexModelV0>.Filter.Where(m => m.Ns == ns.ToString() && m.BlobId == id.ToString());
         await collection.FindOneAndReplaceAsync(filter, model, new FindOneAndReplaceOptions<MongoBlobIndexModelV0, MongoBlobIndexModelV0>
diff --git a/Engine/Source/Programs/Horde/HordeStorage/Horde.Storage/Implementation/Blob/ScyllaBlobIndex.cs b/Engine/Source/Programs/Horde/HordeStorage/Horde.Storage/Implementation/Blob/ScyllaBlobIndex.cs
--- a/Engine/Source/Programs/Horde/HordeStorage/Horde.Storage/Implementation/Blob/ScyllaBlobIndex.cs
+++ b/Engine/Source/Programs/Horde/HordeStorage/Horde.Storage/Implementation/Blob/ScyllaBlobIndex.cs
@@ -35,12 +35,13 @@ public class ScyllaBlobIndex : IBlobIndex
         ));
     }
 
-    public async Task AddBlobToIndex(NamespaceId ns, BlobIdentifier id)
+    public async Task AddBlobToIndex(NamespaceId ns, BlobIdentifier id, string? region = null)
     {
+        region ??= _jupiterSettings.CurrentValue.CurrentSite;
         using IScope _ = Tracer.Instance.StartActive("scylla.insert_blob_index");
 
         await _mapper.UpdateAsync<ScyllaBlobIndexTable>("SET regions = regions + ? WHERE namespace = ? AND blob_id = ?",
-            new string[] { _jupiterSettings.CurrentValue.CurrentSite }, ns.ToString(), new ScyllaBlobIdentifier(id));
+            new string[] { region }, ns.ToString(), new ScyllaBlobIdentifier(id));
     }
 
     public async Task<IBlobIndex.BlobInfo?> GetBlobInfo(NamespaceId ns, BlobIdentifier id)
diff --git a/Engine/Source/Programs/Horde/HordeStorage/Horde.Storage/Implementation/PeerStatusService.cs b/Engine/Source/Programs/Horde/HordeStorage/Horde.Storage/Implementation/PeerStatusService.cs
--- a/Engine/Source/Programs/Horde/HordeStorage/Horde.Storage/Implementation/PeerStatusService.cs
+++ b/Engine/Source/Programs/Horde/HordeStorage/Horde.Storage/Implementation/PeerStatusService.cs
@@ -15,6 +15,12 @@ namespace Horde.Storage.Implementation
     {
         public class PeerStatus
         {
+            public PeerStatus(PeerSettings peerSettings)
+            {
+                Endpoints = peerSettings.Endpoints;
+            }
+
+            public List<PeerEndpoints> Endpoints { get; set; }
             public int Latency { get; set; }
             public bool Reachable { get; set; }
         }
@@ -76,7 +82,7 @@ namespace Horde.Storage.Implementation
                 if (string.Equals(peerSettings.Name, jupiterSettings.CurrentValue.CurrentSite, StringComparison.InvariantCultureIgnoreCase))
                     continue;
 
-                _peers[peerSettings.Name] = new IPeerStatusService.PeerStatus
+                _peers[peerSettings.Name] = new IPeerStatusService.PeerStatus(peerSettings)
                 {
                     Latency = int.MaxValue
                 };
diff --git a/Engine/Source/Programs/Horde/HordeStorage/Horde.Storage/Settings.cs b/Engine/Source/Programs/Horde/HordeStorage/Horde.Storage/Settings.cs
--- a/Engine/Source/Programs/Horde/HordeStorage/Horde.Storage/Settings.cs
+++ b/Engine/Source/Programs/Horde/HordeStorage/Horde.Storage/Settings.cs
@@ -137,6 +137,8 @@ namespace Horde.Storage
         public int? MaxSingleBlobSize { get; set; } = null; // disable blob partitioning
 
         public int LastAccessRollupFrequencySeconds { get; set; } = 900; // 15 minutes
+        public List<string> OnDemandReplicationNamespaces { get; set; } = new List<string> { "*" };
+        public bool EnableOnDemandReplication { get; set; } = false;
     }
 
     public class MongoSettings
@@ -258,8 +260,8 @@ namespace Horde.Storage
 
         public TimeSpan BlobCleanupPollFrequency { get; set; } = TimeSpan.FromMinutes(60);
         public TimeSpan RefCleanupPollFrequency { get; set; } = TimeSpan.FromMinutes(60);
-        public List<string> CleanNamespacesV1 { get; set; } = new List<string>();
-        public List<string> CleanNamespaces { get; set; }  = new List<string>();
+        public List<string> CleanNamespaces { get; set; } = new List<string>();
+        public List<string> CleanNamespacesLegacy { get; set; }  = new List<string>();
 
     }
 
diff --git a/Engine/Source/Programs/Horde/HordeStorage/Horde.Storage/Tests/Functional/Blob/BlobReplicationTests.cs b/Engine/Source/Programs/Horde/HordeStorage/Horde.Storage/Tests/Functional/Blob/BlobReplicationTests.cs
new file mode 100644
--- /dev/null
+++ b/Engine/Source/Programs/Horde/HordeStorage/Horde.Storage/Tests/Functional/Blob/BlobReplicationTests.cs
@@ -0,0 +1,209 @@
+// Copyright Epic Games, Inc. All Rights Reserved.
+
+using System.Collections.Generic;
+using System.Linq;
+using System.Net;
+using System.Net.Http;
+using System.Text;
+using System.Threading.Tasks;
+using EpicGames.Horde.Storage;
+using Horde.Storage.Implementation.Blob;
+using Jupiter.Implementation;
+using Microsoft.AspNetCore.Hosting;
+using Microsoft.AspNetCore.TestHost;
+using Microsoft.Extensions.Configuration;
+using Microsoft.Extensions.DependencyInjection;
+using Microsoft.VisualStudio.TestTools.UnitTesting;
+using Moq;
+using Moq.Contrib.HttpClient;
+using Serilog;
+using Logger = Serilog.Core.Logger;
+
+namespace Horde.Storage.FunctionalTests.Storage
+{
+    [TestClass]
+    public class BlobReplicationTests
+    {
+        protected readonly NamespaceId TestNamespaceName = new NamespaceId("blob-replication-ns");
+
+        
+        [TestMethod]
+        public async Task ReplicateBlobFromRegion()
+        {
+            // 2 regions with data 
+            // region A - will not have data
+            // region B - will have data
+
+            string contents = "This is a random string of content";
+            byte[] bytes = Encoding.ASCII.GetBytes(contents);
+            BlobIdentifier blobIdentifier = BlobIdentifier.FromBlob(bytes);
+            
+            Mock<HttpMessageHandler> handler = new Mock<HttpMessageHandler>();
+
+            // site b has the content and will serve it
+            handler.SetupRequest($"http://siteB.com/internal/api/v1/blobs/{TestNamespaceName}/{blobIdentifier}").ReturnsResponse(HttpStatusCode.OK,
+                message => { message.Content = new ReadOnlyMemoryContent(bytes);}
+            ).Verifiable();
+
+            IConfigurationRoot configuration = new ConfigurationBuilder()
+                // we are not reading the base appSettings here as we want exact control over what runs in the tests
+                .AddJsonFile("appsettings.Testing.json", true)
+                .AddInMemoryCollection(new[] { 
+                    new KeyValuePair<string, string>("Horde_Storage:BlobIndexImplementation", HordeStorageSettings.BlobIndexImplementations.Memory.ToString()), 
+                    new KeyValuePair<string, string>("Horde_Storage:EnableOnDemandReplication", true.ToString()),
+                })
+                .AddEnvironmentVariables()
+                .Build();
+
+            Logger logger = new LoggerConfiguration()
+                .ReadFrom.Configuration(configuration)
+                .CreateLogger();
+
+            TestServer server = new TestServer(new WebHostBuilder()
+                .UseConfiguration(configuration)
+                .UseEnvironment("Testing")
+                .UseSerilog(logger)
+                .ConfigureTestServices(collection =>
+                {
+                    collection.Configure<ClusterSettings>(settings =>
+                    {
+                        settings.Peers = new PeerSettings[]
+                        {
+                            new PeerSettings()
+                            {
+                                Name = "siteA",
+                                FullName = "siteA",
+                                Endpoints = new PeerEndpoints[]
+                                {
+                                    new PeerEndpoints()
+                                    {
+                                        Url = "http://siteA.com/internal",
+                                        IsInternal = true
+                                    },
+                                    new PeerEndpoints()
+                                    {
+                                        Url = "http://siteA.com/public"
+                                    },
+                                }.ToList()
+                            },
+                            new PeerSettings()
+                            {
+                                Name = "siteB",
+                                FullName = "siteB",
+                                Endpoints = new PeerEndpoints[]
+                                {
+                                    new PeerEndpoints()
+                                    {
+                                        Url = "http://siteB.com/internal",
+                                        IsInternal = true
+                                    },
+                                    new PeerEndpoints()
+                                    {
+                                        Url = "http://siteB.com/public"
+                                    },
+                                }.ToList()
+                            },
+
+                        }.ToList();
+                    });
+                    collection.AddSingleton<IHttpClientFactory>(handler.CreateClientFactory());
+                })
+                .UseStartup<HordeStorageStartup>()
+            );
+
+            IBlobIndex? blobIndex = server.Services.GetService<IBlobIndex>();
+            Assert.IsNotNull(blobIndex);
+
+            await blobIndex.AddBlobToIndex(TestNamespaceName, blobIdentifier, "siteB");
+
+            HttpClient httpClient = server.CreateClient();
+            HttpResponseMessage response = await httpClient!.GetAsync(requestUri: $"api/v1/blobs/{TestNamespaceName}/{blobIdentifier}");
+            response.EnsureSuccessStatusCode();
+
+            byte[] responseData = await response.Content.ReadAsByteArrayAsync();
+            CollectionAssert.AreEqual(bytes, responseData);
+
+            handler.Verify();
+        }
+
+                
+        [TestMethod]
+        public async Task ReplicateBlobNotPresent()
+        {
+            // verify calling the blob endpoint when the blob is missing still returns a 404
+
+            string contents = "This is a random string of content";
+            byte[] bytes = Encoding.ASCII.GetBytes(contents);
+            BlobIdentifier blobIdentifier = BlobIdentifier.FromBlob(bytes);
+            
+            Mock<HttpMessageHandler> handler = new Mock<HttpMessageHandler>();
+
+            IConfigurationRoot configuration = new ConfigurationBuilder()
+                // we are not reading the base appSettings here as we want exact control over what runs in the tests
+                .AddJsonFile("appsettings.Testing.json", true)
+                .AddInMemoryCollection(new[] { new KeyValuePair<string, string>("Horde_Storage:BlobIndexImplementation", HordeStorageSettings.BlobIndexImplementations.Memory.ToString()) })
+                .AddEnvironmentVariables()
+                .Build();
+
+            Logger logger = new LoggerConfiguration()
+                .ReadFrom.Configuration(configuration)
+                .CreateLogger();
+
+            TestServer server = new TestServer(new WebHostBuilder()
+                .UseConfiguration(configuration)
+                .UseEnvironment("Testing")
+                .UseSerilog(logger)
+                .ConfigureTestServices(collection =>
+                {
+                    collection.Configure<ClusterSettings>(settings =>
+                    {
+                        settings.Peers = new PeerSettings[]
+                        {
+                            new PeerSettings()
+                            {
+                                Name = "siteA",
+                                FullName = "siteA",
+                                Endpoints = new PeerEndpoints[]
+                                {
+                                    new PeerEndpoints()
+                                    {
+                                        Url = "http://siteA.com/internal",
+                                        IsInternal = true
+                                    },
+                                    new PeerEndpoints()
+                                    {
+                                        Url = "http://siteA.com/public"
+                                    },
+                                }.ToList()
+                            },
+                            new PeerSettings()
+                            {
+                                Name = "siteB",
+                                FullName = "siteB",
+                                Endpoints = new PeerEndpoints[]
+                                {
+                                    new PeerEndpoints()
+                                    {
+                                        Url = "http://siteB.com/internal",
+                                        IsInternal = true
+                                    },
+                                    new PeerEndpoints()
+                                    {
+                                        Url = "http://siteB.com/public"
+                                    },
+                                }.ToList()
+                            },
+
+                        }.ToList();
+                    });
+                    collection.AddSingleton<IHttpClientFactory>(handler.CreateClientFactory());
+                })
+                .UseStartup<HordeStorageStartup>()
+            );
+
+            HttpClient httpClient = server.CreateClient();
+            HttpResponseMessage response = await httpClient!.GetAsync(requestUri: $"api/v1/blobs/{TestNamespaceName}/{blobIdentifier}");
+            Assert.IsTrue(response.StatusCode == HttpStatusCode.NotFound);
+        }
+    }
+}
diff --git a/Engine/Source/Programs/Horde/HordeStorage/Horde.Storage/Tests/UnitTests/HierarchicalBlobStorageTest.cs b/Engine/Source/Programs/Horde/HordeStorage/Horde.Storage/Tests/UnitTests/HierarchicalBlobStorageTest.cs
--- a/Engine/Source/Programs/Horde/HordeStorage/Horde.Storage/Tests/UnitTests/HierarchicalBlobStorageTest.cs
+++ b/Engine/Source/Programs/Horde/HordeStorage/Horde.Storage/Tests/UnitTests/HierarchicalBlobStorageTest.cs
@@ -4,6 +4,7 @@ using System;
 using System.Collections.Generic;
 using System.IO;
 using System.Linq;
+using System.Net.Http;
 using System.Text;
 using System.Threading.Tasks;
 using Dasync.Collections;
@@ -50,7 +51,7 @@ namespace Horde.Storage.UnitTests
             Mock<IServiceProvider> serviceProviderMock = new Mock<IServiceProvider>();
             serviceProviderMock.Setup(x => x.GetService(typeof(MemoryCacheBlobStore))).Returns(new MemoryCacheBlobStore(Mock.Of<IOptionsMonitor<MemoryCacheBlobSettings>>(_ => _.CurrentValue == new MemoryCacheBlobSettings())));
             IOptionsMonitor<HordeStorageSettings> settingsMonitor = Mock.Of<IOptionsMonitor<HordeStorageSettings>>(_ => _.CurrentValue == new HordeStorageSettings());
-            _chained = new BlobService(serviceProviderMock.Object, settingsMonitor, Mock.Of<IBlobIndex>());
+            _chained = new BlobService(serviceProviderMock.Object, settingsMonitor, Mock.Of<IBlobIndex>(), Mock.Of<IPeerStatusService>(), Mock.Of<IHttpClientFactory>(), Mock.Of<IServiceCredentials>());
             _chained.BlobStore = new List<IBlobStore> { _first, _second, _third };
 
             await _first.PutObject(Ns, Encoding.ASCII.GetBytes("onlyFirstContent"), _onlyFirstId);
