Horde Storage - Added resource names to a bunch of datadog scopes to make it easier to follow what queries are being made.

diff --git a/Engine/Source/Programs/Horde/HordeStorage/Horde.Storage/Implementation/Blob/IBlobService.cs b/Engine/Source/Programs/Horde/HordeStorage/Horde.Storage/Implementation/Blob/IBlobService.cs
--- a/Engine/Source/Programs/Horde/HordeStorage/Horde.Storage/Implementation/Blob/IBlobService.cs
+++ b/Engine/Source/Programs/Horde/HordeStorage/Horde.Storage/Implementation/Blob/IBlobService.cs
@@ -116,8 +116,8 @@ public class BlobService : IBlobService
 
     public async Task<BlobIdentifier> PutObject(NamespaceId ns, IBufferedPayload payload, BlobIdentifier identifier)
     {
-        using IScope _ = Tracer.Instance.StartActive("put_blob");
-
+        using IScope scope = Tracer.Instance.StartActive("put_blob");
+        scope.Span.ResourceName = identifier.ToString();
         await using Stream hashStream = payload.GetStream();
         BlobIdentifier id = BlobIdentifier.FromContentHash(await VerifyContentMatchesHash(hashStream, identifier));
 
@@ -132,8 +132,8 @@ public class BlobService : IBlobService
 
     public async Task<BlobIdentifier> PutObject(NamespaceId ns, byte[] payload, BlobIdentifier identifier)
     {
-        using IScope _ = Tracer.Instance.StartActive("put_blob");
-
+        using IScope scope = Tracer.Instance.StartActive("put_blob");
+        scope.Span.ResourceName = identifier.ToString();
         await using Stream hashStream = new MemoryStream(payload);
         BlobIdentifier id = BlobIdentifier.FromContentHash(await VerifyContentMatchesHash(hashStream, identifier));
 
@@ -147,7 +147,8 @@ public class BlobService : IBlobService
 
     public async Task<BlobIdentifier> PutObjectKnownHash(NamespaceId ns, IBufferedPayload content, BlobIdentifier identifier)
     {
-        using IScope _ = Tracer.Instance.StartActive("put_blob");
+        using IScope scope = Tracer.Instance.StartActive("put_blob");
+        scope.Span.ResourceName = identifier.ToString();
         Task<BlobIdentifier> putObjectTask = PutObjectToStores(ns, content, identifier);
         Task addToBlobIndexTask = _blobIndex.AddBlobToIndex(ns, identifier);
 
@@ -191,6 +192,7 @@ public class BlobService : IBlobService
         foreach (IBlobStore store in _blobStores)
         {
             using IScope scope = Tracer.Instance.StartActive("HierarchicalStore.GetObject");
+            scope.Span.ResourceName = blob.ToString();
             scope.Span.SetTag("BlobStore", store.GetType().Name);
             scope.Span.SetTag("ObjectFound", false.ToString());
             try
@@ -309,6 +311,7 @@ public class BlobService : IBlobService
         if (useBlobIndex)
         {
             using IScope scope = Tracer.Instance.StartActive("HierarchicalStore.ObjectExists");
+            scope.Span.ResourceName = blob.ToString();
             scope.Span.SetTag("BlobStore", "BlobIndex");
             bool exists = await _blobIndex.BlobExistsInRegion(ns, blob);
             if (exists)
@@ -325,6 +328,7 @@ public class BlobService : IBlobService
             foreach (IBlobStore store in _blobStores)
             {
                 using IScope scope = Tracer.Instance.StartActive("HierarchicalStore.ObjectExists");
+                scope.Span.ResourceName = blob.ToString();
                 scope.Span.SetTag("BlobStore", store.GetType().Name);
                 if (await store.Exists(ns, blob))
                 {
@@ -349,6 +353,7 @@ public class BlobService : IBlobService
             try
             {
                 using IScope scope = Tracer.Instance.StartActive("HierarchicalStore.DeleteObject");
+                scope.Span.ResourceName = blob.ToString();
                 scope.Span.SetTag("BlobStore", store.GetType().Name);
                 await store.DeleteObject(ns, blob);
                 deletedAtLeastOnce = true;
diff --git a/Engine/Source/Programs/Horde/HordeStorage/Horde.Storage/Implementation/Blob/ScyllaBlobIndex.cs b/Engine/Source/Programs/Horde/HordeStorage/Horde.Storage/Implementation/Blob/ScyllaBlobIndex.cs
--- a/Engine/Source/Programs/Horde/HordeStorage/Horde.Storage/Implementation/Blob/ScyllaBlobIndex.cs
+++ b/Engine/Source/Programs/Horde/HordeStorage/Horde.Storage/Implementation/Blob/ScyllaBlobIndex.cs
@@ -38,7 +38,8 @@ public class ScyllaBlobIndex : IBlobIndex
     public async Task AddBlobToIndex(NamespaceId ns, BlobIdentifier id, string? region = null)
     {
         region ??= _jupiterSettings.CurrentValue.CurrentSite;
-        using IScope _ = Tracer.Instance.StartActive("scylla.insert_blob_index");
+        using IScope scope = Tracer.Instance.StartActive("scylla.insert_blob_index");
+        scope.Span.ResourceName = $"{ns}.{id}";
 
         await _mapper.UpdateAsync<ScyllaBlobIndexTable>("SET regions = regions + ? WHERE namespace = ? AND blob_id = ?",
             new string[] { region }, ns.ToString(), new ScyllaBlobIdentifier(id));
@@ -46,7 +47,8 @@ public class ScyllaBlobIndex : IBlobIndex
 
     public async Task<IBlobIndex.BlobInfo?> GetBlobInfo(NamespaceId ns, BlobIdentifier id)
     {
-        using IScope _ = Tracer.Instance.StartActive("scylla.fetch_blob_index");
+        using IScope scope = Tracer.Instance.StartActive("scylla.fetch_blob_index");
+        scope.Span.ResourceName = $"{ns}.{id}";
 
         ScyllaBlobIndexTable? blobIndex =
             await _mapper.FirstOrDefaultAsync<ScyllaBlobIndexTable>("WHERE namespace = ? AND blob_id = ?",
@@ -66,7 +68,9 @@ public class ScyllaBlobIndex : IBlobIndex
     public async Task<bool> RemoveBlobFromIndex(NamespaceId ns, BlobIdentifier id)
     {
         // TODO: Should this only remove the current region, and the actual row isnt removed until all regions have been removed? Seems overly complicated
-        using IScope _ = Tracer.Instance.StartActive("scylla.remove_from_blob_index");
+        using IScope scope = Tracer.Instance.StartActive("scylla.remove_from_blob_index");
+        scope.Span.ResourceName = $"{ns}.{id}";
+
         await _mapper.DeleteAsync<ScyllaBlobIndexTable>("WHERE namespace = ? AND blob_id = ?", ns.ToString(),
             new ScyllaBlobIdentifier(id));
 
diff --git a/Engine/Source/Programs/Horde/HordeStorage/Horde.Storage/Implementation/Objects/ScyllaReferencesStore.cs b/Engine/Source/Programs/Horde/HordeStorage/Horde.Storage/Implementation/Objects/ScyllaReferencesStore.cs
--- a/Engine/Source/Programs/Horde/HordeStorage/Horde.Storage/Implementation/Objects/ScyllaReferencesStore.cs
+++ b/Engine/Source/Programs/Horde/HordeStorage/Horde.Storage/Implementation/Objects/ScyllaReferencesStore.cs
@@ -66,7 +66,8 @@ namespace Horde.Storage.Implementation
 
         public async Task<ObjectRecord> Get(NamespaceId ns, BucketId bucket, IoHashKey name, IReferencesStore.FieldFlags flags)
         {
-            using IScope _ = Tracer.Instance.StartActive("scylla.get");
+            using IScope scope = Tracer.Instance.StartActive("scylla.get");
+            scope.Span.ResourceName = $"{ns}.{bucket}.{name}";
 
             ScyllaObject? o;
             bool includePayload = (flags & IReferencesStore.FieldFlags.IncludePayload) != 0;
@@ -94,7 +95,9 @@ namespace Horde.Storage.Implementation
 
         public async Task Put(NamespaceId ns, BucketId bucket, IoHashKey name, BlobIdentifier blobHash, byte[] blob, bool isFinalized)
         {
-            using IScope _ = Tracer.Instance.StartActive("scylla.put");
+            using IScope scope = Tracer.Instance.StartActive("scylla.put");
+            scope.Span.ResourceName = $"{ns}.{bucket}.{name}";
+
             if (blob.LongLength > _settings.CurrentValue.InlineBlobMaxSize)
             {
                 // do not inline large blobs
@@ -113,7 +116,8 @@ namespace Horde.Storage.Implementation
 
         public async Task Finalize(NamespaceId ns, BucketId bucket, IoHashKey name, BlobIdentifier blobIdentifier)
         {
-            using IScope _ = Tracer.Instance.StartActive("scylla.finalize");
+            using IScope scope = Tracer.Instance.StartActive("scylla.finalize");
+            scope.Span.ResourceName = $"{ns}.{bucket}.{name}";
 
             await _mapper.UpdateAsync<ScyllaObject>("SET is_finalized=true WHERE namespace=? AND bucket=? AND name=?", ns.ToString(), bucket.ToString(), name.ToString());
         }
@@ -256,7 +260,9 @@ namespace Horde.Storage.Implementation
 
         public async Task<bool> Delete(NamespaceId ns, BucketId bucket, IoHashKey key)
         {
-            using IScope _ = Tracer.Instance.StartActive("scylla.delete_record");
+            using IScope scope = Tracer.Instance.StartActive("scylla.delete_record");
+            scope.Span.ResourceName = $"{ns}.{bucket}.{key}";
+
             ObjectRecord record;
             try
             {
