Horde Storage - Fixed failing test

diff --git a/Engine/Source/Programs/Horde/HordeStorage/Horde.Storage/Implementation/GC/CacheCleanup.cs b/Engine/Source/Programs/Horde/HordeStorage/Horde.Storage/Implementation/GC/CacheCleanup.cs
--- a/Engine/Source/Programs/Horde/HordeStorage/Horde.Storage/Implementation/GC/CacheCleanup.cs
+++ b/Engine/Source/Programs/Horde/HordeStorage/Horde.Storage/Implementation/GC/CacheCleanup.cs
@@ -55,7 +55,7 @@ namespace Horde.Storage.Implementation
             }
             else
                 throw new NotImplementedException(
-                    $"Namespace {ns} not present in CleanNamespaces or CleanNamespacesV1 lists, unable to clean it as we do not know which method to use.");
+                    $"Namespace {ns} not present in CleanNamespaces or CleanNamespacesLegacy lists, unable to clean it as we do not know which method to use.");
            
         }
 
diff --git a/Engine/Source/Programs/Horde/HordeStorage/Horde.Storage/Tests/Functional/GC/GCRefTests.cs b/Engine/Source/Programs/Horde/HordeStorage/Horde.Storage/Tests/Functional/GC/GCRefTests.cs
--- a/Engine/Source/Programs/Horde/HordeStorage/Horde.Storage/Tests/Functional/GC/GCRefTests.cs
+++ b/Engine/Source/Programs/Horde/HordeStorage/Horde.Storage/Tests/Functional/GC/GCRefTests.cs
@@ -43,7 +43,7 @@ namespace Horde.Storage.FunctionalTests.GC
                             .AddEnvironmentVariables()
                             .AddInMemoryCollection(new List<KeyValuePair<string, string>>()
                             {
-                                new KeyValuePair<string, string>("GC:CleanNamespacesV1:0", TestNamespace.ToString()),
+                                new KeyValuePair<string, string>("GC:CleanNamespacesLegacy:0", TestNamespace.ToString()),
                             })
                             .Build();
 
