Horde Storage - Added endpoint for uploading content without knowings its hash (its returned back in the response). This hash should still be compared against the content to make sure the correct content was uploaded but allows for verification after upload rather then before.

diff --git a/Engine/Source/Programs/Horde/HordeStorage/Horde.Storage/Controllers/StorageController.cs b/Engine/Source/Programs/Horde/HordeStorage/Horde.Storage/Controllers/StorageController.cs
--- a/Engine/Source/Programs/Horde/HordeStorage/Horde.Storage/Controllers/StorageController.cs
+++ b/Engine/Source/Programs/Horde/HordeStorage/Horde.Storage/Controllers/StorageController.cs
@@ -189,6 +189,35 @@ namespace Horde.Storage.Controllers
             });
         }
 
+        [HttpPost("{ns}")]
+        [Authorize("Storage.write")]
+        [RequiredContentType(MediaTypeNames.Application.Octet)]
+        [DisableRequestSizeLimit]
+        public async Task<IActionResult> Post(
+            [Required] NamespaceId ns)
+        {
+            AuthorizationResult authorizationResult = await _authorizationService.AuthorizeAsync(User, ns, NamespaceAccessRequirement.Name);
+
+            if (!authorizationResult.Succeeded)
+            {
+                return Forbid();
+            }
+
+            _diagnosticContext.Set("Content-Length", Request.ContentLength ?? -1);
+            using IBufferedPayload payload = await _bufferedPayloadFactory.CreateFromRequest(Request);
+            
+            await using Stream stream = payload.GetStream();
+
+            BlobIdentifier id = await BlobIdentifier.FromStream(stream);
+            await _storage.PutObjectKnownHash(ns, payload, id);
+
+            return Ok(new
+            {
+                Identifier = id.ToString()
+            });
+        }
+
+
         [HttpDelete("{ns}/{id}")]
         [Authorize("Storage.delete")]
         public async Task<IActionResult> Delete(
diff --git a/Engine/Source/Programs/Horde/HordeStorage/Horde.Storage/Tests/Functional/StorageTests.cs b/Engine/Source/Programs/Horde/HordeStorage/Horde.Storage/Tests/Functional/StorageTests.cs
--- a/Engine/Source/Programs/Horde/HordeStorage/Horde.Storage/Tests/Functional/StorageTests.cs
+++ b/Engine/Source/Programs/Horde/HordeStorage/Horde.Storage/Tests/Functional/StorageTests.cs
@@ -420,6 +420,21 @@ namespace Horde.Storage.FunctionalTests.Storage
             Assert.AreEqual(contentHash, content.Identifier);
         }
 
+        [TestMethod]
+        public async Task PostSmallBlob()
+        {
+            byte[] payload = Encoding.ASCII.GetBytes("I am a small blob");
+            using ByteArrayContent requestContent = new ByteArrayContent(payload);
+            requestContent.Headers.ContentType = new MediaTypeHeaderValue("application/octet-stream");
+            requestContent.Headers.ContentLength = payload.Length;
+            BlobIdentifier contentHash = BlobIdentifier.FromBlob(payload);
+            HttpResponseMessage result = await _httpClient!.PostAsync(requestUri: $"api/v1/s/{TestNamespaceName}", requestContent);
+
+            result.EnsureSuccessStatusCode();
+            InsertResponse content = await result.Content.ReadAsAsync<InsertResponse>();
+            Assert.AreEqual(contentHash, content.Identifier);
+        }
+
         [TestMethod]
         public async Task DeleteBlob()
         {
