Horde Storage - Made sure we do not commit content to the blob index until we are sure it was succesfully written to the blob store. Might help fix some inconsistency we have seen between the state of S3 and the state of the blob index.
Also made sure that when the blob store consistency check runs, it does in fact update the blob index when it deletes a bad blob and not just removes it from the store (this check hasnt really found any issues for a while so this is likely not related but still good to fix).
Lastly added some error output when a blob that we expect to find (as we have done a exists check for it) is in fact not present when using the inline type.

diff --git a/Engine/Source/Programs/Horde/HordeStorage/Horde.Storage/Controllers/ReferencesController.cs b/Engine/Source/Programs/Horde/HordeStorage/Horde.Storage/Controllers/ReferencesController.cs
--- a/Engine/Source/Programs/Horde/HordeStorage/Horde.Storage/Controllers/ReferencesController.cs
+++ b/Engine/Source/Programs/Horde/HordeStorage/Horde.Storage/Controllers/ReferencesController.cs
@@ -217,10 +217,18 @@ namespace Horde.Storage.Controllers
                             if (referencedBlobs.Count == 1)
                             {
                                 BlobIdentifier attachmentToSend = referencedBlobs.First();
-                                BlobContents referencedBlobContents = await _blobStore.GetObject(ns, attachmentToSend);
-                                Response.Headers[CommonHeaders.InlinePayloadHash] = attachmentToSend.ToString();
+                                try
+                                {
+                                    BlobContents referencedBlobContents = await _blobStore.GetObject(ns, attachmentToSend);
+                                    Response.Headers[CommonHeaders.InlinePayloadHash] = attachmentToSend.ToString();
 
-                                await WriteBody(referencedBlobContents, CustomMediaTypeNames.JupiterInlinedPayload);
+                                    await WriteBody(referencedBlobContents, CustomMediaTypeNames.JupiterInlinedPayload);
+                                }
+                                catch (BlobNotFoundException)
+                                {
+                                    _logger.Error("Failed to find blob {Blob} in namespace {Namespace} but it was expected to exist after reference resolving", attachmentToSend, ns);
+                                    throw;
+                                }
                                 return new EmptyResult();
                             }
                             else if (referencedBlobs.Count == 0)
diff --git a/Engine/Source/Programs/Horde/HordeStorage/Horde.Storage/Implementation/Blob/CachedBlobIndex.cs b/Engine/Source/Programs/Horde/HordeStorage/Horde.Storage/Implementation/Blob/CachedBlobIndex.cs
--- a/Engine/Source/Programs/Horde/HordeStorage/Horde.Storage/Implementation/Blob/CachedBlobIndex.cs
+++ b/Engine/Source/Programs/Horde/HordeStorage/Horde.Storage/Implementation/Blob/CachedBlobIndex.cs
@@ -28,6 +28,7 @@ public class CachedBlobIndex : IBlobIndex
             return new IBlobIndex.BlobInfo
             {
                 Namespace = ns,
+                BlobIdentifier = id,
                 References = new List<(BucketId, IoHashKey)>(),
                 Regions = new HashSet<string> {_jupiterSettings.CurrentValue.CurrentSite}
             };
@@ -49,6 +50,12 @@ public class CachedBlobIndex : IBlobIndex
         return false;
     }
 
+    public async Task RemoveBlobFromRegion(NamespaceId ns, BlobIdentifier id, string? region = null)
+    {
+        // We do not actually track any blob information when running in cached mode
+        await Task.CompletedTask;
+    }
+
     public async Task<bool> BlobExistsInRegion(NamespaceId ns, BlobIdentifier blobIdentifier)
     {
         return await _fileSystemStore.Exists(ns, blobIdentifier);
diff --git a/Engine/Source/Programs/Horde/HordeStorage/Horde.Storage/Implementation/Blob/ConsistencyCheckService.cs b/Engine/Source/Programs/Horde/HordeStorage/Horde.Storage/Implementation/Blob/ConsistencyCheckService.cs
--- a/Engine/Source/Programs/Horde/HordeStorage/Horde.Storage/Implementation/Blob/ConsistencyCheckService.cs
+++ b/Engine/Source/Programs/Horde/HordeStorage/Horde.Storage/Implementation/Blob/ConsistencyCheckService.cs
@@ -9,6 +9,7 @@ using Cassandra;
 using Dasync.Collections;
 using Datadog.Trace;
 using EpicGames.Horde.Storage;
+using Horde.Storage.Implementation.Blob;
 using Jupiter;
 using Jupiter.Implementation;
 using Microsoft.Extensions.DependencyInjection;
@@ -25,6 +26,7 @@ namespace Horde.Storage.Implementation
         private readonly ILeaderElection _leaderElection;
         private readonly IRefsStore _refsStore;
         private readonly IReferencesStore _referencesStore;
+        private readonly IBlobIndex _blobIndex;
         private readonly ILogger _logger = Log.ForContext<ConsistencyCheckService>();
 
         public class ConsistencyState
@@ -36,7 +38,7 @@ namespace Horde.Storage.Implementation
             return _settings.CurrentValue.Enabled;
         }
 
-        public ConsistencyCheckService(IOptionsMonitor<ConsistencyCheckSettings> settings, IServiceProvider provider, ILeaderElection leaderElection, IRefsStore refsStore, IReferencesStore referencesStore) :
+        public ConsistencyCheckService(IOptionsMonitor<ConsistencyCheckSettings> settings, IServiceProvider provider, ILeaderElection leaderElection, IRefsStore refsStore, IReferencesStore referencesStore, IBlobIndex blobIndex) :
             base(serviceName: nameof(ConsistencyCheckService), TimeSpan.FromSeconds(settings.CurrentValue.ConsistencyCheckPollFrequencySeconds), new ConsistencyState())
         {
             _settings = settings;
@@ -44,6 +46,7 @@ namespace Horde.Storage.Implementation
             _leaderElection = leaderElection;
             _refsStore = refsStore;
             _referencesStore = referencesStore;
+            _blobIndex = blobIndex;
         }
 
         public override async Task<bool> OnPoll(ConsistencyState state, CancellationToken cancellationToken)
@@ -111,6 +114,8 @@ namespace Horde.Storage.Implementation
                         _logger.Error("Mismatching hash for S3 object {Blob} in {Namespace}, new hash has {NewHash}. Deleting incorrect blob.", blob, ns, newHash);
 
                         await s3Store.DeleteObject(ns, blob);
+                        // update blob index tracking to indicate that we no longer have this blob in this region
+                        await _blobIndex.RemoveBlobFromRegion(ns, blob);
                     }
                 }
                 
diff --git a/Engine/Source/Programs/Horde/HordeStorage/Horde.Storage/Implementation/Blob/IBlobIndex.cs b/Engine/Source/Programs/Horde/HordeStorage/Horde.Storage/Implementation/Blob/IBlobIndex.cs
--- a/Engine/Source/Programs/Horde/HordeStorage/Horde.Storage/Implementation/Blob/IBlobIndex.cs
+++ b/Engine/Source/Programs/Horde/HordeStorage/Horde.Storage/Implementation/Blob/IBlobIndex.cs
@@ -13,6 +13,7 @@ namespace Horde.Storage.Implementation.Blob
         {
             public HashSet<string> Regions { get; init; } = new HashSet<string>();
             public NamespaceId Namespace { get; init; }
+            public BlobIdentifier BlobIdentifier { get; init; } = null!;
             public List<(BucketId, IoHashKey)> References { get; init; } = new List<(BucketId, IoHashKey)>();
         }
 
@@ -21,6 +22,8 @@ namespace Horde.Storage.Implementation.Blob
         Task<BlobInfo?> GetBlobInfo(NamespaceId ns, BlobIdentifier id);
 
         Task<bool> RemoveBlobFromIndex(NamespaceId ns, BlobIdentifier id);
+        Task RemoveBlobFromRegion(NamespaceId ns, BlobIdentifier id, string? region = null);
+
         Task<bool> BlobExistsInRegion(NamespaceId ns, BlobIdentifier blobIdentifier);
         Task AddRefToBlobs(NamespaceId ns, BucketId bucket, IoHashKey key, BlobIdentifier[] blobs);
     }
diff --git a/Engine/Source/Programs/Horde/HordeStorage/Horde.Storage/Implementation/Blob/IBlobService.cs b/Engine/Source/Programs/Horde/HordeStorage/Horde.Storage/Implementation/Blob/IBlobService.cs
--- a/Engine/Source/Programs/Horde/HordeStorage/Horde.Storage/Implementation/Blob/IBlobService.cs
+++ b/Engine/Source/Programs/Horde/HordeStorage/Horde.Storage/Implementation/Blob/IBlobService.cs
@@ -121,12 +121,10 @@ public class BlobService : IBlobService
         await using Stream hashStream = payload.GetStream();
         BlobIdentifier id = BlobIdentifier.FromContentHash(await VerifyContentMatchesHash(hashStream, identifier));
 
-        Task<BlobIdentifier> putObjectTask = PutObjectToStores(ns, payload, id);
-        Task addToBlobIndexTask = _blobIndex.AddBlobToIndex(ns, id);
+        BlobIdentifier objectStoreIdentifier = await PutObjectToStores(ns, payload, id);
+        await _blobIndex.AddBlobToIndex(ns, id);
 
-        await Task.WhenAll(putObjectTask, addToBlobIndexTask);
-
-        return await putObjectTask;
+        return objectStoreIdentifier;
 
     }
 
@@ -137,24 +135,20 @@ public class BlobService : IBlobService
         await using Stream hashStream = new MemoryStream(payload);
         BlobIdentifier id = BlobIdentifier.FromContentHash(await VerifyContentMatchesHash(hashStream, identifier));
 
-        Task<BlobIdentifier> putObjectTask = PutObjectToStores(ns, payload, id);
-        Task addToBlobIndexTask = _blobIndex.AddBlobToIndex(ns, id);
-
-        await Task.WhenAll(putObjectTask, addToBlobIndexTask);
+        BlobIdentifier objectStoreIdentifier = await PutObjectToStores(ns, payload, id);
+        await _blobIndex.AddBlobToIndex(ns, id);
 
-        return await putObjectTask;
+        return objectStoreIdentifier;
     }
 
     public async Task<BlobIdentifier> PutObjectKnownHash(NamespaceId ns, IBufferedPayload content, BlobIdentifier identifier)
     {
         using IScope scope = Tracer.Instance.StartActive("put_blob");
         scope.Span.ResourceName = identifier.ToString();
-        Task<BlobIdentifier> putObjectTask = PutObjectToStores(ns, content, identifier);
-        Task addToBlobIndexTask = _blobIndex.AddBlobToIndex(ns, identifier);
-
-        await Task.WhenAll(putObjectTask, addToBlobIndexTask);
+        BlobIdentifier objectStoreIdentifier = await PutObjectToStores(ns, content, identifier);
+        await _blobIndex.AddBlobToIndex(ns, identifier);
 
-        return await putObjectTask;
+        return objectStoreIdentifier;
     }
 
     private async Task<BlobIdentifier> PutObjectToStores(NamespaceId ns, IBufferedPayload bufferedPayload, BlobIdentifier identifier)
diff --git a/Engine/Source/Programs/Horde/HordeStorage/Horde.Storage/Implementation/Blob/MemoryBlobIndex.cs b/Engine/Source/Programs/Horde/HordeStorage/Horde.Storage/Implementation/Blob/MemoryBlobIndex.cs
--- a/Engine/Source/Programs/Horde/HordeStorage/Horde.Storage/Implementation/Blob/MemoryBlobIndex.cs
+++ b/Engine/Source/Programs/Horde/HordeStorage/Horde.Storage/Implementation/Blob/MemoryBlobIndex.cs
@@ -35,7 +35,7 @@ public class MemoryBlobIndex : IBlobIndex
     {
         region ??= _jupiterSettings.CurrentValue.CurrentSite;
         ConcurrentDictionary<BlobIdentifier, MemoryBlobInfo> index = GetNamespaceContainer(ns);
-        index[id] = NewBlobInfo(ns, region);
+        index[id] = NewBlobInfo(ns, id, region);
         return Task.CompletedTask;
     }
 
@@ -56,6 +56,23 @@ public class MemoryBlobIndex : IBlobIndex
         return Task.FromResult(index.TryRemove(id, out MemoryBlobInfo? _));
     }
 
+    public Task RemoveBlobFromRegion(NamespaceId ns, BlobIdentifier id, string? region = null)
+    {
+        region ??= _jupiterSettings.CurrentValue.CurrentSite;
+        ConcurrentDictionary<BlobIdentifier, MemoryBlobInfo> index = GetNamespaceContainer(ns);
+
+        index.AddOrUpdate(id, _ =>
+        {
+            MemoryBlobInfo info = NewBlobInfo(ns, id, region);
+            return info;
+        }, (_, info) =>
+        {
+            info.Regions.Remove(region);
+            return info;
+        });
+        return Task.CompletedTask;
+    }
+
     public async Task<bool> BlobExistsInRegion(NamespaceId ns, BlobIdentifier blobIdentifier)
     {
         IBlobIndex.BlobInfo? blobInfo = await GetBlobInfo(ns, blobIdentifier);
@@ -70,7 +87,7 @@ public class MemoryBlobIndex : IBlobIndex
 
             index.AddOrUpdate(id, _ =>
             {
-                MemoryBlobInfo info = NewBlobInfo(ns, _jupiterSettings.CurrentValue.CurrentSite);
+                MemoryBlobInfo info = NewBlobInfo(ns, id, _jupiterSettings.CurrentValue.CurrentSite);
                 info.References.Add((bucket, key));
                 return info;
             }, (_, info) =>
@@ -83,12 +100,13 @@ public class MemoryBlobIndex : IBlobIndex
         return Task.CompletedTask;
     }
 
-    private MemoryBlobInfo NewBlobInfo(NamespaceId ns, string region)
+    private MemoryBlobInfo NewBlobInfo(NamespaceId ns, BlobIdentifier blob, string region)
     {
         MemoryBlobInfo info = new MemoryBlobInfo
         {
             Regions = new HashSet<string> { region },
             Namespace = ns,
+            BlobIdentifier = blob,
         };
         return info;
     }
diff --git a/Engine/Source/Programs/Horde/HordeStorage/Horde.Storage/Implementation/Blob/MongoBlobIndex.cs b/Engine/Source/Programs/Horde/HordeStorage/Horde.Storage/Implementation/Blob/MongoBlobIndex.cs
--- a/Engine/Source/Programs/Horde/HordeStorage/Horde.Storage/Implementation/Blob/MongoBlobIndex.cs
+++ b/Engine/Source/Programs/Horde/HordeStorage/Horde.Storage/Implementation/Blob/MongoBlobIndex.cs
@@ -73,6 +73,22 @@ public class MongoBlobIndex : MongoStore, IBlobIndex
         return result.IsAcknowledged;
     }
 
+    public async Task RemoveBlobFromRegion(NamespaceId ns, BlobIdentifier id, string? region = null)
+    {
+        region ??= _jupiterSettings.CurrentValue.CurrentSite;
+
+        IBlobIndex.BlobInfo? blobInfo = await GetBlobInfo(ns, id);
+        IMongoCollection<MongoBlobIndexModelV0> collection = GetCollection<MongoBlobIndexModelV0>();
+        MongoBlobIndexModelV0 model = new MongoBlobIndexModelV0(blobInfo!);
+        model.Regions.Remove(region);
+        FilterDefinition<MongoBlobIndexModelV0> filter = Builders<MongoBlobIndexModelV0>.Filter.Where(m => m.Ns == ns.ToString() && m.BlobId == id.ToString());
+        await collection.FindOneAndReplaceAsync(filter, model, new FindOneAndReplaceOptions<MongoBlobIndexModelV0, MongoBlobIndexModelV0>
+        {
+            IsUpsert = true
+        });
+
+    }
+
     public async Task<bool> BlobExistsInRegion(NamespaceId ns, BlobIdentifier blobIdentifier)
     {
         IBlobIndex.BlobInfo? blobInfo = await GetBlobInfo(ns, blobIdentifier);
@@ -116,6 +132,14 @@ class MongoBlobIndexModelV0
         References = references;
     }
 
+    public MongoBlobIndexModelV0(IBlobIndex.BlobInfo blobInfo)
+    {
+        Ns = blobInfo.Namespace.ToString();
+        BlobId = blobInfo.BlobIdentifier.ToString();
+        Regions = blobInfo.Regions.ToList();
+        References = blobInfo.References.Select(pair => new Dictionary<string, string>{ {"bucket", pair.Item1.ToString()}, {"key", pair.Item2.ToString()}}).ToList();
+    }
+
     public MongoBlobIndexModelV0(NamespaceId ns, BlobIdentifier blobId)
     {
         Ns = ns.ToString();
@@ -138,6 +162,7 @@ class MongoBlobIndexModelV0
         return new IBlobIndex.BlobInfo()
         {
             Namespace = new NamespaceId(Ns),
+            BlobIdentifier = new BlobIdentifier(BlobId),
             Regions = Regions.ToHashSet(),
             References = References.Select(dictionary =>
                 (new BucketId(dictionary["bucket"]), new IoHashKey(dictionary["key"]))).ToList(),
diff --git a/Engine/Source/Programs/Horde/HordeStorage/Horde.Storage/Implementation/Blob/ScyllaBlobIndex.cs b/Engine/Source/Programs/Horde/HordeStorage/Horde.Storage/Implementation/Blob/ScyllaBlobIndex.cs
--- a/Engine/Source/Programs/Horde/HordeStorage/Horde.Storage/Implementation/Blob/ScyllaBlobIndex.cs
+++ b/Engine/Source/Programs/Horde/HordeStorage/Horde.Storage/Implementation/Blob/ScyllaBlobIndex.cs
@@ -60,6 +60,7 @@ public class ScyllaBlobIndex : IBlobIndex
         return new IBlobIndex.BlobInfo
         {
             Regions = blobIndex.Regions ?? new HashSet<string>(),
+            BlobIdentifier = id,
             Namespace = new NamespaceId(blobIndex.Namespace),
             References = blobIndex.References != null ? blobIndex.References.Select(reference => reference.AsTuple()).ToList() : new List<(BucketId, IoHashKey)>()
         };
@@ -67,7 +68,6 @@ public class ScyllaBlobIndex : IBlobIndex
 
     public async Task<bool> RemoveBlobFromIndex(NamespaceId ns, BlobIdentifier id)
     {
-        // TODO: Should this only remove the current region, and the actual row isnt removed until all regions have been removed? Seems overly complicated
         using IScope scope = Tracer.Instance.StartActive("scylla.remove_from_blob_index");
         scope.Span.ResourceName = $"{ns}.{id}";
 
@@ -77,6 +77,15 @@ public class ScyllaBlobIndex : IBlobIndex
         return true;
     }
 
+    public async Task RemoveBlobFromRegion(NamespaceId ns, BlobIdentifier id, string? region = null)
+    {
+        region ??= _jupiterSettings.CurrentValue.CurrentSite;
+        using IScope scope = Tracer.Instance.StartActive("scylla.remove_blob_index_region");
+        scope.Span.ResourceName = $"{ns}.{id}";
+
+        await _mapper.UpdateAsync<ScyllaBlobIndexTable>("SET regions = regions - ? WHERE namespace = ? AND blob_id = ?", new string[] { region }, ns.ToString(), new ScyllaBlobIdentifier(id));
+    }
+
     public async Task<bool> BlobExistsInRegion(NamespaceId ns, BlobIdentifier blobIdentifier)
     {
         IBlobIndex.BlobInfo? blobInfo = await GetBlobInfo(ns, blobIdentifier);
