Horde Storage - Make sure we do not throw exceptions into the timer update code and instead log them so our structured logging sees it.

diff --git a/Engine/Source/Programs/Horde/HordeStorage/Horde.Storage/Implementation/LastAccess/LastAccessServiceReferences.cs b/Engine/Source/Programs/Horde/HordeStorage/Horde.Storage/Implementation/LastAccess/LastAccessServiceReferences.cs
--- a/Engine/Source/Programs/Horde/HordeStorage/Horde.Storage/Implementation/LastAccess/LastAccessServiceReferences.cs
+++ b/Engine/Source/Programs/Horde/HordeStorage/Horde.Storage/Implementation/LastAccess/LastAccessServiceReferences.cs
@@ -53,10 +53,17 @@ namespace Horde.Storage.Implementation
 
         private void OnUpdate(object? state)
         {
-            // call results to make sure we join the task
-            Task.WaitAll(
-                ProcessLastAccessRecords()
-            );
+            try
+            {
+                // call results to make sure we join the task
+                Task.WaitAll(
+                    ProcessLastAccessRecords()
+                );
+            }
+            catch (Exception e)
+            {
+                _logger.Error(e, "Exception thrown while processing last access records");
+            }
         }
 
         internal async Task<List<(LastAccessRecord, DateTime)>> ProcessLastAccessRecords()
