Horde: Add per-stream settings for replicating commit data.

- Remove ref ids from commit objects, so we can keep multiple refs for each depending on configured settings.
- Replication is still disabled by default for streams.

diff --git a/Engine/Source/Programs/Horde/Horde.Build/Commands/Perforce/ReplicateCommand.cs b/Engine/Source/Programs/Horde/Horde.Build/Commands/Perforce/ReplicateCommand.cs
--- a/Engine/Source/Programs/Horde/Horde.Build/Commands/Perforce/ReplicateCommand.cs
+++ b/Engine/Source/Programs/Horde/Horde.Build/Commands/Perforce/ReplicateCommand.cs
@@ -32,7 +32,7 @@ namespace HordeServer.Commands
 		public int Change { get; set; }
 
 		[CommandLine]
-		public int BaseChange { get; set; }
+		public int? BaseChange { get; set; }
 
 		[CommandLine]
 		public int Count { get; set; } = 1;
@@ -49,6 +49,9 @@ namespace HordeServer.Commands
 		[CommandLine]
 		public string Filter { get; set; } = "...";
 
+		[CommandLine]
+		public bool Metadata { get; set; } = false;
+
 		[CommandLine]
 		public DirectoryReference? OutputDir { get; set; }
 
@@ -87,12 +90,6 @@ namespace HordeServer.Commands
 			Dictionary<IStream, int> StreamToFirstChange = new Dictionary<IStream, int>();
 			StreamToFirstChange[Stream] = Change;
 
-			CommitTree? BaseTree = null;
-			if (BaseChange != 0)
-			{
-				BaseTree = await GetBaseTreeAsync(CommitCollection, Stream.Id, Change);
-			}
-
 			await foreach (NewCommit NewCommit in CommitService.FindCommitsForClusterAsync(Stream.ClusterName, StreamToFirstChange).Take(Count))
 			{
 				string BriefSummary = NewCommit.Description.Replace('\n', ' ');
@@ -101,36 +98,15 @@ namespace HordeServer.Commands
 
 				if (Content)
 				{
-					if (BaseTree == null && !Clean)
-					{
-						BaseTree = await GetBaseTreeAsync(CommitCollection, Stream.Id, NewCommit.Change);
-					}
-					else
+					if (Clean || BaseChange != null)
 					{
-						BaseTree = await CommitService.FindCommitTreeAsync(Stream, NewCommit.Change, BaseTree, Filter);
+						await CommitService.WriteCommitTreeAsync(Stream, NewCommit.Change, BaseChange, Filter, Metadata);
 					}
-
-					NewCommit.TreeRefId = BaseTree.RefId;
+					BaseChange = NewCommit.Change;
 				}
-
-				await CommitCollection.AddOrReplaceAsync(NewCommit);
 			}
 
 			return 0;
 		}
-
-		static async Task<CommitTree> GetBaseTreeAsync(ICommitCollection CommitCollection, StreamId StreamId, int Change)
-		{
-			ICommit? Commit = await CommitCollection.GetCommitAsync(StreamId, Change);
-			if (Commit == null)
-			{
-				throw new InvalidOperationException($"Unable to find existing commit for {StreamId} @ CL {Change}");
-			}
-			if (Commit.TreeRefId == null)
-			{
-				throw new InvalidOperationException($"No tree for commit {StreamId} @ CL {Change}");
-			}
-			return new CommitTree(Commit.StreamId, Commit.Change, Commit.TreeRefId.Value);
-		}
 	}
 }
diff --git a/Engine/Source/Programs/Horde/Horde.Build/Perforce/ICommit.cs b/Engine/Source/Programs/Horde/Horde.Build/Perforce/ICommit.cs
--- a/Engine/Source/Programs/Horde/Horde.Build/Perforce/ICommit.cs
+++ b/Engine/Source/Programs/Horde/Horde.Build/Perforce/ICommit.cs
@@ -67,11 +67,6 @@ namespace HordeServer.Commits
 		/// Date/time that change was committed
 		/// </summary>
 		public DateTime DateUtc { get; }
-
-		/// <summary>
-		/// Name of a reference describing the tree at this changelist
-		/// </summary>
-		public RefId? TreeRefId { get; }
 	}
 
 	/// <summary>
@@ -103,16 +98,12 @@ namespace HordeServer.Commits
 		/// <inheritdoc cref="ICommit.DateUtc"/>
 		public DateTime DateUtc { get; set; }
 
-		/// <inheritdoc cref="ICommit.TreeRefId"/>
-		public RefId? TreeRefId { get; set; }
-
 		/// <summary>
 		/// Constructor
 		/// </summary>
 		public NewCommit(ICommit Commit)
 			: this(Commit.StreamId, Commit.Change, Commit.OriginalChange, Commit.AuthorId, Commit.OwnerId, Commit.Description, Commit.BasePath, Commit.DateUtc)
 		{
-			this.TreeRefId = Commit.TreeRefId;
 		}
 
 		/// <summary>
diff --git a/Engine/Source/Programs/Horde/Horde.Build/Perforce/Impl/CommitCollection.cs b/Engine/Source/Programs/Horde/Horde.Build/Perforce/Impl/CommitCollection.cs
--- a/Engine/Source/Programs/Horde/Horde.Build/Perforce/Impl/CommitCollection.cs
+++ b/Engine/Source/Programs/Horde/Horde.Build/Perforce/Impl/CommitCollection.cs
@@ -50,9 +50,6 @@ namespace HordeServer.Commits.Impl
 			[BsonElement("p")]
 			public string BasePath { get; set; } = String.Empty;
 
-			[BsonElement("rid")]
-			public RefId? TreeRefId { get; set; }
-
 			[BsonElement("t")]
 			public DateTime DateUtc { get; set; }
 
@@ -78,7 +75,6 @@ namespace HordeServer.Commits.Impl
 				}
 				this.Description = NewCommit.Description;
 				this.BasePath = NewCommit.BasePath;
-				this.TreeRefId = NewCommit.TreeRefId;
 				this.DateUtc = NewCommit.DateUtc;
 			}
 		}
diff --git a/Engine/Source/Programs/Horde/Horde.Build/Perforce/Impl/CommitService.cs b/Engine/Source/Programs/Horde/Horde.Build/Perforce/Impl/CommitService.cs
--- a/Engine/Source/Programs/Horde/Horde.Build/Perforce/Impl/CommitService.cs
+++ b/Engine/Source/Programs/Horde/Horde.Build/Perforce/Impl/CommitService.cs
@@ -56,24 +56,9 @@ namespace HordeServer.Commits.Impl
 		public bool Metadata { get; set; } = true;
 
 		/// <summary>
-		/// Specific list of streams to enable commit mirroring for
+		/// Whether to replicate content to Horde storage. Must also be enabled on a per-stream basis (see <see cref="IStream.ReplicationMode"/>).
 		/// </summary>
-		public List<StreamId> MetadataStreams { get; set; } = new List<StreamId>();
-
-		/// <summary>
-		/// Whether to mirror content to Horde storage
-		/// </summary>
-		public bool Content { get; set; } = false;
-
-		/// <summary>
-		/// List of streams for which content mirroring should be enabled
-		/// </summary>
-		public List<StreamId> ContentStreams { get; set; } = new List<StreamId>();
-
-		/// <summary>
-		/// Write Perforce stub files rather than including the actual content.
-		/// </summary>
-		public bool WriteStubFiles { get; set; } = false;
+		public bool Content { get; set; } = true;
 
 		/// <summary>
 		/// Namespace to store content replicated from Perforce
@@ -92,54 +77,60 @@ namespace HordeServer.Commits.Impl
 	}
 
 	/// <summary>
-	/// Information about a commit tree
+	/// Key for a commit object. This is hashed to produce a ref id.
 	/// </summary>
-	class CommitTree
+	public class CommitKey
 	{
 		/// <summary>
-		/// The stream id
+		/// Stream that this commit came from
 		/// </summary>
-		public StreamId StreamId { get; }
+		[CbField]
+		public StreamId StreamId { get; set; }
 
 		/// <summary>
-		/// The changelist number
+		/// The change being mirrored
 		/// </summary>
-		public int Change { get; }
+		[CbField]
+		public int Change { get; set; }
 
 		/// <summary>
-		/// The ref id for storing this tree
+		/// Filter for paths in the depot included in this tree
 		/// </summary>
-		public RefId RefId { get; }
+		public string? Filter { get; set; }
 
 		/// <summary>
-		/// Constructor
+		/// Whether this commit contains depot path and revision metadata rather than full file contents
+		/// </summary>
+		[CbField]
+		public bool RevisionsOnly { get; set; }
+
+		/// <summary>
+		/// Default constructor
 		/// </summary>
-		public CommitTree(StreamId StreamId, int Change, RefId RefId)
+		public CommitKey()
 		{
-			this.StreamId = StreamId;
-			this.Change = Change;
-			this.RefId = RefId;
 		}
 
 		/// <summary>
-		/// Attempt to create a commit tree from a commit
+		/// Constructor
 		/// </summary>
-		public static CommitTree? FromCommit(ICommit Commit)
+		public CommitKey(StreamId StreamId, int Change, string? Filter, bool RevisionsOnly)
 		{
-			if (Commit.TreeRefId == null)
-			{
-				return null;
-			}
-			else
-			{
-				return new CommitTree(Commit.StreamId, Commit.Change, Commit.TreeRefId.Value);
-			}
+			this.StreamId = StreamId;
+			this.Change = Change;
+			this.Filter = Filter;
+			this.RevisionsOnly = RevisionsOnly;
 		}
 
 		/// <summary>
-		/// The ref id for storing this tree
+		/// Gets the ref id for this key
 		/// </summary>
-		public static RefId GetDefaultRefId(int Change) => new RefId(IoHash.Compute(Encoding.UTF8.GetBytes($"{Change}")));
+		/// <returns></returns>
+		public RefId GetRefId()
+		{
+			CbObject Object = CbSerializer.Serialize(this);
+			return new RefId(IoHash.Compute(Object.GetView().Span));
+		}
 	}
 
 	/// <summary>
@@ -328,11 +319,6 @@ namespace HordeServer.Commits.Impl
 		async ValueTask UpdateMetadataAsync(CancellationToken CancellationToken)
 		{
 			List<IStream> Streams = await StreamCollection.FindAllAsync();
-			if (Options.MetadataStreams.Count > 0)
-			{
-				Streams.RemoveAll(x => !Options.MetadataStreams.Contains(x.Id));
-			}
-
 			foreach (IGrouping<string, IStream> Group in Streams.GroupBy(x => x.ClusterName, StringComparer.OrdinalIgnoreCase))
 			{
 				try
@@ -810,7 +796,7 @@ namespace HordeServer.Commits.Impl
 				return;
 			}
 
-			ICommit? PrevCommit = null;
+			int? PrevChange = null;
 			for (; ; )
 			{
 				// Get the first two changes to be mirrored. The first one should already exist, unless it's the start of replication for this stream.
@@ -821,20 +807,21 @@ namespace HordeServer.Commits.Impl
 				}
 
 				// Check that the previous commit is still valid
-				if (PrevCommit != null && PrevCommit.Change != Values[0])
+				if (PrevChange != null && PrevChange.Value != Values[0])
 				{
-					Logger.LogInformation("Invalidating previous commit; expected {Change}, actually {ActualChange}", PrevCommit.Change, Values[0]);
-					PrevCommit = null;
+					Logger.LogInformation("Invalidating previous commit; expected {Change}, actually {ActualChange}", PrevChange.Value, Values[0]);
+					PrevChange = null;
 				}
 
 				// If we don't have a previous commit tree yet, perform a full snapshot of that change
-				if (PrevCommit == null)
+				if (PrevChange == null)
 				{
-					PrevCommit = await UpdateCommitTreeAsync(Stream, Values[0], null);
+					await WriteCommitTreeAsync(Stream, Values[0], null);
 				}
 
 				// Perform a snapshot of the new change, then remove it from the list
-				PrevCommit = await UpdateCommitTreeAsync(Stream, Values[1], PrevCommit);
+				await WriteCommitTreeAsync(Stream, Values[1], PrevChange);
+				PrevChange = Values[1];
 				await Changes.LeftPopAsync();
 			}
 		}
@@ -844,56 +831,20 @@ namespace HordeServer.Commits.Impl
 		/// </summary>
 		/// <param name="Stream">The stream to replicate</param>
 		/// <param name="Change">Commit to store the tree ref</param>
-		/// <param name="PrevCommit">Previous commit to use as a base</param>
-		/// <returns>New ref id for the given commit</returns>
-		public async Task<ICommit?> UpdateCommitTreeAsync(IStream Stream, int Change, ICommit? PrevCommit)
-		{
-			// Get the commit for the requested change
-			ICommit? Commit = await CommitCollection.GetCommitAsync(Stream.Id, Change);
-			if (Commit == null)
-			{
-				Logger.LogWarning("Missing commit for {Stream} change {Change}", Stream.Id, Change);
-				return Commit;
-			}
-			else if (Commit.TreeRefId != null)
-			{
-				Logger.LogInformation("Skipping replication for {Stream} change {Change}; tree already exists", Stream.Id, Change);
-				return Commit;
-			}
-
-			// Get the ref corresponding to BaseChange
-			CommitTree? BaseTree = null;
-			if (PrevCommit == null)
-			{
-				Logger.LogInformation("No base commit for mirroring {StreamId} change {Change}", Stream.Id, Commit.Change);
-			}
-			else if (PrevCommit.TreeRefId == null)
-			{
-				Logger.LogWarning("Base commit for stream {StreamId} change {Change} does not have a mirrored tree", Stream.Id, PrevCommit.Change);
-			}
-			else
-			{
-				BaseTree = new CommitTree(PrevCommit.StreamId, PrevCommit.Change, PrevCommit.TreeRefId.Value);
-			}
-
-			// Write the new tree
-			CommitTree Tree = await FindCommitTreeAsync(Stream, Commit.Change, BaseTree);
-
-			// Update the commit
-			NewCommit NewCommit = new NewCommit(Commit);
-			NewCommit.TreeRefId = Tree.RefId;
-			return await CommitCollection.AddOrReplaceAsync(NewCommit);
-		}
+		/// <param name="BaseChange">The base change to update from</param>
+		/// <returns>Root tree object</returns>
+		public Task WriteCommitTreeAsync(IStream Stream, int Change, int? BaseChange) => WriteCommitTreeAsync(Stream, Change, BaseChange, Stream.ReplicationFilter, Stream.ReplicationMode == ContentReplicationMode.RevisionsOnly);
 
 		/// <summary>
 		/// Replicates the contents of a stream to Horde storage, optionally using the given change as a starting point
 		/// </summary>
 		/// <param name="Stream">The stream to replicate</param>
 		/// <param name="Change">Commit to store the tree ref</param>
-		/// <param name="BaseTree">The initial change to update from</param>
-		/// <param name="QueryPath">Depot path to query for changes. Will default to the entire depot (but filtered by the workspace)</param>
+		/// <param name="BaseChange">The base change to update from</param>
+		/// <param name="Filter">Depot path to query for changes. Will default to the entire depot (but filtered by the workspace)</param>
+		/// <param name="RevisionsOnly">Whether to replicate file revisions only</param>
 		/// <returns>Root tree object</returns>
-		public async Task<CommitTree> FindCommitTreeAsync(IStream Stream, int Change, CommitTree? BaseTree, string QueryPath = "//...")
+		public async Task WriteCommitTreeAsync(IStream Stream, int Change, int? BaseChange, string? Filter, bool RevisionsOnly)
 		{
 			// Create a client to replicate from this stream, and connect to the server
 			ReplicationClient ClientInfo = await FindOrAddReplicationClientAsync(Stream);
@@ -901,15 +852,15 @@ namespace HordeServer.Commits.Impl
 
 			// Get the initial directory state
 			Bundle<DirectoryNode> CommitBundle;
-			if (BaseTree == null)
+			if (BaseChange == null)
 			{
 				await FlushWorkspaceAsync(ClientInfo, Perforce, 0);
 				CommitBundle = Bundle.Create<DirectoryNode>(StorageClient, Options.NamespaceId, Options.Bundle, MemoryCache);
 			}
 			else
 			{
-				await FlushWorkspaceAsync(ClientInfo, Perforce, BaseTree.Change);
-				CommitBundle = await Bundle.ReadAsync<DirectoryNode>(StorageClient, Options.NamespaceId, GetStreamBucketId(BaseTree.StreamId), BaseTree.RefId, Options.Bundle, MemoryCache);
+				await FlushWorkspaceAsync(ClientInfo, Perforce, BaseChange.Value);
+				CommitBundle = await Bundle.ReadAsync<DirectoryNode>(StorageClient, Options.NamespaceId, GetStreamBucketId(Stream.Id), new CommitKey(Stream.Id, BaseChange.Value, Filter, RevisionsOnly).GetRefId(), Options.Bundle, MemoryCache);
 			}
 
 			// Apply all the updates
@@ -917,9 +868,10 @@ namespace HordeServer.Commits.Impl
 			ClientInfo.Change = -1;
 
 			Utf8String ClientRoot = new Utf8String(ClientInfo.Client.Root);
+			string QueryPath = $"//{ClientInfo.Client.Name}/{Filter}";
 
 			DirectoryNode Contents = CommitBundle.Root;
-			if (Options.WriteStubFiles)
+			if (RevisionsOnly)
 			{
 				int Count = 0;
 				await foreach (PerforceResponse<SyncRecord> Record in Perforce.StreamCommandAsync<SyncRecord>("sync", new[] { "-k", $"{QueryPath}@{Change}" }, null, default))
@@ -997,16 +949,18 @@ namespace HordeServer.Commits.Impl
 
 			ClientInfo.Change = Change;
 
+			// Create the key for this commit
+			CommitKey Key = new CommitKey(Stream.Id, Change, Filter, RevisionsOnly);
+			CbObject KeyObject = CbSerializer.Serialize(Key);
+
 			// Return the new root object
-			RefId RefId = CommitTree.GetDefaultRefId(Change);
+			RefId RefId = new RefId(IoHash.Compute(KeyObject.GetView().Span));
 			Logger.LogInformation("Writing ref {RefId} for {StreamId} change {Change}", RefId, Stream.Id, Change);
-			await CommitBundle.WriteAsync(GetStreamBucketId(Stream.Id), RefId, true);
+			await CommitBundle.WriteAsync(GetStreamBucketId(Stream.Id), RefId, KeyObject, true);
 
 			// Print out stats for the new data
 			BundleStats Stats = CommitBundle.Stats;
 			Logger.LogInformation("Written ref {RefId} for {StreamId} change {Change}. Download: {NewBytes:n0} bytes ({NewRefCount} refs => {NewRefBytes:n0} bytes, {NewBlobCount} blobs => {NewBlobBytes:n0} bytes, {NewNodeCount:n0} nodes => {NewNodeBytes:n0} bytes).", RefId, Stream.Id, Change, Stats.NewRefBytes + Stats.NewBlobBytes, Stats.NewRefCount, Stats.NewRefBytes, Stats.NewBlobCount, Stats.NewBlobBytes, Stats.NewNodeCount, Stats.NewNodeBytes);
-
-			return new CommitTree(Stream.Id, Change, RefId);
 		}
 
 		async Task FlushWorkspaceAsync(ReplicationClient ClientInfo, IPerforceConnection Perforce, int Change)
diff --git a/Engine/Source/Programs/Horde/Horde.Build/Projects/IStream.cs b/Engine/Source/Programs/Horde/Horde.Build/Projects/IStream.cs
--- a/Engine/Source/Programs/Horde/Horde.Build/Projects/IStream.cs
+++ b/Engine/Source/Programs/Horde/Horde.Build/Projects/IStream.cs
@@ -483,6 +483,27 @@ namespace HordeServer.Models
 		}
 	}
 
+	/// <summary>
+	/// How to replicate data for this stream
+	/// </summary>
+	public enum ContentReplicationMode
+	{
+		/// <summary>
+		/// No content will be replicated for this stream
+		/// </summary>
+		None,
+
+		/// <summary>
+		/// Only replicate depot path and revision data for each file
+		/// </summary>
+		RevisionsOnly,
+
+		/// <summary>
+		/// Replicate full stream contents to storage
+		/// </summary>
+		Full,
+	}
+
 	/// <summary>
 	/// Information about a stream
 	/// </summary>
@@ -573,6 +594,16 @@ namespace HordeServer.Models
 		/// </summary>
 		public string? PauseComment { get; }
 
+		/// <summary>
+		/// Whether to replicate data from Perforce to Horde Storage.
+		/// </summary>
+		public ContentReplicationMode ReplicationMode { get; }
+
+		/// <summary>
+		/// Filter for paths to be replicated. May be removed in future.
+		/// </summary>
+		public string? ReplicationFilter { get; }
+
 		/// <summary>
 		/// The ACL for this object
 		/// </summary>
diff --git a/Engine/Source/Programs/Horde/Horde.Build/Projects/Impl/StreamCollection.cs b/Engine/Source/Programs/Horde/Horde.Build/Projects/Impl/StreamCollection.cs
--- a/Engine/Source/Programs/Horde/Horde.Build/Projects/Impl/StreamCollection.cs
+++ b/Engine/Source/Programs/Horde/Horde.Build/Projects/Impl/StreamCollection.cs
@@ -58,6 +58,14 @@ namespace HordeServer.Collections.Impl
 			public Dictionary<TemplateRefId, TemplateRef> Templates { get; set; } = new Dictionary<TemplateRefId, TemplateRef>();
 			public DateTime? PausedUntil { get; set; }
 			public string? PauseComment { get; set; }
+
+			[BsonIgnoreIfDefault]
+			[BsonDefaultValue(ContentReplicationMode.None)]
+			public ContentReplicationMode ReplicationMode { get; set; }
+
+			[BsonIgnoreIfNull]
+			public string? ReplicationFilter { get; set; }
+
 			public Acl? Acl { get; set; }
 			public int UpdateIndex { get; set; }
 			public bool Deleted { get; set; }
@@ -235,6 +243,8 @@ namespace HordeServer.Collections.Impl
 			NewStream.AgentTypes = AgentTypes;
 			NewStream.WorkspaceTypes = WorkspaceTypes;
 			NewStream.Templates = TemplateRefs;
+			NewStream.ReplicationMode = Config.ReplicationMode;
+			NewStream.ReplicationFilter = Config.ReplicationFilter;
 			NewStream.Acl = Acl;
 
 			try
@@ -279,6 +289,8 @@ namespace HordeServer.Collections.Impl
 			Updates.Add(UpdateBuilder.Set(x => x.AgentTypes, AgentTypes ?? new Dictionary<string, AgentType>()));
 			Updates.Add(UpdateBuilder.Set(x => x.WorkspaceTypes, WorkspaceTypes ?? new Dictionary<string, WorkspaceType>()));
 			Updates.Add(UpdateBuilder.Set(x => x.Templates, TemplateRefs ?? new Dictionary<TemplateRefId, TemplateRef>()));
+			Updates.Add(UpdateBuilder.Set(x => x.ReplicationMode, Config.ReplicationMode));
+			Updates.Add(UpdateBuilder.SetOrUnsetNullRef(x => x.ReplicationFilter, Config.ReplicationFilter));
 			Updates.Add(UpdateBuilder.SetOrUnsetNullRef(x => x.Acl, Acl));
 			Updates.Add(UpdateBuilder.Unset(x => x.Deleted));
 
diff --git a/Engine/Source/Programs/Horde/Horde.Build/Server/Config.cs b/Engine/Source/Programs/Horde/Horde.Build/Server/Config.cs
--- a/Engine/Source/Programs/Horde/Horde.Build/Server/Config.cs
+++ b/Engine/Source/Programs/Horde/Horde.Build/Server/Config.cs
@@ -249,6 +249,16 @@ namespace HordeServer.Models
 		/// Reason for pausing builds of the stream
 		/// </summary>
 		public string? PauseComment { get; set; }
+
+		/// <summary>
+		/// How to replicate data from VCS to Horde Storage.
+		/// </summary>
+		public ContentReplicationMode ReplicationMode { get; set; }
+
+		/// <summary>
+		/// Filter for paths to be replicated to storage, as a Perforce wildcard relative to the root of the workspace.
+		/// </summary>
+		public string? ReplicationFilter { get; set; }
 	}
 
 	/// <summary>
