Horde: Remove upgrade logic for old job documents. This is causing invalid user documents to be created.

diff --git a/Engine/Source/Programs/Horde/Horde.Build/Jobs/Impl/JobCollection.cs b/Engine/Source/Programs/Horde/Horde.Build/Jobs/Impl/JobCollection.cs
--- a/Engine/Source/Programs/Horde/Horde.Build/Jobs/Impl/JobCollection.cs
+++ b/Engine/Source/Programs/Horde/Horde.Build/Jobs/Impl/JobCollection.cs
@@ -401,64 +401,13 @@ namespace HordeServer.Collections.Impl
 			Indexes = new DatabaseIndexes(Jobs, DatabaseService.ReadOnlyMode);
 		}
 
-		async Task PostLoadAsync(JobDocument Job)
+		static Task PostLoadAsync(JobDocument Job)
 		{
 			if (Job.GraphHash == ContentHash.Empty)
 			{
 				Job.Batches.Clear();
 			}
-
-			List<UpdateDefinition<JobDocument>> Updates = new List<UpdateDefinition<JobDocument>>();
-			if (Job.StartedByUser_DEPRECATED != null)
-			{
-				IUser StartedByUser = await UserCollection.FindOrAddUserByLoginAsync(Job.StartedByUser_DEPRECATED);
-				Job.StartedByUserId = StartedByUser.Id;
-				Job.StartedByUser_DEPRECATED = null;
-				Updates.Add(Builders<JobDocument>.Update.Set(x => x.StartedByUserId, StartedByUser.Id).Unset(x => x.StartedByUser_DEPRECATED));
-			}
-
-			if (Job.AbortedByUser_DEPRECATED != null)
-			{
-				IUser AbortedByUser = await UserCollection.FindOrAddUserByLoginAsync(Job.AbortedByUser_DEPRECATED);
-				Job.AbortedByUserId = AbortedByUser.Id;
-				Job.AbortedByUser_DEPRECATED = null;
-				Updates.Add(Builders<JobDocument>.Update.Set(x => x.AbortedByUserId, AbortedByUser.Id).Unset(x => x.AbortedByUser_DEPRECATED));
-			}
-
-
-			for (int BatchIdx = 0; BatchIdx < Job.Batches.Count; BatchIdx++)
-			{
-				int LocalBatchIdx = BatchIdx;
-				JobStepBatchDocument Batch = Job.Batches[BatchIdx];
-
-				for (int StepIdx = 0; StepIdx < Batch.Steps.Count; StepIdx++)
-				{
-					int LocalStepIdx = StepIdx;
-					JobStepDocument Step = Batch.Steps[StepIdx];
-
-					if (Step.AbortedByUser_DEPRECATED != null)
-					{
-						IUser AbortedByUser = await UserCollection.FindOrAddUserByLoginAsync(Step.AbortedByUser_DEPRECATED);
-						Step.AbortedByUserId = AbortedByUser.Id;
-						Step.AbortedByUser_DEPRECATED = null;
-						Updates.Add(Builders<JobDocument>.Update.Set(x => x.Batches[LocalBatchIdx].Steps[LocalStepIdx].AbortedByUserId, AbortedByUser.Id).Unset(x => x.Batches[LocalBatchIdx].Steps[LocalStepIdx].AbortedByUser_DEPRECATED));
-					}
-
-					if (Step.RetriedByUser_DEPRECATED != null)
-					{
-						IUser RetriedByUser = await UserCollection.FindOrAddUserByLoginAsync(Step.RetriedByUser_DEPRECATED);
-						Step.RetriedByUserId = RetriedByUser.Id;
-						Step.RetriedByUser_DEPRECATED = null;
-						Updates.Add(Builders<JobDocument>.Update.Set(x => x.Batches[LocalBatchIdx].Steps[LocalStepIdx].RetriedByUserId, RetriedByUser.Id).Unset(x => x.Batches[LocalBatchIdx].Steps[LocalStepIdx].RetriedByUser_DEPRECATED));
-					}
-				}
-			}
-
-
-			if (Updates.Count > 0)
-			{
-				await TryUpdateAsync(Job, Builders<JobDocument>.Update.Combine(Updates));
-			}
+			return Task.CompletedTask;
 		}
 
 		static JobDocument Clone(JobDocument Job)
