Horde: Missing file.

diff --git a/Engine/Source/Programs/Horde/Horde.Build/Utilities/MongoExtensions.cs b/Engine/Source/Programs/Horde/Horde.Build/Utilities/MongoExtensions.cs
--- a/Engine/Source/Programs/Horde/Horde.Build/Utilities/MongoExtensions.cs
+++ b/Engine/Source/Programs/Horde/Horde.Build/Utilities/MongoExtensions.cs
@@ -197,6 +197,26 @@ namespace HordeServer.Utilities
 			}
 		}
 
+		/// <summary>
+		/// Sets a field to a value, or unsets it if the value is null
+		/// </summary>
+		/// <typeparam name="TDocument">The document type</typeparam>
+		/// <param name="UpdateBuilder">Update builder</param>
+		/// <param name="Field">Expression for the field to set</param>
+		/// <param name="Value">New value to set</param>
+		/// <returns>Update defintiion</returns>
+		public static UpdateDefinition<TDocument> SetOrUnsetBool<TDocument>(this UpdateDefinitionBuilder<TDocument> UpdateBuilder, Expression<Func<TDocument, bool>> Field, bool Value)
+		{
+			if (Value)
+			{
+				return UpdateBuilder.Set(Field, Value);
+			}
+			else
+			{
+				return UpdateBuilder.Unset(new ExpressionFieldDefinition<TDocument>(Field));
+			}
+		}
+
 		/// <summary>
 		/// Sets a field to a value, or unsets it if the value is null
 		/// </summary>
