Horde: Bump config service version so it picks up the new issue promotion template flags
#skipci
#rnx

diff --git a/Engine/Source/Programs/Horde/Horde.Build/Server/ConfigService.cs b/Engine/Source/Programs/Horde/Horde.Build/Server/ConfigService.cs
--- a/Engine/Source/Programs/Horde/Horde.Build/Server/ConfigService.cs
+++ b/Engine/Source/Programs/Horde/Horde.Build/Server/ConfigService.cs
@@ -47,7 +47,7 @@ namespace HordeServer.Services
 		/// <summary>
 		/// Config file version number
 		/// </summary>
-		const int Version = 9;
+		const int Version = 10;
 
 		DatabaseService DatabaseService;
 		ProjectService ProjectService;
