Horde: Fix infinite loop in scheduling of builds initiated by gates, which whould eventually time out and result in no jobs being started.

- PerforceService was not returning files as part of change details, due to specifying -S argument when querying changelist. This would cause any code/content change limitations to fail.
- Maximum change to query for gated steps was not being updated while scanning for triggers.
- Added additional context to error message when timeouts occur so we can more easily track down similar errors in future.
- Added test to check correct behavior.

diff --git a/Engine/Source/Programs/Horde/Horde.Build.Tests/SchedulerTests.cs b/Engine/Source/Programs/Horde/Horde.Build.Tests/SchedulerTests.cs
--- a/Engine/Source/Programs/Horde/Horde.Build.Tests/SchedulerTests.cs
+++ b/Engine/Source/Programs/Horde/Horde.Build.Tests/SchedulerTests.cs
@@ -470,6 +470,83 @@ namespace HordeServerTests
 			Assert.AreEqual(1233, Jobs4[0].CodeChange);
 		}
 
+		[TestMethod]
+		public async Task GateTest2Async()
+		{
+			IUser Bob = await UserCollection.FindOrAddUserByLoginAsync("Bob");
+
+			DateTime StartTime = new DateTime(2021, 1, 1, 12, 0, 0, DateTimeKind.Local); // Friday Jan 1, 2021 
+			Clock.UtcNow = StartTime;
+
+			PerforceService.Changes.Clear();
+			PerforceService.AddChange("//UE5/Main", 1230, Bob, "", new[] { "code.cpp" });
+			PerforceService.AddChange("//UE5/Main", 1231, Bob, "", new[] { "content.uasset" });
+			PerforceService.AddChange("//UE5/Main", 1232, Bob, "", new[] { "content.uasset" });
+			PerforceService.AddChange("//UE5/Main", 1233, Bob, "", new[] { "code.cpp" });
+
+			// Create two templates, the second dependent on the first
+			TemplateRefId NewTemplateRefId1 = new TemplateRefId("new-template-1");
+			CreateTemplateRefRequest NewTemplate1 = new CreateTemplateRefRequest();
+			NewTemplate1.Id = NewTemplateRefId1.ToString();
+
+			TemplateRefId NewTemplateRefId2 = new TemplateRefId("new-template-2");
+			CreateTemplateRefRequest NewTemplate2 = new CreateTemplateRefRequest();
+			NewTemplate2.Id = NewTemplateRefId2.ToString();
+			NewTemplate2.Name = "Test template 2";
+			NewTemplate2.Schedule = new CreateScheduleRequest();
+			NewTemplate2.Schedule.MaxChanges = 4;
+			NewTemplate2.Schedule.Filter = new List<ChangeContentFlags> { ChangeContentFlags.ContainsCode };
+			NewTemplate2.Schedule.Gate = new CreateScheduleGateRequest { TemplateId = NewTemplateRefId1.ToString(), Target = "TriggerNext" };
+			NewTemplate2.Schedule.Patterns.Add(new CreateSchedulePatternRequest { Interval = 10 });// (null, 0, null, 10));
+//			NewTemplate2.Schedule.LastTriggerTime = StartTime;
+
+			IStream? Stream = await StreamService.GetStreamAsync(StreamId);
+
+			StreamConfig Config = new StreamConfig();
+			Config.Name = "//UE5/Main";
+			Config.Tabs.Add(new CreateJobsTabRequest { Title = "foo", Templates = new List<string> { NewTemplateRefId1.ToString(), NewTemplateRefId2.ToString() } });
+			Config.Templates.Add(NewTemplate1);
+			Config.Templates.Add(NewTemplate2);
+
+			Stream = (await StreamService.StreamCollection.TryCreateOrReplaceAsync(StreamId, Stream, "", "", ProjectId, Config))!;
+
+			ITemplate Template1 = (await TemplateCollection.GetAsync(Stream.Templates[NewTemplateRefId1].Hash))!;
+			ITemplate Template2 = (await TemplateCollection.GetAsync(Stream.Templates[NewTemplateRefId2].Hash))!;
+
+			// Create the graph
+			IGraph GraphA = await GraphCollection.AddAsync(Template1);
+			NewGroup GroupA = new NewGroup("win", new List<NewNode> { new NewNode("TriggerNext") });
+			GraphA = await GraphCollection.AppendAsync(GraphA, new List<NewGroup> { GroupA });
+
+			// Create successful jobs for all the changes we added above
+			for (int Change = 1230; Change <= 1233; Change++)
+			{
+				int CodeChange = (Change < 1233) ? 1230 : 1233;
+
+				IJob Job1 = await JobService.CreateJobAsync(null, Stream, NewTemplateRefId1, Template.Id, GraphA, "Hello", Change, CodeChange, null, null, null, null, null, false, null, null, true, true, null, null, new List<string> { "-Target=TriggerNext" });
+				for (int BatchIdx = 0; BatchIdx < Job1.Batches.Count; BatchIdx++)
+				{
+					SubResourceId BatchId1 = Job1.Batches[BatchIdx].Id;
+					Job1 = Deref(await JobService.UpdateBatchAsync(Job1, BatchId1, LogId.GenerateNewId(), JobStepBatchState.Running));
+					for (int StepIdx = 0; StepIdx < Job1.Batches[BatchIdx].Steps.Count; StepIdx++)
+					{
+						SubResourceId StepId1 = Job1.Batches[BatchIdx].Steps[StepIdx].Id;
+						Job1 = Deref(await JobService.UpdateStepAsync(Job1, BatchId1, StepId1, JobStepState.Completed, JobStepOutcome.Success, NewLogId: LogId.GenerateNewId()));
+					}
+					Job1 = Deref(await JobService.UpdateBatchAsync(Job1, BatchId1, LogId.GenerateNewId(), JobStepBatchState.Complete));
+				}
+			}
+			await GetNewJobs();
+
+			// Tick the schedule and make sure it doesn't trigger
+			await Clock.AdvanceAsync(TimeSpan.FromMinutes(30.0));
+			await ScheduleService.TriggerAsync(StreamId, NewTemplateRefId2, Clock.UtcNow, default);
+			List<IJob> Jobs3 = await GetNewJobs();
+			Assert.AreEqual(2, Jobs3.Count);
+			Assert.AreEqual(1230, Jobs3[0].Change);
+			Assert.AreEqual(1233, Jobs3[1].Change);
+		}
+
 		[TestMethod]
 		public async Task UpdateConfigAsync()
 		{
diff --git a/Engine/Source/Programs/Horde/Horde.Build/Perforce/Impl/PerforceService.cs b/Engine/Source/Programs/Horde/Horde.Build/Perforce/Impl/PerforceService.cs
--- a/Engine/Source/Programs/Horde/Horde.Build/Perforce/Impl/PerforceService.cs
+++ b/Engine/Source/Programs/Horde/Horde.Build/Perforce/Impl/PerforceService.cs
@@ -766,7 +766,7 @@ namespace HordeServer.Services
 			
 			int Revision = GetSyncRevision(MetaData.DepotPath.Path, MetaData.HeadAction, MetaData.HeadRev);
 			Md5Hash? Digest = String.IsNullOrEmpty(MetaData.Digest) ? (Md5Hash?)null : Md5Hash.Parse(MetaData.Digest);
-			return new ChangeFile(RelativePath, MetaData.DepotPath.Path, Revision, MetaData.FileSize, Digest, MetaData.HeadType.ToString());
+			return new ChangeFile(RelativePath, MetaData.DepotPath.Path, Revision, MetaData.FileSize, Digest, (MetaData.Type ?? MetaData.HeadType).ToString());
 		}
 
 		/// <summary>
@@ -863,7 +863,7 @@ namespace HordeServer.Services
 			{
 				List<ChangeDetails> Results = new List<ChangeDetails>();
 
-				List<string> Args = new List<string> { "-s", "-S", $"{ChangeNumber}" };
+				List<string> Args = new List<string> { "-s", $"{ChangeNumber}" };
 
 				using (P4.P4Command Command = new P4.P4Command(Repository, "describe", true, Args.ToArray()))
 				{
@@ -891,7 +891,7 @@ namespace HordeServer.Services
 						List<ChangeFile> Files = new List<ChangeFile>();
 
 						P4.Changelist Change = new P4.Changelist();
-						Change.FromChangeCmdTaggedOutput(TaggedObject, true, Offset, DSTMismatch);
+						Change.FromChangeCmdTaggedOutput(TaggedObject, false, Offset, DSTMismatch);
 
 						foreach (P4.FileMetaData DescribeFile in Change.Files)
 						{
diff --git a/Engine/Source/Programs/Horde/Horde.Build/Projects/ScheduleService.cs b/Engine/Source/Programs/Horde/Horde.Build/Projects/ScheduleService.cs
--- a/Engine/Source/Programs/Horde/Horde.Build/Projects/ScheduleService.cs
+++ b/Engine/Source/Programs/Horde/Horde.Build/Projects/ScheduleService.cs
@@ -320,7 +320,7 @@ namespace HordeServer.Services
 		/// <param name="UtcNow"></param>
 		/// <param name="CancellationToken"></param>
 		/// <returns>Async task</returns>
-		private async Task<bool> TriggerAsync(StreamId StreamId, TemplateRefId TemplateId, DateTime UtcNow, CancellationToken CancellationToken)
+		internal async Task<bool> TriggerAsync(StreamId StreamId, TemplateRefId TemplateId, DateTime UtcNow, CancellationToken CancellationToken)
 		{
 			IStream? Stream = await StreamService.GetStreamAsync(StreamId);
 			if (Stream == null || !Stream.Templates.TryGetValue(TemplateId, out TemplateRef? TemplateRef))
@@ -469,7 +469,7 @@ namespace HordeServer.Services
 
 				// Adjust the changelist for the desired filter
 				int Change = ChangeDetails.Number;
-				if (ShouldBuildChangeAsync(ChangeDetails, FilterFlags, FileFilter))
+				if (ShouldBuildChange(ChangeDetails, FilterFlags, FileFilter))
 				{
 					int CodeChange = await History.GetCodeChange(Change);
 					if (CodeChange == -1)
@@ -483,7 +483,14 @@ namespace HordeServer.Services
 				// Check we haven't exceeded the time limit
 				if (Timer.Elapsed > TimeSpan.FromMinutes(2.0))
 				{
-					Logger.LogError("Querying for changes to trigger has taken {Time}. Aborting.", Timer.Elapsed);
+					Logger.LogError("Querying for changes to trigger for {StreamId} template {TemplateId} has taken {Time}. Aborting.", Stream.Id, TemplateId, Timer.Elapsed);
+					break;
+				}
+				
+				// Update the remaining range of changes to check for
+				MaxChangeNumber = Change - 1;
+				if(MaxChangeNumber < MinChangeNumber)
+				{
 					break;
 				}
 			}
@@ -533,7 +540,7 @@ namespace HordeServer.Services
 		/// <param name="FilterFlags"></param>
 		/// <param name="FileFilter">Filter for the files to trigger a build</param>
 		/// <returns></returns>
-		private bool ShouldBuildChangeAsync(ChangeDetails Details, ChangeContentFlags? FilterFlags, FileFilter? FileFilter)
+		private bool ShouldBuildChange(ChangeDetails Details, ChangeContentFlags? FilterFlags, FileFilter? FileFilter)
 		{
 			if (Regex.IsMatch(Details.Description, @"^\s*#\s*skipci", RegexOptions.Multiline))
 			{
