Move per namespace settings into the namespace policy
Introduced a setting for controlling if we do last access tracking for namespaces

This gives us better overview of what settings are enabled per namespace (and easier toggling for legacy vs new behavior)

diff --git a/Engine/Source/Programs/Horde/HordeStorage/Horde.Storage/Implementation/Blob/AmazonS3Store.cs b/Engine/Source/Programs/Horde/HordeStorage/Horde.Storage/Implementation/Blob/AmazonS3Store.cs
--- a/Engine/Source/Programs/Horde/HordeStorage/Horde.Storage/Implementation/Blob/AmazonS3Store.cs
+++ b/Engine/Source/Programs/Horde/HordeStorage/Horde.Storage/Implementation/Blob/AmazonS3Store.cs
@@ -165,7 +165,8 @@ namespace Horde.Storage.Implementation
 
         public async Task<bool> Exists(NamespaceId ns, BlobIdentifier blobIdentifier)
         {
-            if (_settings.UseBlobIndexForExistsCheck && _settings.NamespacesThatUseBlobIndexForExistsCheck.Contains(ns.ToString()))
+            NamespaceSettings.PerNamespaceSettings policies = _namespacePolicyResolver.GetPoliciesForNs(ns);
+            if (_settings.UseBlobIndexForExistsCheck && policies.UseBlobIndexForSlowExists)
             {
                 return await _blobIndex.BlobExistsInRegion(ns, blobIdentifier);
             }
diff --git a/Engine/Source/Programs/Horde/HordeStorage/Horde.Storage/Implementation/Blob/IBlobService.cs b/Engine/Source/Programs/Horde/HordeStorage/Horde.Storage/Implementation/Blob/IBlobService.cs
--- a/Engine/Source/Programs/Horde/HordeStorage/Horde.Storage/Implementation/Blob/IBlobService.cs
+++ b/Engine/Source/Programs/Horde/HordeStorage/Horde.Storage/Implementation/Blob/IBlobService.cs
@@ -13,6 +13,7 @@ using Datadog.Trace;
 using EpicGames.Horde.Storage;
 using Horde.Storage.Controllers;
 using Horde.Storage.Implementation.Blob;
+using Jupiter.Common;
 using Jupiter.Common.Implementation;
 using Jupiter.Implementation;
 using Microsoft.Extensions.DependencyInjection;
@@ -57,6 +58,7 @@ public class BlobService : IBlobService
     private readonly IPeerStatusService _peerStatusService;
     private readonly IHttpClientFactory _httpClientFactory;
     private readonly IServiceCredentials _serviceCredentials;
+    private readonly INamespacePolicyResolver _namespacePolicyResolver;
     private readonly ILogger _logger = Log.ForContext<BlobService>();
 
     internal IEnumerable<IBlobStore> BlobStore
@@ -65,7 +67,7 @@ public class BlobService : IBlobService
         set { _blobStores = value.ToList(); } 
     }
 
-    public BlobService(IServiceProvider provider, IOptionsMonitor<HordeStorageSettings> settings, IBlobIndex blobIndex, IPeerStatusService peerStatusService, IHttpClientFactory httpClientFactory, IServiceCredentials serviceCredentials)
+    public BlobService(IServiceProvider provider, IOptionsMonitor<HordeStorageSettings> settings, IBlobIndex blobIndex, IPeerStatusService peerStatusService, IHttpClientFactory httpClientFactory, IServiceCredentials serviceCredentials, INamespacePolicyResolver namespacePolicyResolver)
     {
         _blobStores = GetBlobStores(provider, settings).ToList();
         _settings = settings;
@@ -73,6 +75,7 @@ public class BlobService : IBlobService
         _peerStatusService = peerStatusService;
         _httpClientFactory = httpClientFactory;
         _serviceCredentials = serviceCredentials;
+        _namespacePolicyResolver = namespacePolicyResolver;
     }
 
     private IEnumerable<IBlobStore> GetBlobStores(IServiceProvider provider, IOptionsMonitor<HordeStorageSettings> settings)
@@ -287,7 +290,7 @@ public class BlobService : IBlobService
 
     public bool ShouldFetchBlobOnDemand(NamespaceId ns)
     {
-        return _settings.CurrentValue.EnableOnDemandReplication && _settings.CurrentValue.OnDemandReplicationNamespaces.Any(s => s.Equals("*") || s.Equals(ns.ToString()));
+        return _settings.CurrentValue.EnableOnDemandReplication && _namespacePolicyResolver.GetPoliciesForNs(ns).OnDemandReplication;
     }
 
     private HttpRequestMessage BuildHttpRequest(HttpMethod httpMethod, string uri)
@@ -301,7 +304,7 @@ public class BlobService : IBlobService
 
     public async Task<bool> Exists(NamespaceId ns, BlobIdentifier blob)
     {
-        bool useBlobIndex = _settings.CurrentValue.NamespacesThatUseBlobIndexForExistsCheck.Contains(ns.ToString());
+        bool useBlobIndex = _namespacePolicyResolver.GetPoliciesForNs(ns).UseBlobIndexForExists;
         if (useBlobIndex)
         {
             using IScope scope = Tracer.Instance.StartActive("HierarchicalStore.ObjectExists");
diff --git a/Engine/Source/Programs/Horde/HordeStorage/Horde.Storage/Implementation/GC/CacheCleanup.cs b/Engine/Source/Programs/Horde/HordeStorage/Horde.Storage/Implementation/GC/CacheCleanup.cs
--- a/Engine/Source/Programs/Horde/HordeStorage/Horde.Storage/Implementation/GC/CacheCleanup.cs
+++ b/Engine/Source/Programs/Horde/HordeStorage/Horde.Storage/Implementation/GC/CacheCleanup.cs
@@ -6,6 +6,7 @@ using System.Threading;
 using System.Threading.Tasks;
 using Datadog.Trace;
 using EpicGames.Horde.Storage;
+using Jupiter;
 using Jupiter.Common;
 using Jupiter.Implementation;
 using Microsoft.Extensions.Options;
@@ -24,20 +25,24 @@ namespace Horde.Storage.Implementation
         private readonly IRefsStore _refs;
         private readonly IReferencesStore _referencesStore;
         private readonly IReplicationLog _replicationLog;
+        private readonly INamespacePolicyResolver _namespacePolicyResolver;
         private readonly ITransactionLogWriter _transactionLog;
         private readonly ILogger _logger = Log.ForContext<RefCleanup>();
 
-        public RefCleanup(IOptionsMonitor<GCSettings> settings, IRefsStore refs, ITransactionLogWriter transactionLog, IReferencesStore referencesStore, IReplicationLog replicationLog)
+        public RefCleanup(IOptionsMonitor<GCSettings> settings, IRefsStore refs, ITransactionLogWriter transactionLog, IReferencesStore referencesStore, IReplicationLog replicationLog, INamespacePolicyResolver namespacePolicyResolver)
         {
             _settings = settings;
             _refs = refs;
             _transactionLog = transactionLog;
             _referencesStore = referencesStore;
             _replicationLog = replicationLog;
+            _namespacePolicyResolver = namespacePolicyResolver;
         }
 
         public Task<List<OldRecord>> Cleanup(NamespaceId ns, CancellationToken cancellationToken)
         {
+            NamespaceSettings.PerNamespaceSettings policies = _namespacePolicyResolver.GetPoliciesForNs(ns);
+
             using IScope scope = Tracer.Instance.StartActive("gc.refs.namespace");
             scope.Span.ResourceName = ns.ToString();
             if (ns == INamespacePolicyResolver.JupiterInternalNamespace)
@@ -45,18 +50,23 @@ namespace Horde.Storage.Implementation
                 // do not apply our cleanup policies to the internal namespace
                 return Task.FromResult(new List<OldRecord>());
             }
-            else if (_settings.CurrentValue.CleanNamespaces.Contains(ns.ToString()))
+            else if (!policies.IsLegacyNamespace.HasValue)
+            {
+                throw new NotImplementedException(
+                    $"Namespace {ns} did not set IsLegacyNamespace, unable to clean it as we do not know which method to use.");
+            }
+            else if (!policies.IsLegacyNamespace.Value)
             {
                 return CleanNamespace(ns, cancellationToken);
             }
-            else if (_settings.CurrentValue.CleanNamespacesLegacy.Contains(ns.ToString()))
+            else if (policies.IsLegacyNamespace.Value)
             {
                 return CleanNamespaceLegacy(ns, cancellationToken);
             }
             else
-                throw new NotImplementedException(
-                    $"Namespace {ns} not present in CleanNamespaces or CleanNamespacesLegacy lists, unable to clean it as we do not know which method to use.");
-           
+            {
+                throw new NotImplementedException();
+            }
         }
 
         private async Task<List<OldRecord>> CleanNamespace(NamespaceId ns, CancellationToken cancellationToken)
diff --git a/Engine/Source/Programs/Horde/HordeStorage/Horde.Storage/Implementation/GC/OrphanBlobCleanup.cs b/Engine/Source/Programs/Horde/HordeStorage/Horde.Storage/Implementation/GC/OrphanBlobCleanup.cs
--- a/Engine/Source/Programs/Horde/HordeStorage/Horde.Storage/Implementation/GC/OrphanBlobCleanup.cs
+++ b/Engine/Source/Programs/Horde/HordeStorage/Horde.Storage/Implementation/GC/OrphanBlobCleanup.cs
@@ -10,6 +10,7 @@ using async_enumerable_dotnet;
 using Dasync.Collections;
 using EpicGames.Horde.Storage;
 using Jupiter;
+using Jupiter.Common;
 using Jupiter.Implementation;
 using Microsoft.Extensions.Options;
 using RestSharp;
@@ -29,15 +30,17 @@ namespace Horde.Storage.Implementation
         private readonly IBlobService _blobService;
         private readonly ILeaderElection _leaderElection;
         private readonly IOptionsMonitor<GCSettings> _gcSettings;
+        private readonly INamespacePolicyResolver _namespacePolicyResolver;
         private readonly IRestClient _client;
         private readonly ILogger _logger = Log.ForContext<OrphanBlobCleanup>();
 
         // ReSharper disable once UnusedMember.Global
-        public OrphanBlobCleanup(IBlobService blobService, ILeaderElection leaderElection, IOptionsMonitor<CallistoTransactionLogSettings> callistoSettings, IOptionsMonitor<GCSettings> gcSettings, IServiceCredentials serviceCredentials)
+        public OrphanBlobCleanup(IBlobService blobService, ILeaderElection leaderElection, IOptionsMonitor<CallistoTransactionLogSettings> callistoSettings, IOptionsMonitor<GCSettings> gcSettings, IServiceCredentials serviceCredentials, INamespacePolicyResolver namespacePolicyResolver)
         {
             _blobService = blobService;
             _leaderElection = leaderElection;
             _gcSettings = gcSettings;
+            _namespacePolicyResolver = namespacePolicyResolver;
 
             _client = new RestClient(callistoSettings.CurrentValue.ConnectionString)
             {
@@ -45,12 +48,13 @@ namespace Horde.Storage.Implementation
             }.UseSerializer(() => new JsonNetSerializer());
         }
 
-        internal OrphanBlobCleanup(IBlobService blobService, ILeaderElection leaderElection, IRestClient callistoClient, IOptionsMonitor<GCSettings> gcSettings)
+        internal OrphanBlobCleanup(IBlobService blobService, ILeaderElection leaderElection, IRestClient callistoClient, IOptionsMonitor<GCSettings> gcSettings, INamespacePolicyResolver namespacePolicyResolver)
         {
             _blobService = blobService;
             _leaderElection = leaderElection;
             _client = callistoClient;
             _gcSettings = gcSettings;
+            _namespacePolicyResolver = namespacePolicyResolver;
         }
 
         private struct GCRootState
@@ -167,7 +171,9 @@ namespace Horde.Storage.Implementation
 
         private bool NamespaceShouldBeCleaned(NamespaceId ns)
         {
-            return _gcSettings.CurrentValue.CleanNamespacesLegacy.Contains(ns.ToString());
+            NamespaceSettings.PerNamespaceSettings policy = _namespacePolicyResolver.GetPoliciesForNs(ns);
+
+            return policy.IsLegacyNamespace.HasValue && policy.IsLegacyNamespace.Value;
         }
 
         public virtual async IAsyncEnumerable<NamespaceId> ListNamespaces()
diff --git a/Engine/Source/Programs/Horde/HordeStorage/Horde.Storage/Implementation/GC/OrphanBlobCleanupRefs.cs b/Engine/Source/Programs/Horde/HordeStorage/Horde.Storage/Implementation/GC/OrphanBlobCleanupRefs.cs
--- a/Engine/Source/Programs/Horde/HordeStorage/Horde.Storage/Implementation/GC/OrphanBlobCleanupRefs.cs
+++ b/Engine/Source/Programs/Horde/HordeStorage/Horde.Storage/Implementation/GC/OrphanBlobCleanupRefs.cs
@@ -145,7 +145,9 @@ namespace Horde.Storage.Implementation
 
         private bool NamespaceShouldBeCleaned(NamespaceId ns)
         {
-            return _gcSettings.CurrentValue.CleanNamespaces.Contains(ns.ToString());
+            NamespaceSettings.PerNamespaceSettings policy = _namespacePolicyResolver.GetPoliciesForNs(ns);
+
+            return policy.IsLegacyNamespace.HasValue && !policy.IsLegacyNamespace.Value;
         }
 
         private IAsyncEnumerable<NamespaceId> ListNamespaces()
diff --git a/Engine/Source/Programs/Horde/HordeStorage/Horde.Storage/Implementation/LastAccess/LastAccessService.cs b/Engine/Source/Programs/Horde/HordeStorage/Horde.Storage/Implementation/LastAccess/LastAccessService.cs
--- a/Engine/Source/Programs/Horde/HordeStorage/Horde.Storage/Implementation/LastAccess/LastAccessService.cs
+++ b/Engine/Source/Programs/Horde/HordeStorage/Horde.Storage/Implementation/LastAccess/LastAccessService.cs
@@ -5,6 +5,8 @@ using System.Collections.Generic;
 using System.Threading;
 using System.Threading.Tasks;
 using Datadog.Trace;
+using EpicGames.Horde.Storage;
+using Jupiter.Common;
 using Microsoft.Extensions.Hosting;
 using Microsoft.Extensions.Options;
 using Serilog;
@@ -16,16 +18,18 @@ namespace Horde.Storage.Implementation
     {
         private readonly ILastAccessCache<RefRecord> _lastAccessCacheRecord;
         private readonly IRefsStore _refsStore;
+        private readonly INamespacePolicyResolver _namespacePolicyResolver;
         private readonly ILogger _logger = Log.ForContext<LastAccessService>();
         private Timer? _timer;
         private readonly HordeStorageSettings _settings;
         
         public bool Running { get; private set; }
 
-        public LastAccessService(IOptionsMonitor<HordeStorageSettings> settings, ILastAccessCache<RefRecord> lastAccessCache, IRefsStore refsStore)
+        public LastAccessService(IOptionsMonitor<HordeStorageSettings> settings, ILastAccessCache<RefRecord> lastAccessCache, IRefsStore refsStore, INamespacePolicyResolver namespacePolicyResolver)
         {
             _lastAccessCacheRecord = lastAccessCache;
             _refsStore = refsStore;
+            _namespacePolicyResolver = namespacePolicyResolver;
             _settings = settings.CurrentValue;
         }
 
@@ -66,6 +70,9 @@ namespace Horde.Storage.Implementation
             List<(RefRecord, DateTime)> records = await _lastAccessCacheRecord.GetLastAccessedRecords();
             foreach ((RefRecord record, DateTime lastAccessTime) in records)
             {
+                if (!ShouldTrackLastAccess(record.Namespace))
+                    continue;
+
                 using IScope scope = Tracer.Instance.StartActive("lastAccess.update.record");
                 scope.Span.ResourceName = $"{record.Namespace}:{record.Bucket}.{record.RefName}";
                 _logger.Debug("Updating last access time to {LastAccessTime} for {Record}", lastAccessTime, record);
@@ -75,6 +82,11 @@ namespace Horde.Storage.Implementation
             return records;
         }
 
+        private bool ShouldTrackLastAccess(NamespaceId ns)
+        {
+            return _namespacePolicyResolver.GetPoliciesForNs(ns).LastAccessTracking;
+        }
+
         public void Dispose()
         {
             _timer?.Dispose();
diff --git a/Engine/Source/Programs/Horde/HordeStorage/Horde.Storage/Implementation/LastAccess/LastAccessServiceReferences.cs b/Engine/Source/Programs/Horde/HordeStorage/Horde.Storage/Implementation/LastAccess/LastAccessServiceReferences.cs
--- a/Engine/Source/Programs/Horde/HordeStorage/Horde.Storage/Implementation/LastAccess/LastAccessServiceReferences.cs
+++ b/Engine/Source/Programs/Horde/HordeStorage/Horde.Storage/Implementation/LastAccess/LastAccessServiceReferences.cs
@@ -5,6 +5,8 @@ using System.Collections.Generic;
 using System.Threading;
 using System.Threading.Tasks;
 using Datadog.Trace;
+using EpicGames.Horde.Storage;
+using Jupiter.Common;
 using Microsoft.Extensions.Hosting;
 using Microsoft.Extensions.Options;
 using Serilog;
@@ -16,16 +18,18 @@ namespace Horde.Storage.Implementation
     {
         private readonly ILastAccessCache<LastAccessRecord> _lastAccessCacheRecord;
         private readonly IReferencesStore _referencesStore;
+        private readonly INamespacePolicyResolver _namespacePolicyResolver;
         private readonly ILogger _logger = Log.ForContext<LastAccessServiceReferences>();
         private Timer? _timer;
         private readonly HordeStorageSettings _settings;
         
         public bool Running { get; private set; }
 
-        public LastAccessServiceReferences(IOptionsMonitor<HordeStorageSettings> settings, ILastAccessCache<LastAccessRecord> lastAccessCache, IReferencesStore referencesStore)
+        public LastAccessServiceReferences(IOptionsMonitor<HordeStorageSettings> settings, ILastAccessCache<LastAccessRecord> lastAccessCache, IReferencesStore referencesStore, INamespacePolicyResolver namespacePolicyResolver)
         {
             _lastAccessCacheRecord = lastAccessCache;
             _referencesStore = referencesStore;
+            _namespacePolicyResolver = namespacePolicyResolver;
             _settings = settings.CurrentValue;
         }
 
@@ -73,6 +77,8 @@ namespace Horde.Storage.Implementation
             List<(LastAccessRecord, DateTime)> records = await _lastAccessCacheRecord.GetLastAccessedRecords();
             foreach ((LastAccessRecord record, DateTime lastAccessTime) in records)
             {
+                if (!ShouldTrackLastAccess(record.Namespace))
+                    continue;
                 using IScope scope = Tracer.Instance.StartActive("lastAccess.update.record");
                 scope.Span.ResourceName = $"{record.Namespace}:{record.Bucket}.{record.Key}";
                 _logger.Debug("Updating last access time to {LastAccessTime} for {Record}", lastAccessTime, record);
@@ -82,6 +88,11 @@ namespace Horde.Storage.Implementation
             return records;
         }
 
+        private bool ShouldTrackLastAccess(NamespaceId ns)
+        {
+            return _namespacePolicyResolver.GetPoliciesForNs(ns).LastAccessTracking;
+        }
+
         public void Dispose()
         {
             _timer?.Dispose();
diff --git a/Engine/Source/Programs/Horde/HordeStorage/Horde.Storage/Settings.cs b/Engine/Source/Programs/Horde/HordeStorage/Horde.Storage/Settings.cs
--- a/Engine/Source/Programs/Horde/HordeStorage/Horde.Storage/Settings.cs
+++ b/Engine/Source/Programs/Horde/HordeStorage/Horde.Storage/Settings.cs
@@ -132,12 +132,9 @@ namespace Horde.Storage
         public ContentIdStoreImplementations ContentIdStoreImplementation { get; set; } = ContentIdStoreImplementations.Memory;
         public BlobIndexImplementations BlobIndexImplementation { get; set; } = BlobIndexImplementations.Memory;
 
-        public List<string> NamespacesThatUseBlobIndexForExistsCheck { get; set; } = new List<string>();
-
         public int? MaxSingleBlobSize { get; set; } = null; // disable blob partitioning
 
         public int LastAccessRollupFrequencySeconds { get; set; } = 900; // 15 minutes
-        public List<string> OnDemandReplicationNamespaces { get; set; } = new List<string> { "*" };
         public bool EnableOnDemandReplication { get; set; } = false;
     }
 
@@ -260,8 +257,6 @@ namespace Horde.Storage
 
         public TimeSpan BlobCleanupPollFrequency { get; set; } = TimeSpan.FromMinutes(60);
         public TimeSpan RefCleanupPollFrequency { get; set; } = TimeSpan.FromMinutes(60);
-        public List<string> CleanNamespaces { get; set; } = new List<string>();
-        public List<string> CleanNamespacesLegacy { get; set; }  = new List<string>();
 
     }
 
diff --git a/Engine/Source/Programs/Horde/HordeStorage/Horde.Storage/Tests/Functional/Blob/BlobReplicationTests.cs b/Engine/Source/Programs/Horde/HordeStorage/Horde.Storage/Tests/Functional/Blob/BlobReplicationTests.cs
--- a/Engine/Source/Programs/Horde/HordeStorage/Horde.Storage/Tests/Functional/Blob/BlobReplicationTests.cs
+++ b/Engine/Source/Programs/Horde/HordeStorage/Horde.Storage/Tests/Functional/Blob/BlobReplicationTests.cs
@@ -8,6 +8,7 @@ using System.Text;
 using System.Threading.Tasks;
 using EpicGames.Horde.Storage;
 using Horde.Storage.Implementation.Blob;
+using Jupiter;
 using Jupiter.Implementation;
 using Microsoft.AspNetCore.Hosting;
 using Microsoft.AspNetCore.TestHost;
@@ -24,7 +25,7 @@ namespace Horde.Storage.FunctionalTests.Storage
     [TestClass]
     public class BlobReplicationTests
     {
-        protected readonly NamespaceId TestNamespaceName = new NamespaceId("blob-replication-ns");
+        protected readonly NamespaceId TestNamespaceName = new NamespaceId("test-namespace");
 
         
         [TestMethod]
@@ -106,6 +107,19 @@ namespace Horde.Storage.FunctionalTests.Storage
 
                         }.ToList();
                     });
+                    collection.Configure<NamespaceSettings>(settings =>
+                    {
+                        settings.Policies = new Dictionary<string, NamespaceSettings.PerNamespaceSettings>()
+                        {
+                            {
+                                TestNamespaceName.ToString(), new NamespaceSettings.PerNamespaceSettings()
+                                {
+                                    OnDemandReplication = true
+                                }
+                            }
+                        };
+                    });
+
                     collection.AddSingleton<IHttpClientFactory>(handler.CreateClientFactory());
                 })
                 .UseStartup<HordeStorageStartup>()
diff --git a/Engine/Source/Programs/Horde/HordeStorage/Horde.Storage/Tests/Functional/GC/GCBlobTests.cs b/Engine/Source/Programs/Horde/HordeStorage/Horde.Storage/Tests/Functional/GC/GCBlobTests.cs
--- a/Engine/Source/Programs/Horde/HordeStorage/Horde.Storage/Tests/Functional/GC/GCBlobTests.cs
+++ b/Engine/Source/Programs/Horde/HordeStorage/Horde.Storage/Tests/Functional/GC/GCBlobTests.cs
@@ -23,6 +23,8 @@ using Serilog.Core;
 using Moq;
 using EpicGames.Horde.Storage;
 using EpicGames.Serialization;
+using Jupiter;
+using Jupiter.Common;
 
 namespace Horde.Storage.FunctionalTests.GC
 {
@@ -84,8 +86,11 @@ namespace Horde.Storage.FunctionalTests.GC
                         };
                         IOptionsMonitor<GCSettings> gcSettingsMon = Mock.Of<IOptionsMonitor<GCSettings>>(_ => _.CurrentValue == gcSettings);
 
+                        NamespacePolicyResolver policyResolver = Mock.Of<NamespacePolicyResolver>();
+
+
                         IBlobService blobService = provider.GetService<IBlobService>()!;
-                        return new OrphanBlobCleanup(blobService!, new LeaderElectionStub(true), _callistoBlobMock.Object, gcSettingsMon);
+                        return new OrphanBlobCleanup(blobService!, new LeaderElectionStub(true), _callistoBlobMock.Object, gcSettingsMon, policyResolver);
                     });
 
                 })
@@ -171,11 +176,16 @@ namespace Horde.Storage.FunctionalTests.GC
             {
                 CleanOldBlobs = true,
                 CleanOldBlobsLegacy = true,
-                CleanNamespacesLegacy = new List<string> { TestNamespace.ToString() }
             };
             IOptionsMonitor<GCSettings> gcSettingsMon = Mock.Of<IOptionsMonitor<GCSettings>>(_ => _.CurrentValue == gcSettings);
 
-            OrphanBlobCleanup cleanup = new OrphanBlobCleanup(_blobService!, new LeaderElectionStub(true), _callistoBlobMock.Object, gcSettingsMon);
+            Mock<INamespacePolicyResolver> policyResolverMock = new Mock<INamespacePolicyResolver>();
+            policyResolverMock.Setup(resolver => resolver.GetPoliciesForNs(TestNamespace)).Returns(new NamespaceSettings.PerNamespaceSettings()
+            {
+                IsLegacyNamespace = true,
+            });
+
+            OrphanBlobCleanup cleanup = new OrphanBlobCleanup(_blobService!, new LeaderElectionStub(true), _callistoBlobMock.Object, gcSettingsMon, policyResolverMock.Object);
             List<NamespaceId> namespaces = await cleanup.ListNamespaces().ToListAsync();
             Assert.AreEqual(1, namespaces.Count);
 
@@ -232,7 +242,6 @@ namespace Horde.Storage.FunctionalTests.GC
                 .AddInMemoryCollection(new List<KeyValuePair<string, string>>()
                 {
                     new KeyValuePair<string, string>("Horde_Storage:StorageImplementations:0", "MemoryBlobStore"),
-                    new KeyValuePair<string, string>("GC:CleanNamespaces:0", TestNamespace.ToString()),
                     new KeyValuePair<string, string>("GC:CleanOldBlobs", true.ToString()),
                     new KeyValuePair<string, string>("GC:CleanOldBlobsLegacy", false.ToString()),
                     new KeyValuePair<string, string>("Horde_Storage:BlobIndexImplementation", GetImplementation()),
@@ -248,6 +257,21 @@ namespace Horde.Storage.FunctionalTests.GC
                 .UseConfiguration(configuration)
                 .UseEnvironment("Testing")
                 .UseSerilog(logger)
+                .ConfigureTestServices(collection =>
+                {
+                    collection.Configure<NamespaceSettings>(settings =>
+                    {
+                        settings.Policies = new Dictionary<string, NamespaceSettings.PerNamespaceSettings>()
+                        {
+                            {
+                                TestNamespace.ToString(), new NamespaceSettings.PerNamespaceSettings()
+                                {
+                                    IsLegacyNamespace = false
+                                }
+                            }
+                        };
+                    });
+                })
                 .UseStartup<HordeStorageStartup>()
             );
 
diff --git a/Engine/Source/Programs/Horde/HordeStorage/Horde.Storage/Tests/Functional/GC/GCRefTests.cs b/Engine/Source/Programs/Horde/HordeStorage/Horde.Storage/Tests/Functional/GC/GCRefTests.cs
--- a/Engine/Source/Programs/Horde/HordeStorage/Horde.Storage/Tests/Functional/GC/GCRefTests.cs
+++ b/Engine/Source/Programs/Horde/HordeStorage/Horde.Storage/Tests/Functional/GC/GCRefTests.cs
@@ -20,6 +20,7 @@ using Serilog;
 using Serilog.Core;
 using EpicGames.Horde.Storage;
 using EpicGames.Serialization;
+using Jupiter;
 
 namespace Horde.Storage.FunctionalTests.GC
 {
@@ -41,10 +42,6 @@ namespace Horde.Storage.FunctionalTests.GC
                             // we are not reading the base appSettings here as we want exact control over what runs in the tests
                             .AddJsonFile("appsettings.Testing.json", false)
                             .AddEnvironmentVariables()
-                            .AddInMemoryCollection(new List<KeyValuePair<string, string>>()
-                            {
-                                new KeyValuePair<string, string>("GC:CleanNamespacesLegacy:0", TestNamespace.ToString()),
-                            })
                             .Build();
 
             Logger logger = new LoggerConfiguration()
@@ -70,6 +67,19 @@ namespace Horde.Storage.FunctionalTests.GC
                 {
                     // use our refs mock instead of the actual refs store so we can control which records are invalid
                     collection.AddSingleton<IRefsStore>(_refMock.Object);
+
+                    collection.Configure<NamespaceSettings>(settings =>
+                    {
+                        settings.Policies = new Dictionary<string, NamespaceSettings.PerNamespaceSettings>()
+                        {
+                            {
+                                TestNamespace.ToString(), new NamespaceSettings.PerNamespaceSettings()
+                                {
+                                    IsLegacyNamespace = true
+                                }
+                            }
+                        };
+                    });
                 })
             );
             _httpClient = server.CreateClient();
@@ -157,7 +167,6 @@ namespace Horde.Storage.FunctionalTests.GC
                 .AddEnvironmentVariables()
                 .AddInMemoryCollection(new List<KeyValuePair<string, string>>()
                 {
-                    new KeyValuePair<string, string>("GC:CleanNamespaces:0", TestNamespace.ToString()),
                     new KeyValuePair<string, string>("Horde_Storage:StorageImplementations:0", "MemoryBlobStore"),
                     new KeyValuePair<string, string>("Horde_Storage:BlobIndexImplementation", GetImplementation()),
                 })
@@ -170,6 +179,21 @@ namespace Horde.Storage.FunctionalTests.GC
                 .UseConfiguration(configuration)
                 .UseEnvironment("Testing")
                 .UseSerilog(logger)
+                .ConfigureTestServices(collection =>
+                {
+                    collection.Configure<NamespaceSettings>(settings =>
+                    {
+                        settings.Policies = new Dictionary<string, NamespaceSettings.PerNamespaceSettings>()
+                        {
+                            {
+                                TestNamespace.ToString(), new NamespaceSettings.PerNamespaceSettings()
+                                {
+                                    IsLegacyNamespace = false
+                                }
+                            }
+                        };
+                    });
+                })
                 .UseStartup<HordeStorageStartup>()
             );
             _httpClient = server.CreateClient();
diff --git a/Engine/Source/Programs/Horde/HordeStorage/Horde.Storage/Tests/UnitTests/HierarchicalBlobStorageTest.cs b/Engine/Source/Programs/Horde/HordeStorage/Horde.Storage/Tests/UnitTests/HierarchicalBlobStorageTest.cs
--- a/Engine/Source/Programs/Horde/HordeStorage/Horde.Storage/Tests/UnitTests/HierarchicalBlobStorageTest.cs
+++ b/Engine/Source/Programs/Horde/HordeStorage/Horde.Storage/Tests/UnitTests/HierarchicalBlobStorageTest.cs
@@ -11,6 +11,8 @@ using Dasync.Collections;
 using EpicGames.Horde.Storage;
 using Horde.Storage.Implementation;
 using Horde.Storage.Implementation.Blob;
+using Jupiter;
+using Jupiter.Common;
 using Jupiter.Implementation;
 using Microsoft.Extensions.Options;
 using Microsoft.VisualStudio.TestTools.UnitTesting;
@@ -51,7 +53,10 @@ namespace Horde.Storage.UnitTests
             Mock<IServiceProvider> serviceProviderMock = new Mock<IServiceProvider>();
             serviceProviderMock.Setup(x => x.GetService(typeof(MemoryCacheBlobStore))).Returns(new MemoryCacheBlobStore(Mock.Of<IOptionsMonitor<MemoryCacheBlobSettings>>(_ => _.CurrentValue == new MemoryCacheBlobSettings())));
             IOptionsMonitor<HordeStorageSettings> settingsMonitor = Mock.Of<IOptionsMonitor<HordeStorageSettings>>(_ => _.CurrentValue == new HordeStorageSettings());
-            _chained = new BlobService(serviceProviderMock.Object, settingsMonitor, Mock.Of<IBlobIndex>(), Mock.Of<IPeerStatusService>(), Mock.Of<IHttpClientFactory>(), Mock.Of<IServiceCredentials>());
+            Mock<INamespacePolicyResolver> mockPolicyResolver = new Mock<INamespacePolicyResolver>();
+            mockPolicyResolver.Setup(x => x.GetPoliciesForNs(It.IsAny<NamespaceId>())).Returns(new NamespaceSettings.PerNamespaceSettings());
+
+            _chained = new BlobService(serviceProviderMock.Object, settingsMonitor, Mock.Of<IBlobIndex>(), Mock.Of<IPeerStatusService>(), Mock.Of<IHttpClientFactory>(), Mock.Of<IServiceCredentials>(), mockPolicyResolver.Object);
             _chained.BlobStore = new List<IBlobStore> { _first, _second, _third };
 
             await _first.PutObject(Ns, Encoding.ASCII.GetBytes("onlyFirstContent"), _onlyFirstId);
diff --git a/Engine/Source/Programs/Horde/HordeStorage/Horde.Storage/appsettings.Testing.json b/Engine/Source/Programs/Horde/HordeStorage/Horde.Storage/appsettings.Testing.json
--- a/Engine/Source/Programs/Horde/HordeStorage/Horde.Storage/appsettings.Testing.json
+++ b/Engine/Source/Programs/Horde/HordeStorage/Horde.Storage/appsettings.Testing.json
@@ -86,8 +86,14 @@
       "test-namespace": {
         "Claims": [ "*" ]
       },
+      "test-delete-namespace": {
+        "Claims": [ "*" ]
+      },
       "testbucket": {
         "Claims": [ "*" ]
+      },
+      "snapshot-namespace": {
+        "Claims": [ "*" ]
       }
     }
   },
diff --git a/Engine/Source/Programs/Horde/HordeStorage/Jupiter.Common/BaseStartup.cs b/Engine/Source/Programs/Horde/HordeStorage/Jupiter.Common/BaseStartup.cs
--- a/Engine/Source/Programs/Horde/HordeStorage/Jupiter.Common/BaseStartup.cs
+++ b/Engine/Source/Programs/Horde/HordeStorage/Jupiter.Common/BaseStartup.cs
@@ -464,6 +464,12 @@ namespace Jupiter
         {
             public string[] Claims { get; set; } = Array.Empty<string>();
             public string StoragePool { get; set; } = "";
+
+            public bool LastAccessTracking { get; set; } = true;
+            public bool OnDemandReplication { get; set; } = false;
+            public bool UseBlobIndexForExists { get; set; } = false;
+            public bool UseBlobIndexForSlowExists { get; set; } = false;
+            public bool? IsLegacyNamespace { get; set; } = null;
         }
     }
 
