Horde: Prevent content replication when disabled for a stream.

diff --git a/Engine/Source/Programs/Horde/Horde.Build/Perforce/Impl/CommitService.cs b/Engine/Source/Programs/Horde/Horde.Build/Perforce/Impl/CommitService.cs
--- a/Engine/Source/Programs/Horde/Horde.Build/Perforce/Impl/CommitService.cs
+++ b/Engine/Source/Programs/Horde/Horde.Build/Perforce/Impl/CommitService.cs
@@ -833,7 +833,13 @@ namespace HordeServer.Commits.Impl
 		/// <param name="Change">Commit to store the tree ref</param>
 		/// <param name="BaseChange">The base change to update from</param>
 		/// <returns>Root tree object</returns>
-		public Task WriteCommitTreeAsync(IStream Stream, int Change, int? BaseChange) => WriteCommitTreeAsync(Stream, Change, BaseChange, Stream.ReplicationFilter, Stream.ReplicationMode == ContentReplicationMode.RevisionsOnly);
+		public async Task WriteCommitTreeAsync(IStream Stream, int Change, int? BaseChange)
+		{
+			if (Stream.ReplicationMode != ContentReplicationMode.None)
+			{
+				await WriteCommitTreeAsync(Stream, Change, BaseChange, Stream.ReplicationFilter, Stream.ReplicationMode == ContentReplicationMode.RevisionsOnly);
+			}
+		}
 
 		/// <summary>
 		/// Replicates the contents of a stream to Horde storage, optionally using the given change as a starting point
