Horde Storage - Removed managing ttl records for refs, as we do not use them anymore.

diff --git a/Engine/Source/Programs/Horde/HordeStorage/Horde.Storage/Implementation/Objects/ScyllaReferencesStore.cs b/Engine/Source/Programs/Horde/HordeStorage/Horde.Storage/Implementation/Objects/ScyllaReferencesStore.cs
--- a/Engine/Source/Programs/Horde/HordeStorage/Horde.Storage/Implementation/Objects/ScyllaReferencesStore.cs
+++ b/Engine/Source/Programs/Horde/HordeStorage/Horde.Storage/Implementation/Objects/ScyllaReferencesStore.cs
@@ -50,16 +50,6 @@ namespace Horde.Storage.Implementation
             );"
             ));
 
-            _session.Execute(new SimpleStatement(@"CREATE TABLE IF NOT EXISTS object_last_access (
-                namespace text,
-                partition_index tinyint,
-                accessed_at date,
-                objects set<text>,
-                PRIMARY KEY ((namespace, partition_index, accessed_at))
-            );"
-            ));
-
-
             _getObjectsStatement = _session.Prepare("SELECT bucket, name, last_access_time FROM objects WHERE namespace = ? ALLOW FILTERING");
             _getNamespacesStatement = _session.Prepare("SELECT DISTINCT namespace FROM buckets");
         }
@@ -105,12 +95,10 @@ namespace Horde.Storage.Implementation
             }
 
             // add the bucket in parallel with inserting the actual object
-            var addBucketTask = MaybeAddBucket(ns, bucket);
-            // add the ttl records in parallel with the object
-            var createTTLRecords = CreateTTLRecord(ns, bucket, name, (sbyte)blobHash.HashData[0]);
+            Task addBucketTask = MaybeAddBucket(ns, bucket);
 
             await _mapper.InsertAsync<ScyllaObject>(new ScyllaObject(ns, bucket, name, blob, blobHash, isFinalized));
-            await Task.WhenAll(addBucketTask, createTTLRecords);
+            await addBucketTask;
         }
 
 
@@ -122,107 +110,13 @@ namespace Horde.Storage.Implementation
             await _mapper.UpdateAsync<ScyllaObject>("SET is_finalized=true WHERE namespace=? AND bucket=? AND name=?", ns.ToString(), bucket.ToString(), name.ToString());
         }
 
-        private async Task CreateTTLRecord(NamespaceId ns, BucketId bucket, IoHashKey name, sbyte partitionIndex)
-        {
-            using IScope _ = Tracer.Instance.StartActive("scylla.create_ttl_record");
-            DateTime now = DateTime.Now;
-            LocalDate lastAccessAt = new LocalDate(now.Year, now.Month, now.Day);
-            await _mapper.UpdateAsync<ScyllaObjectLastAccessAt>("SET objects = objects + ? WHERE accessed_at = ? AND namespace = ? AND partition_index = ? ", new string[] {$"{bucket}#{name}"},  lastAccessAt, ns.ToString(), partitionIndex);
-        }
-
-        private Task RemoveTTLRecord(ObjectRecord record)
-        {
-            return RemoveTTLRecord(record.Namespace, record.Bucket, record.Name, record.LastAccess, (sbyte)record.BlobIdentifier.HashData[0]);
-        }
-
-        private async Task RemoveTTLRecord(NamespaceId ns, BucketId bucket, IoHashKey name, DateTime lastAccess, sbyte partitionIndex)
-        {
-            using IScope _ = Tracer.Instance.StartActive("scylla.remove_ttl_record");
-            LocalDate lastAccessAt = new LocalDate(lastAccess.Year, lastAccess.Month, lastAccess.Day);
-            await _mapper.UpdateAsync<ScyllaObjectLastAccessAt>("SET objects = objects - ? WHERE accessed_at = ? AND namespace = ? AND partition_index = ?", new string[] {$"{bucket}#{name}"},  lastAccessAt, ns.ToString(), partitionIndex);
-        }
-
         public async Task UpdateLastAccessTime(NamespaceId ns, BucketId bucket, IoHashKey name, DateTime lastAccessTime)
         {
             using IScope _ = Tracer.Instance.StartActive("scylla.update_last_access_time");
-            // fetch the old record
-            /* ObjectRecord objectRecord = await Get(ns, bucket, name);
-
-            // update the object tracking
-            Task updateObjectTracking = _mapper.UpdateAsync<ScyllaObject>("SET last_access_time = ? WHERE namespace=? AND bucket = ? AND name = ?", lastAccessTime, ns.ToString(), bucket.ToString(), name.ToString());
-
-            // the object was accessed during the same day so no need to move it
-            if (objectRecord.LastAccess.Date != lastAccessTime.Date)
-            {
-                sbyte partitionIndex = (sbyte)objectRecord.BlobIdentifier.HashData[0];
-                // move which last access partition is used for the object
-                LocalDate lastAccessAt = new LocalDate(objectRecord.LastAccess.Year, objectRecord.LastAccess.Month, objectRecord.LastAccess.Day);
-                await _mapper.UpdateAsync<ScyllaObjectLastAccessAt>("SET objects = objects - ? WHERE accessed_at=? AND namespace = ? AND partition_index = ?", new string[] {$"{bucket}#{name}"}, lastAccessAt, ns.ToString(), partitionIndex);
-                LocalDate newLastAccessAt = new LocalDate(lastAccessTime.Year, lastAccessTime.Month, lastAccessTime.Day);
-                await _mapper.UpdateAsync<ScyllaObjectLastAccessAt>("SET objects = objects + ? WHERE accessed_at=? AND namespace = ? AND partition_index = ?", new string[] {$"{bucket}#{name}"}, newLastAccessAt, ns.ToString(), partitionIndex);
-            }
-
-            await updateObjectTracking;*/
 
             await _mapper.UpdateAsync<ScyllaObject>("SET last_access_time = ? WHERE namespace = ? AND bucket = ? AND name = ?", lastAccessTime, ns.ToString(), bucket.ToString(), name.ToString());
-
-            
         }
 
-        /*public async IAsyncEnumerable<ObjectRecord> GetRecords(NamespaceId ns)
-        {
-            using IScope _ = Tracer.Instance.StartActive("scylla.get_records");
-            for (sbyte partitionIndex = sbyte.MinValue; partitionIndex < sbyte.MaxValue; partitionIndex++)
-            {
-                using IScope scope = Tracer.Instance.StartActive("scylla.get_records_partition");
-                scope.Span.ResourceName = $"partition_{partitionIndex}";
-
-                SortedSet<DateTime> dateBuckets = new SortedSet<DateTime>();
-                {
-                    using IScope dateBucketsScope = Tracer.Instance.StartActive("scylla.get_records.date_buckets");
-
-                    RowSet rowSetDateBuckets = await _session.ExecuteAsync(new SimpleStatement("SELECT accessed_at FROM object_last_access WHERE namespace = ? AND partition_index = ? PER PARTITION LIMIT 1 ALLOW FILTERING", ns.ToString(), partitionIndex));
-                    foreach (Row row in rowSetDateBuckets)
-                    {
-                        if (rowSetDateBuckets.GetAvailableWithoutFetching() == 0)
-                            await rowSetDateBuckets.FetchMoreResultsAsync();
-
-                        LocalDate date = row.GetValue<LocalDate>("accessed_at");
-
-                        dateBuckets.Add(new DateTime(date.Year, date.Month, date.Day));
-                    }
-                }
-
-
-                foreach (DateTime dateBucket in dateBuckets)
-                {
-                    LocalDate date = new LocalDate(dateBucket.Year, dateBucket.Month, dateBucket.Day);
-                    foreach (ScyllaObjectLastAccessAt lastAccessRecord in await _mapper.FetchAsync<ScyllaObjectLastAccessAt>("WHERE namespace = ? AND accessed_at = ? AND partition_index = ?", ns.ToString(), date, partitionIndex))
-                    {
-                        foreach (string o in lastAccessRecord.Objects)
-                        {
-                            int bucketSeparator = o.IndexOf("#", StringComparison.InvariantCultureIgnoreCase);
-                            BucketId bucket = new BucketId(o.Substring(0, bucketSeparator));
-                            IoHashKey name = new IoHashKey(o.Substring(bucketSeparator + 1));
-
-                            ObjectRecord record;
-                            try
-                            {
-                                 record = await Get(ns, bucket, name);
-                            }
-                            catch (ObjectNotFoundException)
-                            {
-                                await RemoveTTLRecord(ns, bucket, name, dateBucket, partitionIndex);
-                                continue;
-                            }
-                            yield return record;
-                        }
-                    }
-                }
-
-            }
-        }*/
-        
         public async IAsyncEnumerable<(BucketId, IoHashKey, DateTime)> GetRecords(NamespaceId ns)
         {
             using IScope _ = Tracer.Instance.StartActive("scylla.get_records");
@@ -274,10 +168,8 @@ namespace Horde.Storage.Implementation
                 return false;
             }
 
-            Task removeTTL = RemoveTTLRecord(record);
             AppliedInfo<ScyllaObject> info = await _mapper.DeleteIfAsync<ScyllaObject>("WHERE namespace=? AND bucket=? AND name=?", ns.ToString(), bucket.ToString(), key.ToString());
 
-            await removeTTL;
             if (info.Applied)
                 return true;
 
