Horde Storage - Added resolveAttachment flag which can be used for HEAD or GET batch operations to indicate if we want to resolve that attachments exists as part of the operation (used to happen by default in HEAD but not in GET, now we can explicitly decide if it should or not).

diff --git a/Engine/Source/Programs/Horde/HordeStorage/Horde.Storage/Controllers/ReferencesController.cs b/Engine/Source/Programs/Horde/HordeStorage/Horde.Storage/Controllers/ReferencesController.cs
--- a/Engine/Source/Programs/Horde/HordeStorage/Horde.Storage/Controllers/ReferencesController.cs
+++ b/Engine/Source/Programs/Horde/HordeStorage/Horde.Storage/Controllers/ReferencesController.cs
@@ -631,11 +631,18 @@ namespace Horde.Storage.Controllers
                         throw new Exception();
 
                     CbObject cb = new CbObject(await blob.Stream.ToByteArray());
+
+                    if (op.ResolveAttachments ?? false)
+                    {
+                        IAsyncEnumerable<BlobIdentifier> references = _referenceResolver.ResolveReferences(ns, cb);
+                        List<BlobIdentifier>? _ = await references.ToListAsync();
+                    }
+
                     return (cb, HttpStatusCode.OK);
                 }
-                catch (ObjectNotFoundException e)
+                catch (Exception ex) when( ex is ObjectNotFoundException or PartialReferenceResolveException or ReferenceIsMissingBlobsException)
                 {
-                    return ToErrorResult(e, HttpStatusCode.NotFound);
+                    return ToErrorResult(ex, HttpStatusCode.NotFound);
                 }
                 catch (Exception e)
                 {
@@ -656,20 +663,22 @@ namespace Horde.Storage.Controllers
 
                     blob ??= await _blobStore.GetObject(ns, record.BlobIdentifier);
 
-                    // we have to verify the blobs are available locally, as the record of the key is replicated a head of the content
-                    // TODO: Once we support inline replication this step is not needed as at least one region as this blob, just maybe not this current one
-                    byte[] blobContents = await blob.Stream.ToByteArray();
-                    CbObject cb = new CbObject(blobContents);
-                    // the reference resolver will throw if any blob is missing, so no need to do anything other then process each reference
-                    IAsyncEnumerable<BlobIdentifier> references = _referenceResolver.ResolveReferences(ns, cb);
-                    List<BlobIdentifier>? _ = await references.ToListAsync();
+                    if (op.ResolveAttachments ?? false)
+                    {
+                        byte[] blobContents = await blob.Stream.ToByteArray();
+                        CbObject cb = new CbObject(blobContents);
+                        // the reference resolver will throw if any blob is missing, so no need to do anything other then process each reference
+                        IAsyncEnumerable<BlobIdentifier> references = _referenceResolver.ResolveReferences(ns, cb);
+                        List<BlobIdentifier>? _ = await references.ToListAsync();
+                    }
+
 
                     if (blob == null)
                         throw new Exception();
 
                     return (CbObject.Build(writer => writer.WriteBool("exists", true)), HttpStatusCode.OK);
                 }
-                catch (Exception ex) when( ex is ObjectNotFoundException || ex is PartialReferenceResolveException || ex is ReferenceIsMissingBlobsException)
+                catch (Exception ex) when( ex is ObjectNotFoundException or PartialReferenceResolveException or ReferenceIsMissingBlobsException)
                 {
                     return (CbObject.Build(writer => writer.WriteBool("exists", false)), HttpStatusCode.NotFound);
                 }
@@ -934,6 +943,9 @@ namespace Horde.Storage.Controllers
             [CbField("key")]
             public IoHashKey Key { get; set; }
 
+            [CbField("resolveAttachments")]
+            public bool? ResolveAttachments { get; set; } = null;
+
             [CbField("payload")]
             public CbObject? Payload { get; set; } = null;
 
diff --git a/Engine/Source/Programs/Horde/HordeStorage/Horde.Storage/Tests/Functional/Objects/ReferencesTests.cs b/Engine/Source/Programs/Horde/HordeStorage/Horde.Storage/Tests/Functional/Objects/ReferencesTests.cs
--- a/Engine/Source/Programs/Horde/HordeStorage/Horde.Storage/Tests/Functional/Objects/ReferencesTests.cs
+++ b/Engine/Source/Programs/Horde/HordeStorage/Horde.Storage/Tests/Functional/Objects/ReferencesTests.cs
@@ -1511,6 +1511,21 @@ namespace Horde.Storage.FunctionalTests.References
                 result.EnsureSuccessStatusCode();
             }
 
+            byte[] blobContentsMissing = Encoding.ASCII.GetBytes("This contents will not be submitted");
+            CbObject missingAttachmentObject = CbObject.Build(writer => writer.WriteBinaryAttachment("Attachment", IoHash.Compute(blobContentsMissing)));
+            IoHashKey missingAttachmentKey = IoHashKey.FromName("blobMissingAttachment");
+
+            {
+                BlobIdentifier blobHash = BlobIdentifier.FromBlob(blobContents);
+
+                HttpContent requestContent = new ByteArrayContent(blobContents);
+                requestContent.Headers.ContentType = new MediaTypeHeaderValue(MediaTypeNames.Application.Octet);
+                requestContent.Headers.Add(CommonHeaders.HashHeaderName, blobHash.ToString());
+
+                HttpResponseMessage result = await _httpClient!.PutAsync(requestUri: $"api/v1/blobs/{TestNamespace}/{blobHash}", requestContent);
+                result.EnsureSuccessStatusCode();
+            }
+
             {
                 byte[] cbObjectBytes = newReferenceObject.GetView().ToArray();
                 BlobIdentifier blobHash = BlobIdentifier.FromBlob(cbObjectBytes);
@@ -1543,10 +1558,21 @@ namespace Horde.Storage.FunctionalTests.References
             getObjectOp2.WriteString("key", newReferenceObjectKey.ToString());
             getObjectOp2.EndObject();
 
+            
+            CbWriter getObjectOp3 = new CbWriter();
+            getObjectOp3.BeginObject();
+            getObjectOp3.WriteInteger("opId", 2);
+            getObjectOp3.WriteString("op", BatchOps.BatchOp.Operation.GET.ToString());
+            getObjectOp3.WriteString("bucket", bucket.ToString());
+            getObjectOp3.WriteString("key", missingAttachmentKey.ToString());
+            getObjectOp3.WriteBool("resolveAttachments", true);
+            getObjectOp3.EndObject();
+
             CbObject[] ops = new[]
             {
                 getObjectOp.ToObject(),
                 getObjectOp2.ToObject(),
+                getObjectOp3.ToObject(),
             };
 
             CbWriter batchRequestWriter = new CbWriter();
@@ -1572,7 +1598,7 @@ namespace Horde.Storage.FunctionalTests.References
                 byte[] roundTrippedBuffer = ms.ToArray();
 
                 BatchOpsResponse response = CbSerializer.Deserialize<BatchOpsResponse>(roundTrippedBuffer);
-                Assert.AreEqual(2, response.Results.Count);
+                Assert.AreEqual(3, response.Results.Count);
                 
                 BatchOpsResponse.OpResponses op0 = response.Results.First(r => r.OpId == 0);
                 Assert.IsNotNull(op0.Response);
@@ -1583,6 +1609,11 @@ namespace Horde.Storage.FunctionalTests.References
                 Assert.IsNotNull(op1.Response);
                 Assert.AreEqual(200, op1.StatusCode);
                 Assert.AreEqual(newReferenceObject, op1.Response);
+
+                BatchOpsResponse.OpResponses op2 = response.Results.First(r => r.OpId == 2);
+                Assert.IsNotNull(op2.Response);
+                Assert.AreEqual(404, op2.StatusCode);
+                Assert.AreNotEqual(missingAttachmentObject, op2.Response);
             }
         }
 
