Horde Storage - Fixed 500 (unhandled exception) when trying to access a namespace which isnt configured

diff --git a/Engine/Source/Programs/Horde/HordeStorage/Jupiter.Common/Authentication/NamespaceAuthorizationHandler.cs b/Engine/Source/Programs/Horde/HordeStorage/Jupiter.Common/Authentication/NamespaceAuthorizationHandler.cs
--- a/Engine/Source/Programs/Horde/HordeStorage/Jupiter.Common/Authentication/NamespaceAuthorizationHandler.cs
+++ b/Engine/Source/Programs/Horde/HordeStorage/Jupiter.Common/Authentication/NamespaceAuthorizationHandler.cs
@@ -31,32 +31,40 @@ namespace Jupiter
                 return Task.CompletedTask;
             }
 
-            NamespaceSettings.PerNamespaceSettings settings = _namespacePolicyResolver.GetPoliciesForNs(namespaceName);
-            // These are ANDed, e.g. all claims needs to be present
-            foreach (string expectedClaim in settings.Claims)
+            try
             {
-                // if expected claim is * then everyone is allowed to use the namespace
-                if (expectedClaim == "*")
-                {
-                    context.Succeed(requirement);
-                    continue;
-                }
+                NamespaceSettings.PerNamespaceSettings settings = _namespacePolicyResolver.GetPoliciesForNs(namespaceName);
 
-                if (expectedClaim.Contains('='))
+                // These are ANDed, e.g. all claims needs to be present
+                foreach (string expectedClaim in settings.Claims)
                 {
-                    int separatorIndex = expectedClaim.IndexOf('=');
-                    string claimName = expectedClaim.Substring(0, separatorIndex);
-                    string claimValue = expectedClaim.Substring(separatorIndex + 1);
-                    if (context.User.HasClaim(claim => claim.Type == claimName && claim.Value == claimValue))
+                    // if expected claim is * then everyone is allowed to use the namespace
+                    if (expectedClaim == "*")
                     {
                         context.Succeed(requirement);
                         continue;
                     }
+
+                    if (expectedClaim.Contains('='))
+                    {
+                        int separatorIndex = expectedClaim.IndexOf('=');
+                        string claimName = expectedClaim.Substring(0, separatorIndex);
+                        string claimValue = expectedClaim.Substring(separatorIndex + 1);
+                        if (context.User.HasClaim(claim => claim.Type == claimName && claim.Value == claimValue))
+                        {
+                            context.Succeed(requirement);
+                            continue;
+                        }
+                    }
+                    if (context.User.HasClaim(claim => claim.Type == expectedClaim))
+                    {
+                        context.Succeed(requirement);
+                    }
                 }
-                if (context.User.HasClaim(claim => claim.Type == expectedClaim))
-                {
-                    context.Succeed(requirement);
-                }
+            }
+            catch (UnknownNamespaceException)
+            {
+                // if the namespace doesn't have a policy setup, e.g. we do not know which claims to require then we can just exit here as the auth will fail
             }
 
             return Task.CompletedTask;
diff --git a/Engine/Source/Programs/Horde/HordeStorage/Jupiter.Common/Implementation/NamespacePolicyResolver.cs b/Engine/Source/Programs/Horde/HordeStorage/Jupiter.Common/Implementation/NamespacePolicyResolver.cs
--- a/Engine/Source/Programs/Horde/HordeStorage/Jupiter.Common/Implementation/NamespacePolicyResolver.cs
+++ b/Engine/Source/Programs/Horde/HordeStorage/Jupiter.Common/Implementation/NamespacePolicyResolver.cs
@@ -50,7 +50,7 @@ namespace Jupiter.Common
                 return defaultSettings;
             }
 
-            throw new ArgumentException($"Unable to find a valid policy for namespace {ns}");
+            throw new UnknownNamespaceException($"Unable to find a valid policy for namespace {ns}");
         }
     }
 }
