Horde: Normalize namespace names to match assembly names.

diff --git a/Engine/Source/Programs/Horde/Horde.Agent.Tests/Helper.cs b/Engine/Source/Programs/Horde/Horde.Agent.Tests/Helper.cs
--- a/Engine/Source/Programs/Horde/Horde.Agent.Tests/Helper.cs
+++ b/Engine/Source/Programs/Horde/Horde.Agent.Tests/Helper.cs
@@ -7,14 +7,14 @@ using System.Threading.Tasks;
 using Google.Protobuf.WellKnownTypes;
 using Grpc.Core;
 using Grpc.Net.Client;
-using HordeAgent.Execution.Interfaces;
-using HordeAgent.Utility;
+using Horde.Agent.Execution.Interfaces;
+using Horde.Agent.Utility;
 using HordeCommon;
 using HordeCommon.Rpc;
 using Microsoft.Extensions.Logging;
 using Microsoft.Extensions.Options;
 
-namespace HordeAgentTests
+namespace Horde.Agent.Tests
 {
 	// Stub for fulfilling IOptionsMonitor interface during testing
 	// Copied from HordeServerTests until a good way to share code between these is decided.
diff --git a/Engine/Source/Programs/Horde/Horde.Agent.Tests/LogParserTests.cs b/Engine/Source/Programs/Horde/Horde.Agent.Tests/LogParserTests.cs
--- a/Engine/Source/Programs/Horde/Horde.Agent.Tests/LogParserTests.cs
+++ b/Engine/Source/Programs/Horde/Horde.Agent.Tests/LogParserTests.cs
@@ -1,8 +1,8 @@
 // Copyright Epic Games, Inc. All Rights Reserved.
 
 using EpicGames.Core;
-using HordeAgent.Parser;
-using HordeAgent.Utility;
+using Horde.Agent.Parser;
+using Horde.Agent.Utility;
 using HordeCommon;
 using Microsoft.Extensions.Logging;
 using Microsoft.Extensions.Logging.Abstractions;
@@ -20,7 +20,7 @@ using System.Text;
 using System.Text.Json;
 using System.Text.RegularExpressions;
 
-namespace HordeAgentTests
+namespace Horde.Agent.Tests
 {
 	[TestClass]
 	public class LogParserTests
diff --git a/Engine/Source/Programs/Horde/Horde.Agent.Tests/RpcConnectionTest.cs b/Engine/Source/Programs/Horde/Horde.Agent.Tests/RpcConnectionTest.cs
--- a/Engine/Source/Programs/Horde/Horde.Agent.Tests/RpcConnectionTest.cs
+++ b/Engine/Source/Programs/Horde/Horde.Agent.Tests/RpcConnectionTest.cs
@@ -5,14 +5,14 @@ using System.Threading;
 using System.Threading.Tasks;
 using Grpc.Core;
 using Grpc.Net.Client;
-using HordeAgent.Utility;
+using Horde.Agent.Utility;
 using HordeCommon.Rpc;
 using Microsoft.Extensions.Logging;
 using Microsoft.Extensions.Logging.Console;
 using Microsoft.VisualStudio.TestTools.UnitTesting;
 using Moq;
 
-namespace HordeAgentTests
+namespace Horde.Agent.Tests
 {
 #if false
 	[TestClass]
diff --git a/Engine/Source/Programs/Horde/Horde.Agent.Tests/WorkerServiceTest.cs b/Engine/Source/Programs/Horde/Horde.Agent.Tests/WorkerServiceTest.cs
--- a/Engine/Source/Programs/Horde/Horde.Agent.Tests/WorkerServiceTest.cs
+++ b/Engine/Source/Programs/Horde/Horde.Agent.Tests/WorkerServiceTest.cs
@@ -5,10 +5,10 @@ using System.IO;
 using System.Threading;
 using System.Threading.Tasks;
 using Grpc.Core;
-using HordeAgent;
-using HordeAgent.Execution.Interfaces;
-using HordeAgent.Services;
-using HordeAgent.Utility;
+using Horde.Agent;
+using Horde.Agent.Execution.Interfaces;
+using Horde.Agent.Services;
+using Horde.Agent.Utility;
 using HordeCommon;
 using HordeCommon.Rpc;
 using HordeCommon.Rpc.Messages;
@@ -18,7 +18,7 @@ using Microsoft.Extensions.Logging.Console;
 using Microsoft.Extensions.Options;
 using Microsoft.VisualStudio.TestTools.UnitTesting;
 
-namespace HordeAgentTests
+namespace Horde.Agent.Tests
 {
 	[TestClass]
 	public class WorkerServiceTest
diff --git a/Engine/Source/Programs/Horde/Horde.Agent/AgentSettings.cs b/Engine/Source/Programs/Horde/Horde.Agent/AgentSettings.cs
--- a/Engine/Source/Programs/Horde/Horde.Agent/AgentSettings.cs
+++ b/Engine/Source/Programs/Horde/Horde.Agent/AgentSettings.cs
@@ -16,7 +16,7 @@ using System.Text;
 using System.Text.Json;
 using System.Threading.Tasks;
 
-namespace HordeAgent
+namespace Horde.Agent
 {
 	/// <summary>
 	/// The type of executor to use
diff --git a/Engine/Source/Programs/Horde/Horde.Agent/Commands/Bundles/CreateCommand.cs b/Engine/Source/Programs/Horde/Horde.Agent/Commands/Bundles/CreateCommand.cs
--- a/Engine/Source/Programs/Horde/Horde.Agent/Commands/Bundles/CreateCommand.cs
+++ b/Engine/Source/Programs/Horde/Horde.Agent/Commands/Bundles/CreateCommand.cs
@@ -16,7 +16,7 @@ using System.IO;
 using System.Linq;
 using System.Threading.Tasks;
 
-namespace HordeAgent.Commands.Bundles
+namespace Horde.Agent.Commands.Bundles
 {
 	abstract class BundleCommandBase : Command
 	{
diff --git a/Engine/Source/Programs/Horde/Horde.Agent/Commands/Bundles/DeltaCommand.cs b/Engine/Source/Programs/Horde/Horde.Agent/Commands/Bundles/DeltaCommand.cs
--- a/Engine/Source/Programs/Horde/Horde.Agent/Commands/Bundles/DeltaCommand.cs
+++ b/Engine/Source/Programs/Horde/Horde.Agent/Commands/Bundles/DeltaCommand.cs
@@ -15,7 +15,7 @@ using System.Linq;
 using System.Text;
 using System.Threading.Tasks;
 
-namespace HordeAgent.Commands.Bundles
+namespace Horde.Agent.Commands.Bundles
 {
 	[Command("bundle", "delta", "Desribe the delta between two bundles")]
 	class DeltaCommand : BundleCommandBase
diff --git a/Engine/Source/Programs/Horde/Horde.Agent/Commands/Compute/ComputeCommand.cs b/Engine/Source/Programs/Horde/Horde.Agent/Commands/Compute/ComputeCommand.cs
--- a/Engine/Source/Programs/Horde/Horde.Agent/Commands/Compute/ComputeCommand.cs
+++ b/Engine/Source/Programs/Horde/Horde.Agent/Commands/Compute/ComputeCommand.cs
@@ -18,7 +18,7 @@ using EpicGames.Serialization;
 using EpicGames.Horde.Compute;
 using EpicGames.Horde.Storage;
 
-namespace HordeAgent.Commands
+namespace Horde.Agent.Commands
 {
 	/// <summary>
 	/// Installs the agent as a service
diff --git a/Engine/Source/Programs/Horde/Horde.Agent/Commands/Parse/ParseCommand.cs b/Engine/Source/Programs/Horde/Horde.Agent/Commands/Parse/ParseCommand.cs
--- a/Engine/Source/Programs/Horde/Horde.Agent/Commands/Parse/ParseCommand.cs
+++ b/Engine/Source/Programs/Horde/Horde.Agent/Commands/Parse/ParseCommand.cs
@@ -1,7 +1,7 @@
 // Copyright Epic Games, Inc. All Rights Reserved.
 
 using EpicGames.Core;
-using HordeAgent.Parser;
+using Horde.Agent.Parser;
 using Microsoft.Extensions.Logging;
 using System;
 using System.Collections.Generic;
@@ -10,7 +10,7 @@ using System.IO;
 using System.Text;
 using System.Threading.Tasks;
 
-namespace HordeAgent.Commands.Parse
+namespace Horde.Agent.Commands.Parse
 {
 	/// <summary>
 	/// Installs the agent as a service
diff --git a/Engine/Source/Programs/Horde/Horde.Agent/Commands/Service/InstallCommand.cs b/Engine/Source/Programs/Horde/Horde.Agent/Commands/Service/InstallCommand.cs
--- a/Engine/Source/Programs/Horde/Horde.Agent/Commands/Service/InstallCommand.cs
+++ b/Engine/Source/Programs/Horde/Horde.Agent/Commands/Service/InstallCommand.cs
@@ -1,6 +1,6 @@
 // Copyright Epic Games, Inc. All Rights Reserved.
 
-using HordeAgent.Utility;
+using Horde.Agent.Utility;
 using Microsoft.Extensions.Logging;
 using System;
 using System.Collections.Generic;
@@ -9,7 +9,7 @@ using System.Text;
 using System.Threading.Tasks;
 using EpicGames.Core;
 
-namespace HordeAgent.Commands.Service
+namespace Horde.Agent.Commands.Service
 {
 	/// <summary>
 	/// Installs the agent as a service
diff --git a/Engine/Source/Programs/Horde/Horde.Agent/Commands/Service/RunCommand.cs b/Engine/Source/Programs/Horde/Horde.Agent/Commands/Service/RunCommand.cs
--- a/Engine/Source/Programs/Horde/Horde.Agent/Commands/Service/RunCommand.cs
+++ b/Engine/Source/Programs/Horde/Horde.Agent/Commands/Service/RunCommand.cs
@@ -1,8 +1,8 @@
 // Copyright Epic Games, Inc. All Rights Reserved.
 
-using HordeAgent.Commands;
-using HordeAgent.Services;
-using HordeAgent.Utility;
+using Horde.Agent.Commands;
+using Horde.Agent.Services;
+using Horde.Agent.Utility;
 using Microsoft.Extensions.Configuration;
 using Microsoft.Extensions.DependencyInjection;
 using Microsoft.Extensions.Hosting;
@@ -27,7 +27,7 @@ using EpicGames.Horde.Storage;
 using EpicGames.Horde.Storage.Impl;
 using EpicGames.Horde.Compute;
 
-namespace HordeAgent.Modes.Service
+namespace Horde.Agent.Modes.Service
 {
 	/// <summary>
 	/// 
diff --git a/Engine/Source/Programs/Horde/Horde.Agent/Commands/Service/UninstallCommand.cs b/Engine/Source/Programs/Horde/Horde.Agent/Commands/Service/UninstallCommand.cs
--- a/Engine/Source/Programs/Horde/Horde.Agent/Commands/Service/UninstallCommand.cs
+++ b/Engine/Source/Programs/Horde/Horde.Agent/Commands/Service/UninstallCommand.cs
@@ -1,7 +1,7 @@
 // Copyright Epic Games, Inc. All Rights Reserved.
 
 using EpicGames.Core;
-using HordeAgent.Utility;
+using Horde.Agent.Utility;
 using Microsoft.Extensions.Logging;
 using System;
 using System.Collections.Generic;
@@ -9,7 +9,7 @@ using System.Reflection;
 using System.Text;
 using System.Threading.Tasks;
 
-namespace HordeAgent.Commands.Service
+namespace Horde.Agent.Commands.Service
 {
 	/// <summary>
 	/// Uninstalls the service
diff --git a/Engine/Source/Programs/Horde/Horde.Agent/Commands/Service/UpgradeCommand.cs b/Engine/Source/Programs/Horde/Horde.Agent/Commands/Service/UpgradeCommand.cs
--- a/Engine/Source/Programs/Horde/Horde.Agent/Commands/Service/UpgradeCommand.cs
+++ b/Engine/Source/Programs/Horde/Horde.Agent/Commands/Service/UpgradeCommand.cs
@@ -1,7 +1,7 @@
 // Copyright Epic Games, Inc. All Rights Reserved.
 
-using HordeAgent.Services;
-using HordeAgent.Utility;
+using Horde.Agent.Services;
+using Horde.Agent.Utility;
 using Microsoft.Extensions.Logging;
 using System;
 using System.Collections.Generic;
@@ -18,7 +18,7 @@ using System.Runtime.InteropServices;
 using System.Text.RegularExpressions;
 using System.Threading;
 
-namespace HordeAgent.Commands
+namespace Horde.Agent.Commands
 {
 	/// <summary>
 	/// Upgrades a running service to the current application
diff --git a/Engine/Source/Programs/Horde/Horde.Agent/Commands/Utilities/CapsCommand.cs b/Engine/Source/Programs/Horde/Horde.Agent/Commands/Utilities/CapsCommand.cs
--- a/Engine/Source/Programs/Horde/Horde.Agent/Commands/Utilities/CapsCommand.cs
+++ b/Engine/Source/Programs/Horde/Horde.Agent/Commands/Utilities/CapsCommand.cs
@@ -1,7 +1,7 @@
 // Copyright Epic Games, Inc. All Rights Reserved.
 
 using Grpc.Net.Client;
-using HordeAgent.Utility;
+using Horde.Agent.Utility;
 using HordeCommon.Rpc;
 using Microsoft.Extensions.Configuration;
 using Microsoft.Extensions.Logging;
@@ -17,9 +17,9 @@ using System.Text;
 using System.Threading.Tasks;
 using EpicGames.Core;
 using HordeCommon.Rpc.Messages;
-using HordeAgent.Services;
+using Horde.Agent.Services;
 
-namespace HordeAgent.Commands.Utilities
+namespace Horde.Agent.Commands.Utilities
 {
 	/// <summary>
 	/// Shows capabilities of this agent
diff --git a/Engine/Source/Programs/Horde/Horde.Agent/Commands/Utilities/CreateCertCommand.cs b/Engine/Source/Programs/Horde/Horde.Agent/Commands/Utilities/CreateCertCommand.cs
--- a/Engine/Source/Programs/Horde/Horde.Agent/Commands/Utilities/CreateCertCommand.cs
+++ b/Engine/Source/Programs/Horde/Horde.Agent/Commands/Utilities/CreateCertCommand.cs
@@ -16,7 +16,7 @@ using System.Text.RegularExpressions;
 using System.Threading.Tasks;
 using EpicGames.Core;
 
-namespace HordeAgent.Commands.Certs
+namespace Horde.Agent.Commands.Certs
 {
 	/// <summary>
 	/// Creates a certificate that can be used for server/agent SSL connections
diff --git a/Engine/Source/Programs/Horde/Horde.Agent/Commands/Utilities/PublishCommand.cs b/Engine/Source/Programs/Horde/Horde.Agent/Commands/Utilities/PublishCommand.cs
--- a/Engine/Source/Programs/Horde/Horde.Agent/Commands/Utilities/PublishCommand.cs
+++ b/Engine/Source/Programs/Horde/Horde.Agent/Commands/Utilities/PublishCommand.cs
@@ -1,7 +1,7 @@
 // Copyright Epic Games, Inc. All Rights Reserved.
 
 using Grpc.Net.Client;
-using HordeAgent.Utility;
+using Horde.Agent.Utility;
 using HordeCommon.Rpc;
 using Microsoft.Extensions.Configuration;
 using Microsoft.Extensions.Logging;
@@ -17,7 +17,7 @@ using System.Text;
 using System.Threading.Tasks;
 using EpicGames.Core;
 
-namespace HordeAgent.Commands.Utilities
+namespace Horde.Agent.Commands.Utilities
 {
 	/// <summary>
 	/// Publishes the contents of this application directory to the server
diff --git a/Engine/Source/Programs/Horde/Horde.Agent/Commands/Workspace/CleanCommand.cs b/Engine/Source/Programs/Horde/Horde.Agent/Commands/Workspace/CleanCommand.cs
--- a/Engine/Source/Programs/Horde/Horde.Agent/Commands/Workspace/CleanCommand.cs
+++ b/Engine/Source/Programs/Horde/Horde.Agent/Commands/Workspace/CleanCommand.cs
@@ -12,7 +12,7 @@ using EpicGames.Core;
 using EpicGames.Perforce.Managed;
 using EpicGames.Perforce;
 
-namespace HordeAgent.Commands.Workspace
+namespace Horde.Agent.Commands.Workspace
 {
 	[Command("Workspace", "Clean", "Cleans all modified files from the workspace.")]
 	class CleanCommand : WorkspaceCommand
diff --git a/Engine/Source/Programs/Horde/Horde.Agent/Commands/Workspace/ClearCommand.cs b/Engine/Source/Programs/Horde/Horde.Agent/Commands/Workspace/ClearCommand.cs
--- a/Engine/Source/Programs/Horde/Horde.Agent/Commands/Workspace/ClearCommand.cs
+++ b/Engine/Source/Programs/Horde/Horde.Agent/Commands/Workspace/ClearCommand.cs
@@ -11,7 +11,7 @@ using System.Text;
 using System.Threading;
 using System.Threading.Tasks;
 
-namespace HordeAgent.Commands.Workspace
+namespace Horde.Agent.Commands.Workspace
 {
 	[Command("Workspace", "Clear", "Empties the staging directory of any files, returning them to the cache")]
 	class ClearCommand : WorkspaceCommand
diff --git a/Engine/Source/Programs/Horde/Horde.Agent/Commands/Workspace/WorkspaceCommand.cs b/Engine/Source/Programs/Horde/Horde.Agent/Commands/Workspace/WorkspaceCommand.cs
--- a/Engine/Source/Programs/Horde/Horde.Agent/Commands/Workspace/WorkspaceCommand.cs
+++ b/Engine/Source/Programs/Horde/Horde.Agent/Commands/Workspace/WorkspaceCommand.cs
@@ -1,6 +1,6 @@
 // Copyright Epic Games, Inc. All Rights Reserved.
 
-using HordeAgent.Utility;
+using Horde.Agent.Utility;
 using Microsoft.Extensions.Logging;
 using System;
 using System.Collections.Generic;
@@ -15,7 +15,7 @@ using EpicGames.Core;
 using EpicGames.Perforce;
 using EpicGames.Perforce.Managed;
 
-namespace HordeAgent.Commands.Workspace
+namespace Horde.Agent.Commands.Workspace
 {
 	abstract class WorkspaceCommand : Command
 	{
diff --git a/Engine/Source/Programs/Horde/Horde.Agent/Execution/BuildGraphExecutor.cs b/Engine/Source/Programs/Horde/Horde.Agent/Execution/BuildGraphExecutor.cs
--- a/Engine/Source/Programs/Horde/Horde.Agent/Execution/BuildGraphExecutor.cs
+++ b/Engine/Source/Programs/Horde/Horde.Agent/Execution/BuildGraphExecutor.cs
@@ -2,9 +2,9 @@
 
 using EpicGames.Core;
 using Grpc.Core;
-using HordeAgent.Execution.Interfaces;
-using HordeAgent.Parser;
-using HordeAgent.Utility;
+using Horde.Agent.Execution.Interfaces;
+using Horde.Agent.Parser;
+using Horde.Agent.Utility;
 using HordeCommon;
 using HordeCommon.Rpc;
 using Microsoft.Extensions.Logging;
@@ -26,7 +26,7 @@ using System.Text.RegularExpressions;
 using System.Threading;
 using System.Threading.Tasks;
 
-namespace HordeAgent.Execution
+namespace Horde.Agent.Execution
 {
 	abstract class BuildGraphExecutor : IExecutor
 	{
diff --git a/Engine/Source/Programs/Horde/Horde.Agent/Execution/ComputeTaskExecutor.cs b/Engine/Source/Programs/Horde/Horde.Agent/Execution/ComputeTaskExecutor.cs
--- a/Engine/Source/Programs/Horde/Horde.Agent/Execution/ComputeTaskExecutor.cs
+++ b/Engine/Source/Programs/Horde/Horde.Agent/Execution/ComputeTaskExecutor.cs
@@ -18,7 +18,7 @@ using Google.Protobuf.WellKnownTypes;
 using Microsoft.Extensions.Hosting;
 using System.Collections.Concurrent;
 using Grpc.Net.Client;
-using HordeAgent.Utility;
+using Horde.Agent.Utility;
 using HordeCommon.Rpc.Tasks;
 using Microsoft.Extensions.Logging;
 using HordeCommon.Rpc;
@@ -26,7 +26,7 @@ using EpicGames.Serialization;
 using EpicGames.Horde.Compute;
 using EpicGames.Horde.Common;
 
-namespace HordeAgent
+namespace Horde.Agent.Execution
 {
 	/// <summary>
 	/// Executes remote actions in a sandbox
diff --git a/Engine/Source/Programs/Horde/Horde.Agent/Execution/Interfaces/IExecutor.cs b/Engine/Source/Programs/Horde/Horde.Agent/Execution/Interfaces/IExecutor.cs
--- a/Engine/Source/Programs/Horde/Horde.Agent/Execution/Interfaces/IExecutor.cs
+++ b/Engine/Source/Programs/Horde/Horde.Agent/Execution/Interfaces/IExecutor.cs
@@ -1,7 +1,7 @@
 // Copyright Epic Games, Inc. All Rights Reserved.
 
-using HordeAgent.Parser.Interfaces;
-using HordeAgent.Utility;
+using Horde.Agent.Parser.Interfaces;
+using Horde.Agent.Utility;
 using HordeCommon;
 using HordeCommon.Rpc;
 using Microsoft.Extensions.Logging;
@@ -11,7 +11,7 @@ using System.Text;
 using System.Threading;
 using System.Threading.Tasks;
 
-namespace HordeAgent.Execution.Interfaces
+namespace Horde.Agent.Execution.Interfaces
 {
 	interface IExecutor
 	{
diff --git a/Engine/Source/Programs/Horde/Horde.Agent/Execution/LocalExecutor.cs b/Engine/Source/Programs/Horde/Horde.Agent/Execution/LocalExecutor.cs
--- a/Engine/Source/Programs/Horde/Horde.Agent/Execution/LocalExecutor.cs
+++ b/Engine/Source/Programs/Horde/Horde.Agent/Execution/LocalExecutor.cs
@@ -1,7 +1,7 @@
 // Copyright Epic Games, Inc. All Rights Reserved.
 
 using EpicGames.Core;
-using HordeAgent.Utility;
+using Horde.Agent.Utility;
 using HordeCommon;
 using HordeCommon.Rpc;
 using Microsoft.Extensions.Logging;
@@ -11,7 +11,7 @@ using System.Text;
 using System.Threading;
 using System.Threading.Tasks;
 
-namespace HordeAgent.Execution
+namespace Horde.Agent.Execution
 {
 	class LocalExecutor : BuildGraphExecutor
 	{
diff --git a/Engine/Source/Programs/Horde/Horde.Agent/Execution/PerforceExecutor.cs b/Engine/Source/Programs/Horde/Horde.Agent/Execution/PerforceExecutor.cs
--- a/Engine/Source/Programs/Horde/Horde.Agent/Execution/PerforceExecutor.cs
+++ b/Engine/Source/Programs/Horde/Horde.Agent/Execution/PerforceExecutor.cs
@@ -1,10 +1,10 @@
 // Copyright Epic Games, Inc. All Rights Reserved.
 
 using EpicGames.Perforce.Managed;
-using HordeAgent.Execution.Interfaces;
-using HordeAgent.Parser;
-using HordeAgent.Parser.Interfaces;
-using HordeAgent.Utility;
+using Horde.Agent.Execution.Interfaces;
+using Horde.Agent.Parser;
+using Horde.Agent.Parser.Interfaces;
+using Horde.Agent.Utility;
 using HordeCommon;
 using HordeCommon.Rpc;
 using HordeCommon.Rpc.Messages;
@@ -30,7 +30,7 @@ using Google.Protobuf.WellKnownTypes;
 using OpenTracing;
 using OpenTracing.Util;
 
-namespace HordeAgent.Execution
+namespace Horde.Agent.Execution
 {
 	class PerforceExecutor : BuildGraphExecutor
 	{
diff --git a/Engine/Source/Programs/Horde/Horde.Agent/Execution/TestExecutor.cs b/Engine/Source/Programs/Horde/Horde.Agent/Execution/TestExecutor.cs
--- a/Engine/Source/Programs/Horde/Horde.Agent/Execution/TestExecutor.cs
+++ b/Engine/Source/Programs/Horde/Horde.Agent/Execution/TestExecutor.cs
@@ -11,15 +11,15 @@ using System.Text.Json;
 using System.Threading;
 using System.Threading.Tasks;
 using EpicGames.Core;
-using HordeAgent.Execution.Interfaces;
-using HordeAgent.Parser;
-using HordeAgent.Parser.Interfaces;
-using HordeAgent.Utility;
+using Horde.Agent.Execution.Interfaces;
+using Horde.Agent.Parser;
+using Horde.Agent.Parser.Interfaces;
+using Horde.Agent.Utility;
 using HordeCommon;
 using HordeCommon.Rpc;
 using Microsoft.Extensions.Logging;
 
-namespace HordeAgent.Execution
+namespace Horde.Agent.Execution
 {
 	class TestExecutor : BuildGraphExecutor
 	{
diff --git a/Engine/Source/Programs/Horde/Horde.Agent/Logging.cs b/Engine/Source/Programs/Horde/Horde.Agent/Logging.cs
--- a/Engine/Source/Programs/Horde/Horde.Agent/Logging.cs
+++ b/Engine/Source/Programs/Horde/Horde.Agent/Logging.cs
@@ -20,7 +20,7 @@ using System.Text.Json;
 using System.Buffers;
 using MessageTemplate = EpicGames.Core.MessageTemplate;
 
-namespace HordeAgent
+namespace Horde.Agent
 {
 	static class Logging
 	{
diff --git a/Engine/Source/Programs/Horde/Horde.Agent/Parser/Interfaces/ILogContext.cs b/Engine/Source/Programs/Horde/Horde.Agent/Parser/Interfaces/ILogContext.cs
--- a/Engine/Source/Programs/Horde/Horde.Agent/Parser/Interfaces/ILogContext.cs
+++ b/Engine/Source/Programs/Horde/Horde.Agent/Parser/Interfaces/ILogContext.cs
@@ -6,7 +6,7 @@ using System;
 using System.Collections.Generic;
 using System.Text;
 
-namespace HordeAgent.Parser.Interfaces
+namespace Horde.Agent.Parser.Interfaces
 {
 	/// <summary>
 	/// Interface for context on the log being parsed
diff --git a/Engine/Source/Programs/Horde/Horde.Agent/Parser/LogEventMarkup.cs b/Engine/Source/Programs/Horde/Horde.Agent/Parser/LogEventMarkup.cs
--- a/Engine/Source/Programs/Horde/Horde.Agent/Parser/LogEventMarkup.cs
+++ b/Engine/Source/Programs/Horde/Horde.Agent/Parser/LogEventMarkup.cs
@@ -1,7 +1,7 @@
 // Copyright Epic Games, Inc. All Rights Reserved.
 
 using EpicGames.Core;
-using HordeAgent.Parser.Interfaces;
+using Horde.Agent.Parser.Interfaces;
 using System;
 using System.Collections.Generic;
 using System.IO;
@@ -10,7 +10,7 @@ using System.Text;
 using System.Text.Json;
 using System.Text.RegularExpressions;
 
-namespace HordeAgent.Parser
+namespace Horde.Agent.Parser
 {
 	/// <summary>
 	/// Extension methods to allow adding markup to log event spans
diff --git a/Engine/Source/Programs/Horde/Horde.Agent/Parser/LogParser.cs b/Engine/Source/Programs/Horde/Horde.Agent/Parser/LogParser.cs
--- a/Engine/Source/Programs/Horde/Horde.Agent/Parser/LogParser.cs
+++ b/Engine/Source/Programs/Horde/Horde.Agent/Parser/LogParser.cs
@@ -1,8 +1,8 @@
 // Copyright Epic Games, Inc. All Rights Reserved.
 
 using EpicGames.Core;
-using HordeAgent.Parser.Interfaces;
-using HordeAgent.Utility;
+using Horde.Agent.Parser.Interfaces;
+using Horde.Agent.Utility;
 using HordeCommon;
 using Microsoft.Extensions.Logging;
 using System;
@@ -18,7 +18,7 @@ using System.Text.RegularExpressions;
 using System.Threading;
 using System.Threading.Tasks;
 
-namespace HordeAgent.Parser
+namespace Horde.Agent.Parser
 {
 	/// <summary>
 	/// Concrete implementation of <see cref="ILogContext"/>
diff --git a/Engine/Source/Programs/Horde/Horde.Agent/Parser/Matchers/CompileEventMatcher.cs b/Engine/Source/Programs/Horde/Horde.Agent/Parser/Matchers/CompileEventMatcher.cs
--- a/Engine/Source/Programs/Horde/Horde.Agent/Parser/Matchers/CompileEventMatcher.cs
+++ b/Engine/Source/Programs/Horde/Horde.Agent/Parser/Matchers/CompileEventMatcher.cs
@@ -1,8 +1,8 @@
 // Copyright Epic Games, Inc. All Rights Reserved.
 
 using EpicGames.Core;
-using HordeAgent.Parser;
-using HordeAgent.Parser.Interfaces;
+using Horde.Agent.Parser;
+using Horde.Agent.Parser.Interfaces;
 using HordeCommon;
 using Microsoft.Extensions.Logging;
 using System;
@@ -14,7 +14,7 @@ using System.Text;
 using System.Text.RegularExpressions;
 using System.Threading.Tasks;
 
-namespace HordeAgent.Parser.Matchers
+namespace Horde.Agent.Parser.Matchers
 {
 	/// <summary>
 	/// Matches compile errors and annotates with the source file path and revision
diff --git a/Engine/Source/Programs/Horde/Horde.Agent/Parser/Matchers/ContentEventMatcher.cs b/Engine/Source/Programs/Horde/Horde.Agent/Parser/Matchers/ContentEventMatcher.cs
--- a/Engine/Source/Programs/Horde/Horde.Agent/Parser/Matchers/ContentEventMatcher.cs
+++ b/Engine/Source/Programs/Horde/Horde.Agent/Parser/Matchers/ContentEventMatcher.cs
@@ -1,7 +1,7 @@
 // Copyright Epic Games, Inc. All Rights Reserved.
 
 using EpicGames.Core;
-using HordeAgent.Parser.Interfaces;
+using Horde.Agent.Parser.Interfaces;
 using HordeCommon;
 using Microsoft.Extensions.Logging;
 using System;
@@ -9,7 +9,7 @@ using System.Collections.Generic;
 using System.Text;
 using System.Text.RegularExpressions;
 
-namespace HordeAgent.Parser.Matchers
+namespace Horde.Agent.Parser.Matchers
 {
 	/// <summary>
 	/// Matches compile errors and annotates with the source file path and revision
diff --git a/Engine/Source/Programs/Horde/Horde.Agent/Parser/Matchers/CrashEventMatcher.cs b/Engine/Source/Programs/Horde/Horde.Agent/Parser/Matchers/CrashEventMatcher.cs
--- a/Engine/Source/Programs/Horde/Horde.Agent/Parser/Matchers/CrashEventMatcher.cs
+++ b/Engine/Source/Programs/Horde/Horde.Agent/Parser/Matchers/CrashEventMatcher.cs
@@ -1,8 +1,8 @@
 // Copyright Epic Games, Inc. All Rights Reserved.
 
 using EpicGames.Core;
-using HordeAgent.Parser;
-using HordeAgent.Parser.Interfaces;
+using Horde.Agent.Parser;
+using Horde.Agent.Parser.Interfaces;
 using HordeCommon;
 using Microsoft.Extensions.Logging;
 using System;
@@ -12,7 +12,7 @@ using System.Text;
 using System.Text.RegularExpressions;
 using System.Threading.Tasks;
 
-namespace HordeAgent.Parser.Matchers
+namespace Horde.Agent.Parser.Matchers
 {
 	/// <summary>
 	/// Matcher for engine crashes
diff --git a/Engine/Source/Programs/Horde/Horde.Agent/Parser/Matchers/ExitCodeEventMatcher.cs b/Engine/Source/Programs/Horde/Horde.Agent/Parser/Matchers/ExitCodeEventMatcher.cs
--- a/Engine/Source/Programs/Horde/Horde.Agent/Parser/Matchers/ExitCodeEventMatcher.cs
+++ b/Engine/Source/Programs/Horde/Horde.Agent/Parser/Matchers/ExitCodeEventMatcher.cs
@@ -1,11 +1,11 @@
 // Copyright Epic Games, Inc. All Rights Reserved.
 
 using EpicGames.Core;
-using HordeAgent.Parser.Interfaces;
+using Horde.Agent.Parser.Interfaces;
 using HordeCommon;
 using Microsoft.Extensions.Logging;
 
-namespace HordeAgent.Parser.Matchers
+namespace Horde.Agent.Parser.Matchers
 {
 	/// <summary>
 	/// Matcher for editor/UAT instances exiting with an error
diff --git a/Engine/Source/Programs/Horde/Horde.Agent/Program.cs b/Engine/Source/Programs/Horde/Horde.Agent/Program.cs
--- a/Engine/Source/Programs/Horde/Horde.Agent/Program.cs
+++ b/Engine/Source/Programs/Horde/Horde.Agent/Program.cs
@@ -1,8 +1,8 @@
 // Copyright Epic Games, Inc. All Rights Reserved.
 
-using HordeAgent.Commands;
-using HordeAgent.Services;
-using HordeAgent.Utility;
+using Horde.Agent.Commands;
+using Horde.Agent.Services;
+using Horde.Agent.Utility;
 using Microsoft.Extensions.Configuration;
 using Microsoft.Extensions.DependencyInjection;
 using Microsoft.Extensions.Hosting;
@@ -30,7 +30,7 @@ using System.Text;
 using EpicGames.Perforce;
 using System.Security.Cryptography;
 
-namespace HordeAgent
+namespace Horde.Agent
 {
 	/// <summary>
 	/// Entry point
@@ -91,7 +91,7 @@ namespace HordeAgent
 
 			// Execute all the commands
 			IServiceProvider ServiceProvider = Services.BuildServiceProvider();
-			return await CommandHost.RunAsync(new CommandLineArguments(Args), ServiceProvider, typeof(HordeAgent.Modes.Service.RunCommand));
+			return await CommandHost.RunAsync(new CommandLineArguments(Args), ServiceProvider, typeof(Horde.Agent.Modes.Service.RunCommand));
 		}
 
 		/// <summary>
diff --git a/Engine/Source/Programs/Horde/Horde.Agent/Services/GrpcService.cs b/Engine/Source/Programs/Horde/Horde.Agent/Services/GrpcService.cs
--- a/Engine/Source/Programs/Horde/Horde.Agent/Services/GrpcService.cs
+++ b/Engine/Source/Programs/Horde/Horde.Agent/Services/GrpcService.cs
@@ -1,7 +1,7 @@
 // Copyright Epic Games, Inc. All Rights Reserved.
 
 using Grpc.Net.Client;
-using HordeAgent.Utility;
+using Horde.Agent.Utility;
 using Microsoft.Extensions.Http;
 using Microsoft.Extensions.Logging;
 using Microsoft.Extensions.Options;
@@ -15,7 +15,7 @@ using System.Text;
 using Grpc.Core;
 using HordeCommon.Rpc;
 
-namespace HordeAgent.Services
+namespace Horde.Agent.Services
 {
 	/// <summary>
 	/// Service which creates a configured Grpc channel
diff --git a/Engine/Source/Programs/Horde/Horde.Agent/Services/WorkerService.cs b/Engine/Source/Programs/Horde/Horde.Agent/Services/WorkerService.cs
--- a/Engine/Source/Programs/Horde/Horde.Agent/Services/WorkerService.cs
+++ b/Engine/Source/Programs/Horde/Horde.Agent/Services/WorkerService.cs
@@ -8,10 +8,10 @@ using Google.Protobuf;
 using Google.Protobuf.WellKnownTypes;
 using Grpc.Core;
 using Grpc.Net.Client;
-using HordeAgent.Execution;
-using HordeAgent.Execution.Interfaces;
-using HordeAgent.Parser;
-using HordeAgent.Utility;
+using Horde.Agent.Execution;
+using Horde.Agent.Execution.Interfaces;
+using Horde.Agent.Parser;
+using Horde.Agent.Utility;
 using HordeCommon;
 using HordeCommon.Rpc;
 using HordeCommon.Rpc.Messages;
@@ -44,7 +44,7 @@ using System.Net;
 using EpicGames.Horde.Storage.Impl;
 using System.Text.RegularExpressions;
 
-namespace HordeAgent.Services
+namespace Horde.Agent.Services
 {
 	/// <summary>
 	/// Implements the message handling loop for an agent. Runs asynchronously until disposed.
diff --git a/Engine/Source/Programs/Horde/Horde.Agent/Utility/ArtifactUploader.cs b/Engine/Source/Programs/Horde/Horde.Agent/Utility/ArtifactUploader.cs
--- a/Engine/Source/Programs/Horde/Horde.Agent/Utility/ArtifactUploader.cs
+++ b/Engine/Source/Programs/Horde/Horde.Agent/Utility/ArtifactUploader.cs
@@ -15,7 +15,7 @@ using System.Text;
 using System.Threading;
 using System.Threading.Tasks;
 
-namespace HordeAgent.Utility
+namespace Horde.Agent.Utility
 {
 	/// <summary>
 	/// Helper class for uploading artifacts to the server
diff --git a/Engine/Source/Programs/Horde/Horde.Agent/Utility/CertificateHelper.cs b/Engine/Source/Programs/Horde/Horde.Agent/Utility/CertificateHelper.cs
--- a/Engine/Source/Programs/Horde/Horde.Agent/Utility/CertificateHelper.cs
+++ b/Engine/Source/Programs/Horde/Horde.Agent/Utility/CertificateHelper.cs
@@ -10,7 +10,7 @@ using System.Security.Cryptography.X509Certificates;
 using System.Text;
 using EpicGames.Core;
 
-namespace HordeAgent.Utility
+namespace Horde.Agent.Utility
 {
 	/// <summary>
 	/// Helper functions for dealing with certificates
diff --git a/Engine/Source/Programs/Horde/Horde.Agent/Utility/JsonLogger.cs b/Engine/Source/Programs/Horde/Horde.Agent/Utility/JsonLogger.cs
--- a/Engine/Source/Programs/Horde/Horde.Agent/Utility/JsonLogger.cs
+++ b/Engine/Source/Programs/Horde/Horde.Agent/Utility/JsonLogger.cs
@@ -1,7 +1,7 @@
 // Copyright Epic Games, Inc. All Rights Reserved.
 
 using EpicGames.Core;
-using HordeAgent.Parser;
+using Horde.Agent.Parser;
 using HordeCommon;
 using Microsoft.Extensions.Logging;
 using System;
@@ -12,7 +12,7 @@ using System.Linq;
 using System.Text;
 using System.Text.Json;
 
-namespace HordeAgent.Utility
+namespace Horde.Agent.Utility
 {
 	/// <summary>
 	/// Class to handle uploading log data to the server in the background
diff --git a/Engine/Source/Programs/Horde/Horde.Agent/Utility/JsonRpcLogger.cs b/Engine/Source/Programs/Horde/Horde.Agent/Utility/JsonRpcLogger.cs
--- a/Engine/Source/Programs/Horde/Horde.Agent/Utility/JsonRpcLogger.cs
+++ b/Engine/Source/Programs/Horde/Horde.Agent/Utility/JsonRpcLogger.cs
@@ -1,7 +1,7 @@
 // Copyright Epic Games, Inc. All Rights Reserved.
 
-using HordeAgent.Parser.Interfaces;
-using HordeAgent.Utility;
+using Horde.Agent.Parser.Interfaces;
+using Horde.Agent.Utility;
 using HordeCommon;
 using HordeCommon.Rpc;
 using Microsoft.Extensions.Logging;
@@ -18,7 +18,7 @@ using System.Threading;
 using System.Threading.Channels;
 using System.Threading.Tasks;
 
-namespace HordeAgent.Parser
+namespace Horde.Agent.Parser
 {
 	/// <summary>
 	/// Class to handle uploading log data to the server in the background
diff --git a/Engine/Source/Programs/Horde/Horde.Build.Tests/AgentControllerTest.cs b/Engine/Source/Programs/Horde/Horde.Build.Tests/AgentControllerTest.cs
--- a/Engine/Source/Programs/Horde/Horde.Build.Tests/AgentControllerTest.cs
+++ b/Engine/Source/Programs/Horde/Horde.Build.Tests/AgentControllerTest.cs
@@ -1,13 +1,13 @@
 // Copyright Epic Games, Inc. All Rights Reserved.
 
 using System.Threading.Tasks;
-using HordeServer.Api;
+using Horde.Build.Api;
 using Microsoft.VisualStudio.TestTools.UnitTesting;
-using HordeServer.Controllers;
-using HordeServer.Models;
+using Horde.Build.Controllers;
+using Horde.Build.Models;
 using Microsoft.AspNetCore.Mvc;
 
-namespace HordeServerTests
+namespace Horde.Build.Tests
 {
 	[TestClass]
     public class AgentControllerDbTest : TestSetup
diff --git a/Engine/Source/Programs/Horde/Horde.Build.Tests/AgentServiceTest.cs b/Engine/Source/Programs/Horde/Horde.Build.Tests/AgentServiceTest.cs
--- a/Engine/Source/Programs/Horde/Horde.Build.Tests/AgentServiceTest.cs
+++ b/Engine/Source/Programs/Horde/Horde.Build.Tests/AgentServiceTest.cs
@@ -5,16 +5,16 @@ using System.Collections.Generic;
 using System.Security.Claims;
 using System.Threading.Tasks;
 using HordeCommon;
-using HordeServer.Api;
+using Horde.Build.Api;
 using Microsoft.VisualStudio.TestTools.UnitTesting;
-using HordeServer.Models;
-using HordeServer.Services;
+using Horde.Build.Models;
+using Horde.Build.Services;
 using Microsoft.AspNetCore.Mvc;
-using AgentSoftwareVersion = HordeServer.Utilities.StringId<HordeServer.Collections.IAgentSoftwareCollection>;
-using HordeServer.Utilities;
+using AgentSoftwareVersion = Horde.Build.Utilities.StringId<Horde.Build.Collections.IAgentSoftwareCollection>;
+using Horde.Build.Utilities;
 using MongoDB.Bson;
 
-namespace HordeServerTests
+namespace Horde.Build.Tests
 {
     /// <summary>
     /// Testing the agent service
diff --git a/Engine/Source/Programs/Horde/Horde.Build.Tests/ControllerIntegrationTest.cs b/Engine/Source/Programs/Horde/Horde.Build.Tests/ControllerIntegrationTest.cs
--- a/Engine/Source/Programs/Horde/Horde.Build.Tests/ControllerIntegrationTest.cs
+++ b/Engine/Source/Programs/Horde/Horde.Build.Tests/ControllerIntegrationTest.cs
@@ -6,11 +6,11 @@ using System.Net.Http;
 using System.Net.Http.Headers;
 using System.Runtime.CompilerServices;
 using System.Threading.Tasks;
-using HordeServer;
-using HordeServer.Collections;
-using HordeServer.Collections.Impl;
-using HordeServer.Jobs;
-using HordeServer.Services;
+using Horde.Build;
+using Horde.Build.Collections;
+using Horde.Build.Collections.Impl;
+using Horde.Build.Jobs;
+using Horde.Build.Services;
 using Microsoft.AspNetCore.Builder;
 using Microsoft.AspNetCore.Hosting;
 using Microsoft.AspNetCore.Mvc.Testing;
@@ -19,7 +19,7 @@ using Microsoft.Extensions.Configuration;
 using Microsoft.Extensions.DependencyInjection;
 using Microsoft.Extensions.Hosting;
 
-namespace HordeServerTests
+namespace Horde.Build.Tests
 {
 	public class TestWebApplicationFactory<TStartup> : WebApplicationFactory<TStartup> where TStartup : class
     {
diff --git a/Engine/Source/Programs/Horde/Horde.Build.Tests/DatabaseIntegrationTest.cs b/Engine/Source/Programs/Horde/Horde.Build.Tests/DatabaseIntegrationTest.cs
--- a/Engine/Source/Programs/Horde/Horde.Build.Tests/DatabaseIntegrationTest.cs
+++ b/Engine/Source/Programs/Horde/Horde.Build.Tests/DatabaseIntegrationTest.cs
@@ -4,16 +4,16 @@ using System;
 using System.Collections.Generic;
 using System.Net;
 using System.Threading.Tasks;
-using HordeServer;
-using HordeServer.Services;
-using HordeServer.Utilities;
+using Horde.Build;
+using Horde.Build.Services;
+using Horde.Build.Utilities;
 using Microsoft.Extensions.Logging;
 using Microsoft.Extensions.Options;
 using Microsoft.VisualStudio.TestTools.UnitTesting;
 using MongoDB.Driver;
 using StackExchange.Redis;
 
-namespace HordeServerTests
+namespace Horde.Build.Tests
 {
     // Stub for fulfilling IOptionsMonitor interface during testing
     public class TestOptionsMonitor<T> : IOptionsMonitor<T>, IDisposable
diff --git a/Engine/Source/Programs/Horde/Horde.Build.Tests/DatabaseRunner.cs b/Engine/Source/Programs/Horde/Horde.Build.Tests/DatabaseRunner.cs
--- a/Engine/Source/Programs/Horde/Horde.Build.Tests/DatabaseRunner.cs
+++ b/Engine/Source/Programs/Horde/Horde.Build.Tests/DatabaseRunner.cs
@@ -13,7 +13,7 @@ using System.Runtime.InteropServices;
 using Castle.Core.Internal;
 using EpicGames.Core;
 
-namespace HordeServerTests
+namespace Horde.Build.Tests
 {
 	public abstract class DatabaseRunner
 	{
diff --git a/Engine/Source/Programs/Horde/Horde.Build.Tests/DatabaseRunnerTest.cs b/Engine/Source/Programs/Horde/Horde.Build.Tests/DatabaseRunnerTest.cs
--- a/Engine/Source/Programs/Horde/Horde.Build.Tests/DatabaseRunnerTest.cs
+++ b/Engine/Source/Programs/Horde/Horde.Build.Tests/DatabaseRunnerTest.cs
@@ -3,7 +3,7 @@
 using System.Threading;
 using Microsoft.VisualStudio.TestTools.UnitTesting;
 
-namespace HordeServerTests
+namespace Horde.Build.Tests
 {
     [TestClass]
     public class DatabaseRunnerTest
diff --git a/Engine/Source/Programs/Horde/Horde.Build.Tests/DebugTest.cs b/Engine/Source/Programs/Horde/Horde.Build.Tests/DebugTest.cs
--- a/Engine/Source/Programs/Horde/Horde.Build.Tests/DebugTest.cs
+++ b/Engine/Source/Programs/Horde/Horde.Build.Tests/DebugTest.cs
@@ -4,22 +4,22 @@ using System;
 using System.Collections.Generic;
 using System.Threading.Tasks;
 using Castle.Core.Internal;
-using HordeServer;
-using HordeServer.Models;
+using Horde.Build;
+using Horde.Build.Models;
 using Microsoft.VisualStudio.TestTools.UnitTesting;
-using HordeServer.Services;
-using HordeServer.Services.Impl;
-using HordeServer.Utilities;
+using Horde.Build.Services;
+using Horde.Build.Services.Impl;
+using Horde.Build.Utilities;
 using Microsoft.Extensions.DependencyInjection;
 using Microsoft.Extensions.Logging;
 using Microsoft.Extensions.Logging.Console;
 using MongoDB.Bson;
 
 #if false
-using PoolId = HordeServer.Utilities.StringId<HordeServer.Models.IPool>;
+using PoolId = Horde.Build.Utilities.StringId<Horde.Build.Models.IPool>;
 using Microsoft.Extensions.Options;
 
-namespace HordeServerTests
+namespace Horde.Build.Tests
 {
 	using JobId = ObjectId<IJob>;
 
diff --git a/Engine/Source/Programs/Horde/Horde.Build.Tests/Fixture.cs b/Engine/Source/Programs/Horde/Horde.Build.Tests/Fixture.cs
--- a/Engine/Source/Programs/Horde/Horde.Build.Tests/Fixture.cs
+++ b/Engine/Source/Programs/Horde/Horde.Build.Tests/Fixture.cs
@@ -6,19 +6,19 @@ using System.Text;
 using System.Threading.Tasks;
 using EpicGames.Core;
 using HordeCommon;
-using HordeServer.Collections;
-using HordeServer.Models;
-using HordeServer.Services;
-using HordeServerTests.Stubs.Services;
+using Horde.Build.Collections;
+using Horde.Build.Models;
+using Horde.Build.Services;
+using Horde.Build.Tests.Stubs.Services;
 using MongoDB.Bson;
-using StreamId = HordeServer.Utilities.StringId<HordeServer.Models.IStream>;
-using TemplateRefId = HordeServer.Utilities.StringId<HordeServer.Models.TemplateRef>;
-using ProjectId = HordeServer.Utilities.StringId<HordeServer.Models.IProject>;
-using HordeServer.Api;
-using HordeServer.Jobs;
-using HordeServer.Utilities;
-
-namespace HordeServerTests
+using StreamId = Horde.Build.Utilities.StringId<Horde.Build.Models.IStream>;
+using TemplateRefId = Horde.Build.Utilities.StringId<Horde.Build.Models.TemplateRef>;
+using ProjectId = Horde.Build.Utilities.StringId<Horde.Build.Models.IProject>;
+using Horde.Build.Api;
+using Horde.Build.Jobs;
+using Horde.Build.Utilities;
+
+namespace Horde.Build.Tests
 {
 	using JobId = ObjectId<IJob>;
 
diff --git a/Engine/Source/Programs/Horde/Horde.Build.Tests/Fleet/AutoscaleServiceV2Test.cs b/Engine/Source/Programs/Horde/Horde.Build.Tests/Fleet/AutoscaleServiceV2Test.cs
--- a/Engine/Source/Programs/Horde/Horde.Build.Tests/Fleet/AutoscaleServiceV2Test.cs
+++ b/Engine/Source/Programs/Horde/Horde.Build.Tests/Fleet/AutoscaleServiceV2Test.cs
@@ -5,11 +5,11 @@ using System.Collections.Generic;
 using System.Threading;
 using System.Threading.Tasks;
 using Horde.Build.Fleet.Autoscale;
-using HordeServer;
-using HordeServer.Models;
-using HordeServer.Services;
-using HordeServer.Utilities;
-using HordeServerTests;
+using Horde.Build;
+using Horde.Build.Models;
+using Horde.Build.Services;
+using Horde.Build.Utilities;
+using Horde.Build.Tests;
 using Microsoft.Extensions.DependencyInjection;
 using Microsoft.Extensions.Logging;
 using Microsoft.Extensions.Options;
diff --git a/Engine/Source/Programs/Horde/Horde.Build.Tests/Fleet/JobQueueStrategyTest.cs b/Engine/Source/Programs/Horde/Horde.Build.Tests/Fleet/JobQueueStrategyTest.cs
--- a/Engine/Source/Programs/Horde/Horde.Build.Tests/Fleet/JobQueueStrategyTest.cs
+++ b/Engine/Source/Programs/Horde/Horde.Build.Tests/Fleet/JobQueueStrategyTest.cs
@@ -5,10 +5,10 @@ using System.Collections.Generic;
 using System.Threading.Tasks;
 using EpicGames.Core;
 using Horde.Build.Fleet.Autoscale;
-using HordeServer.Api;
-using HordeServer.Models;
-using HordeServer.Utilities;
-using HordeServerTests;
+using Horde.Build.Api;
+using Horde.Build.Models;
+using Horde.Build.Utilities;
+using Horde.Build.Tests;
 using Microsoft.VisualStudio.TestTools.UnitTesting;
 
 namespace Horde.Build.Tests.Fleet
diff --git a/Engine/Source/Programs/Horde/Horde.Build.Tests/Fleet/LeaseUtilizationStrategyTest.cs b/Engine/Source/Programs/Horde/Horde.Build.Tests/Fleet/LeaseUtilizationStrategyTest.cs
--- a/Engine/Source/Programs/Horde/Horde.Build.Tests/Fleet/LeaseUtilizationStrategyTest.cs
+++ b/Engine/Source/Programs/Horde/Horde.Build.Tests/Fleet/LeaseUtilizationStrategyTest.cs
@@ -8,9 +8,9 @@ using Google.Protobuf.WellKnownTypes;
 using Horde.Build.Fleet.Autoscale;
 using HordeCommon;
 using HordeCommon.Rpc.Tasks;
-using HordeServer.Models;
-using HordeServer.Utilities;
-using HordeServerTests;
+using Horde.Build.Models;
+using Horde.Build.Utilities;
+using Horde.Build.Tests;
 using Microsoft.VisualStudio.TestTools.UnitTesting;
 
 namespace Horde.Build.Tests.Fleet
diff --git a/Engine/Source/Programs/Horde/Horde.Build.Tests/IssueServiceTests.cs b/Engine/Source/Programs/Horde/Horde.Build.Tests/IssueServiceTests.cs
--- a/Engine/Source/Programs/Horde/Horde.Build.Tests/IssueServiceTests.cs
+++ b/Engine/Source/Programs/Horde/Horde.Build.Tests/IssueServiceTests.cs
@@ -1,24 +1,24 @@
 // Copyright Epic Games, Inc. All Rights Reserved.
 
 using EpicGames.Core;
-using HordeAgent.Parser;
-using HordeAgent.Utility;
+using Horde.Agent.Parser;
+using Horde.Agent.Utility;
 using HordeCommon;
-using HordeServer;
-using HordeServer.Collections;
-using HordeServer.Collections.Impl;
-using HordeServer.IssueHandlers.Impl;
-using HordeServer.Logs;
-using HordeServer.Logs.Builder;
-using HordeServer.Logs.Storage;
-using HordeServer.Logs.Storage.Impl;
-using HordeServer.Models;
-using HordeServer.Notifications;
-using HordeServer.Services;
-using HordeServer.Services.Impl;
-using HordeServer.Utilities;
-using HordeServerTests.Stubs.Collections;
-using HordeServerTests.Stubs.Services;
+using Horde.Build;
+using Horde.Build.Collections;
+using Horde.Build.Collections.Impl;
+using Horde.Build.IssueHandlers.Impl;
+using Horde.Build.Logs;
+using Horde.Build.Logs.Builder;
+using Horde.Build.Logs.Storage;
+using Horde.Build.Logs.Storage.Impl;
+using Horde.Build.Models;
+using Horde.Build.Notifications;
+using Horde.Build.Services;
+using Horde.Build.Services.Impl;
+using Horde.Build.Utilities;
+using Horde.Build.Tests.Stubs.Collections;
+using Horde.Build.Tests.Stubs.Services;
 using Microsoft.Extensions.DependencyInjection;
 using Microsoft.Extensions.Logging;
 using Microsoft.Extensions.Logging.Abstractions;
@@ -34,11 +34,11 @@ using System.Runtime.InteropServices;
 using System.Text.Json;
 using System.Threading.Tasks;
 
-using ProjectId = HordeServer.Utilities.StringId<HordeServer.Models.IProject>;
-using StreamId = HordeServer.Utilities.StringId<HordeServer.Models.IStream>;
-using TemplateRefId = HordeServer.Utilities.StringId<HordeServer.Models.TemplateRef>;
+using ProjectId = Horde.Build.Utilities.StringId<Horde.Build.Models.IProject>;
+using StreamId = Horde.Build.Utilities.StringId<Horde.Build.Models.IStream>;
+using TemplateRefId = Horde.Build.Utilities.StringId<Horde.Build.Models.TemplateRef>;
 
-namespace HordeServerTests
+namespace Horde.Build.Tests
 {
 	using JobId = ObjectId<IJob>;
 	using LogId = ObjectId<ILogFile>;
@@ -189,7 +189,7 @@ namespace HordeServerTests
 				{
 					SubResourceId StepId = new SubResourceId((ushort)((GroupIdx * 100) + NodeIdx));
 
-					ILogFile LogFile = LogFileService.CreateLogFileAsync(JobId, null, HordeServer.Api.LogType.Json).Result;
+					ILogFile LogFile = LogFileService.CreateLogFileAsync(JobId, null, Horde.Build.Api.LogType.Json).Result;
 
 					Mock<IJobStep> Step = new Mock<IJobStep>(MockBehavior.Strict);
 					Step.SetupGet(x => x.Id).Returns(StepId);
diff --git a/Engine/Source/Programs/Horde/Horde.Build.Tests/JobCollectionTests.cs b/Engine/Source/Programs/Horde/Horde.Build.Tests/JobCollectionTests.cs
--- a/Engine/Source/Programs/Horde/Horde.Build.Tests/JobCollectionTests.cs
+++ b/Engine/Source/Programs/Horde/Horde.Build.Tests/JobCollectionTests.cs
@@ -1,30 +1,30 @@
 // Copyright Epic Games, Inc. All Rights Reserved.
 
-using HordeServer.Collections;
+using Horde.Build.Collections;
 using MongoDB.Bson;
 using System;
 using System.Collections.Generic;
 using System.Text;
 
-using ProjectId = HordeServer.Utilities.StringId<HordeServer.Models.IProject>;
-using PoolId = HordeServer.Utilities.StringId<HordeServer.Models.IPool>;
-using StreamId = HordeServer.Utilities.StringId<HordeServer.Models.IStream>;
-using TemplateRefId = HordeServer.Utilities.StringId<HordeServer.Models.TemplateRef>;
+using ProjectId = Horde.Build.Utilities.StringId<Horde.Build.Models.IProject>;
+using PoolId = Horde.Build.Utilities.StringId<Horde.Build.Models.IPool>;
+using StreamId = Horde.Build.Utilities.StringId<Horde.Build.Models.IStream>;
+using TemplateRefId = Horde.Build.Utilities.StringId<Horde.Build.Models.TemplateRef>;
 
 using Microsoft.VisualStudio.TestTools.UnitTesting;
 using EpicGames.Core;
 using Moq;
-using HordeServer.Models;
-using HordeServer.Services;
+using Horde.Build.Models;
+using Horde.Build.Services;
 using System.Threading.Tasks;
-using HordeServer.Api;
+using Horde.Build.Api;
 using System.Linq;
 using HordeCommon;
-using HordeServer.Collections.Impl;
+using Horde.Build.Collections.Impl;
 using Microsoft.Extensions.Logging.Abstractions;
-using HordeServer.Utilities;
+using Horde.Build.Utilities;
 
-namespace HordeServerTests
+namespace Horde.Build.Tests
 {
 	using JobId = ObjectId<IJob>;
 	using LeaseId = ObjectId<ILease>;
diff --git a/Engine/Source/Programs/Horde/Horde.Build.Tests/JobServiceTests.cs b/Engine/Source/Programs/Horde/Horde.Build.Tests/JobServiceTests.cs
--- a/Engine/Source/Programs/Horde/Horde.Build.Tests/JobServiceTests.cs
+++ b/Engine/Source/Programs/Horde/Horde.Build.Tests/JobServiceTests.cs
@@ -2,20 +2,20 @@
 
 using System;
 using HordeCommon;
-using HordeServer.Api;
-using HordeServer.Models;
+using Horde.Build.Api;
+using Horde.Build.Models;
 using Microsoft.VisualStudio.TestTools.UnitTesting;
 using MongoDB.Bson;
 using System.Collections.Generic;
 using System.Text;
 using System.Threading.Tasks;
 using EpicGames.Core;
-using ProjectId = HordeServer.Utilities.StringId<HordeServer.Models.IProject>;
-using StreamId = HordeServer.Utilities.StringId<HordeServer.Models.IStream>;
-using TemplateRefId = HordeServer.Utilities.StringId<HordeServer.Models.TemplateRef>;
-using HordeServer.Utilities;
+using ProjectId = Horde.Build.Utilities.StringId<Horde.Build.Models.IProject>;
+using StreamId = Horde.Build.Utilities.StringId<Horde.Build.Models.IStream>;
+using TemplateRefId = Horde.Build.Utilities.StringId<Horde.Build.Models.TemplateRef>;
+using Horde.Build.Utilities;
 
-namespace HordeServerTests
+namespace Horde.Build.Tests
 {
 	using JobId = ObjectId<IJob>;
 	using LogId = ObjectId<ILogFile>;
diff --git a/Engine/Source/Programs/Horde/Horde.Build.Tests/JobTaskSourceTests.cs b/Engine/Source/Programs/Horde/Horde.Build.Tests/JobTaskSourceTests.cs
--- a/Engine/Source/Programs/Horde/Horde.Build.Tests/JobTaskSourceTests.cs
+++ b/Engine/Source/Programs/Horde/Horde.Build.Tests/JobTaskSourceTests.cs
@@ -1,15 +1,15 @@
 // Copyright Epic Games, Inc. All Rights Reserved.
 
 using HordeCommon;
-using HordeServer.Api;
-using HordeServer.Models;
+using Horde.Build.Api;
+using Horde.Build.Models;
 using Microsoft.VisualStudio.TestTools.UnitTesting;
 using System.Collections.Generic;
 using System.Threading;
 using System.Threading.Tasks;
-using HordeServer.Utilities;
+using Horde.Build.Utilities;
 
-namespace HordeServerTests
+namespace Horde.Build.Tests
 {
 	using JobId = ObjectId<IJob>;
 	using LogId = ObjectId<ILogFile>;
diff --git a/Engine/Source/Programs/Horde/Horde.Build.Tests/JobsControllerDbTest.cs b/Engine/Source/Programs/Horde/Horde.Build.Tests/JobsControllerDbTest.cs
--- a/Engine/Source/Programs/Horde/Horde.Build.Tests/JobsControllerDbTest.cs
+++ b/Engine/Source/Programs/Horde/Horde.Build.Tests/JobsControllerDbTest.cs
@@ -4,14 +4,14 @@ using System;
 using System.Collections.Generic;
 using System.Diagnostics;
 using System.Threading.Tasks;
-using HordeServer.Api;
+using Horde.Build.Api;
 using Microsoft.VisualStudio.TestTools.UnitTesting;
-using HordeServer.Controllers;
-using HordeServer.Models;
-using HordeServer.Utilities;
+using Horde.Build.Controllers;
+using Horde.Build.Models;
+using Horde.Build.Utilities;
 using Microsoft.AspNetCore.Mvc;
 
-namespace HordeServerTests
+namespace Horde.Build.Tests
 {
     /// <summary>
     /// Database-only test for testing the Job controller. Different from the JobsController test that set up
diff --git a/Engine/Source/Programs/Horde/Horde.Build.Tests/LeasesControllerDbTest.cs b/Engine/Source/Programs/Horde/Horde.Build.Tests/LeasesControllerDbTest.cs
--- a/Engine/Source/Programs/Horde/Horde.Build.Tests/LeasesControllerDbTest.cs
+++ b/Engine/Source/Programs/Horde/Horde.Build.Tests/LeasesControllerDbTest.cs
@@ -4,10 +4,10 @@ using System;
 using System.Collections.Generic;
 using System.Threading.Tasks;
 using HordeCommon;
-using HordeServer.Api;
-using HordeServer.Models;
-using HordeServer.Utilities;
-using HordeServerTests;
+using Horde.Build.Api;
+using Horde.Build.Models;
+using Horde.Build.Utilities;
+using Horde.Build.Tests;
 using Microsoft.AspNetCore.Mvc;
 using Microsoft.VisualStudio.TestTools.UnitTesting;
 
diff --git a/Engine/Source/Programs/Horde/Horde.Build.Tests/LogBuilderTests.cs b/Engine/Source/Programs/Horde/Horde.Build.Tests/LogBuilderTests.cs
--- a/Engine/Source/Programs/Horde/Horde.Build.Tests/LogBuilderTests.cs
+++ b/Engine/Source/Programs/Horde/Horde.Build.Tests/LogBuilderTests.cs
@@ -1,10 +1,10 @@
 // Copyright Epic Games, Inc. All Rights Reserved.
 
-using HordeServer.Api;
-using HordeServer.Logs;
-using HordeServer.Logs.Builder;
-using HordeServer.Models;
-using HordeServer.Utilities;
+using Horde.Build.Api;
+using Horde.Build.Logs;
+using Horde.Build.Logs.Builder;
+using Horde.Build.Models;
+using Horde.Build.Utilities;
 using Microsoft.Extensions.Logging.Abstractions;
 using Microsoft.VisualStudio.TestTools.UnitTesting;
 using MongoDB.Bson;
@@ -13,7 +13,7 @@ using System.Collections.Generic;
 using System.Text;
 using System.Threading.Tasks;
 
-namespace HordeServerTests
+namespace Horde.Build.Tests
 {
 	using LogId = ObjectId<ILogFile>;
 
diff --git a/Engine/Source/Programs/Horde/Horde.Build.Tests/LogChunkTests.cs b/Engine/Source/Programs/Horde/Horde.Build.Tests/LogChunkTests.cs
--- a/Engine/Source/Programs/Horde/Horde.Build.Tests/LogChunkTests.cs
+++ b/Engine/Source/Programs/Horde/Horde.Build.Tests/LogChunkTests.cs
@@ -1,13 +1,13 @@
 // Copyright Epic Games, Inc. All Rights Reserved.
 
-using HordeServer.Api;
-using HordeServer.Logs;
+using Horde.Build.Api;
+using Horde.Build.Logs;
 using Microsoft.VisualStudio.TestTools.UnitTesting;
 using System;
 using System.Collections.Generic;
 using System.Text;
 
-namespace HordeServerTests
+namespace Horde.Build.Tests
 {
 	[TestClass]
 	public class LogChunkTests
diff --git a/Engine/Source/Programs/Horde/Horde.Build.Tests/LogFileServiceTests.cs b/Engine/Source/Programs/Horde/Horde.Build.Tests/LogFileServiceTests.cs
--- a/Engine/Source/Programs/Horde/Horde.Build.Tests/LogFileServiceTests.cs
+++ b/Engine/Source/Programs/Horde/Horde.Build.Tests/LogFileServiceTests.cs
@@ -11,28 +11,28 @@ using System.Security.Claims;
 using System.Text;
 using System.Threading.Tasks;
 using HordeCommon;
-using HordeServer;
-using HordeServer.Api;
-using HordeServer.Collections.Impl;
-using HordeServer.Logs;
-using HordeServer.Logs.Builder;
-using HordeServer.Logs.Readers;
-using HordeServer.Logs.Storage;
-using HordeServer.Logs.Storage.Impl;
-using HordeServer.Models;
-using HordeServer.Services;
-using HordeServer.Storage;
-using HordeServer.Storage.Backends;
-using HordeServer.Utilities;
+using Horde.Build;
+using Horde.Build.Api;
+using Horde.Build.Collections.Impl;
+using Horde.Build.Logs;
+using Horde.Build.Logs.Builder;
+using Horde.Build.Logs.Readers;
+using Horde.Build.Logs.Storage;
+using Horde.Build.Logs.Storage.Impl;
+using Horde.Build.Models;
+using Horde.Build.Services;
+using Horde.Build.Storage;
+using Horde.Build.Storage.Backends;
+using Horde.Build.Utilities;
 using Microsoft.Extensions.Configuration;
 using Microsoft.Extensions.Logging;
 using Microsoft.Extensions.Logging.Abstractions;
 using Microsoft.VisualStudio.TestTools.UnitTesting;
 using MongoDB.Bson;
 using StackExchange.Redis;
-using PoolId = HordeServer.Utilities.StringId<HordeServer.Models.IPool>;
+using PoolId = Horde.Build.Utilities.StringId<Horde.Build.Models.IPool>;
 
-namespace HordeServerTests
+namespace Horde.Build.Tests
 {
 	using JobId = ObjectId<IJob>;
 	using LogId = ObjectId<ILogFile>;
diff --git a/Engine/Source/Programs/Horde/Horde.Build.Tests/LogIndexingTests.cs b/Engine/Source/Programs/Horde/Horde.Build.Tests/LogIndexingTests.cs
--- a/Engine/Source/Programs/Horde/Horde.Build.Tests/LogIndexingTests.cs
+++ b/Engine/Source/Programs/Horde/Horde.Build.Tests/LogIndexingTests.cs
@@ -1,15 +1,15 @@
 // Copyright Epic Games, Inc. All Rights Reserved.
 
-using HordeServer.Api;
-using HordeServer.Collections;
-using HordeServer.Logs;
-using HordeServer.Logs.Builder;
-using HordeServer.Logs.Storage;
-using HordeServer.Logs.Storage.Impl;
-using HordeServer.Models;
-using HordeServer.Services;
-using HordeServer.Utilities;
-using HordeServer.Compute.Tests.Properties;
+using Horde.Build.Api;
+using Horde.Build.Collections;
+using Horde.Build.Logs;
+using Horde.Build.Logs.Builder;
+using Horde.Build.Logs.Storage;
+using Horde.Build.Logs.Storage.Impl;
+using Horde.Build.Models;
+using Horde.Build.Services;
+using Horde.Build.Utilities;
+using Horde.Build.Compute.Tests.Properties;
 using Microsoft.Extensions.Configuration;
 using Microsoft.Extensions.DependencyInjection;
 using Microsoft.Extensions.Logging;
@@ -21,11 +21,11 @@ using System.IO;
 using System.Linq;
 using System.Text;
 using System.Threading.Tasks;
-using HordeServer.Collections.Impl;
+using Horde.Build.Collections.Impl;
 using HordeCommon;
 using Microsoft.Extensions.Logging.Abstractions;
 
-namespace HordeServerTests
+namespace Horde.Build.Tests
 {
 	using JobId = ObjectId<IJob>;
 
diff --git a/Engine/Source/Programs/Horde/Horde.Build.Tests/Models/Job/GetLabelStatesTests.cs b/Engine/Source/Programs/Horde/Horde.Build.Tests/Models/Job/GetLabelStatesTests.cs
--- a/Engine/Source/Programs/Horde/Horde.Build.Tests/Models/Job/GetLabelStatesTests.cs
+++ b/Engine/Source/Programs/Horde/Horde.Build.Tests/Models/Job/GetLabelStatesTests.cs
@@ -1,15 +1,15 @@
 // Copyright Epic Games, Inc. All Rights Reserved.
 
 using System.Collections.Generic;
-using HordeServer.Api;
+using Horde.Build.Api;
 using HordeCommon;
-using HordeServer.Models;
+using Horde.Build.Models;
 using Microsoft.VisualStudio.TestTools.UnitTesting;
 using MongoDB.Bson;
 using Moq;
-using HordeServer.Utilities;
+using Horde.Build.Utilities;
 
-namespace HordeServerTests
+namespace Horde.Build.Tests
 {
 	using JobId = ObjectId<IJob>;
 
diff --git a/Engine/Source/Programs/Horde/Horde.Build.Tests/NotificationServiceTests.cs b/Engine/Source/Programs/Horde/Horde.Build.Tests/NotificationServiceTests.cs
--- a/Engine/Source/Programs/Horde/Horde.Build.Tests/NotificationServiceTests.cs
+++ b/Engine/Source/Programs/Horde/Horde.Build.Tests/NotificationServiceTests.cs
@@ -4,19 +4,19 @@ using System;
 using System.Collections.Generic;
 using System.Threading.Tasks;
 using HordeCommon;
-using HordeServer.Api;
-using HordeServer.Models;
-using HordeServer.Notifications;
-using HordeServer.Notifications.Impl;
-using HordeServer.Utilities;
+using Horde.Build.Api;
+using Horde.Build.Models;
+using Horde.Build.Notifications;
+using Horde.Build.Notifications.Impl;
+using Horde.Build.Utilities;
 using Microsoft.Extensions.DependencyInjection;
 using Microsoft.Extensions.Logging;
 using Microsoft.VisualStudio.TestTools.UnitTesting;
 using Moq;
 using StackExchange.Redis;
-using StreamId = HordeServer.Utilities.StringId<HordeServer.Models.IStream>;
+using StreamId = Horde.Build.Utilities.StringId<Horde.Build.Models.IStream>;
 
-namespace HordeServerTests
+namespace Horde.Build.Tests
 {
 	using JobId = ObjectId<IJob>;
 
diff --git a/Engine/Source/Programs/Horde/Horde.Build.Tests/PoolServiceTests.cs b/Engine/Source/Programs/Horde/Horde.Build.Tests/PoolServiceTests.cs
--- a/Engine/Source/Programs/Horde/Horde.Build.Tests/PoolServiceTests.cs
+++ b/Engine/Source/Programs/Horde/Horde.Build.Tests/PoolServiceTests.cs
@@ -4,13 +4,13 @@ using System;
 using System.Collections.Generic;
 using System.Threading.Tasks;
 using HordeCommon;
-using HordeServer.Collections.Impl;
-using HordeServer.Models;
-using HordeServer.Services;
+using Horde.Build.Collections.Impl;
+using Horde.Build.Models;
+using Horde.Build.Services;
 using Microsoft.VisualStudio.TestTools.UnitTesting;
-using PoolId = HordeServer.Utilities.StringId<HordeServer.Models.IPool>;
+using PoolId = Horde.Build.Utilities.StringId<Horde.Build.Models.IPool>;
 
-namespace HordeServerTests
+namespace Horde.Build.Tests
 {
     [TestClass]
     public class PoolServiceTests : TestSetup
diff --git a/Engine/Source/Programs/Horde/Horde.Build.Tests/PoolsControllerTest.cs b/Engine/Source/Programs/Horde/Horde.Build.Tests/PoolsControllerTest.cs
--- a/Engine/Source/Programs/Horde/Horde.Build.Tests/PoolsControllerTest.cs
+++ b/Engine/Source/Programs/Horde/Horde.Build.Tests/PoolsControllerTest.cs
@@ -3,10 +3,10 @@
 using System.Collections.Generic;
 using System.Threading.Tasks;
 using Horde.Build.Fleet.Autoscale;
-using HordeServer.Api;
-using HordeServer.Models;
-using HordeServer.Utilities;
-using HordeServerTests;
+using Horde.Build.Api;
+using Horde.Build.Models;
+using Horde.Build.Utilities;
+using Horde.Build.Tests;
 using Microsoft.AspNetCore.Mvc;
 using Microsoft.VisualStudio.TestTools.UnitTesting;
 
@@ -36,7 +36,7 @@ namespace Horde.Build.Tests
 		        ScaleOutCooldown = 111,
 		        ScaleInCooldown = 222,
 		        SizeStrategy = PoolSizeStrategy.JobQueue,
-		        JobQueueSettings = new HordeServer.Api.JobQueueSettings(new Horde.Build.Fleet.Autoscale.JobQueueSettings(0.35, 0.85))
+		        JobQueueSettings = new Horde.Build.Api.JobQueueSettings(new Horde.Build.Fleet.Autoscale.JobQueueSettings(0.35, 0.85))
 	        };
 	        
 	        ActionResult<CreatePoolResponse> Result = await PoolsController.CreatePoolAsync(Request);
diff --git a/Engine/Source/Programs/Horde/Horde.Build.Tests/ProgramTests.cs b/Engine/Source/Programs/Horde/Horde.Build.Tests/ProgramTests.cs
--- a/Engine/Source/Programs/Horde/Horde.Build.Tests/ProgramTests.cs
+++ b/Engine/Source/Programs/Horde/Horde.Build.Tests/ProgramTests.cs
@@ -6,7 +6,7 @@ using System.Runtime.InteropServices;
 using System.Security.Cryptography.X509Certificates;
 using EpicGames.Core;
 using Horde.Build.Commands;
-using HordeServer;
+using Horde.Build;
 using Microsoft.VisualStudio.TestTools.UnitTesting;
 
 namespace Horde.Build.Tests
diff --git a/Engine/Source/Programs/Horde/Horde.Build.Tests/RpcServiceTest.cs b/Engine/Source/Programs/Horde/Horde.Build.Tests/RpcServiceTest.cs
--- a/Engine/Source/Programs/Horde/Horde.Build.Tests/RpcServiceTest.cs
+++ b/Engine/Source/Programs/Horde/Horde.Build.Tests/RpcServiceTest.cs
@@ -11,13 +11,13 @@ using System.Threading.Channels;
 using System.Threading.Tasks;
 using Google.Protobuf;
 using Grpc.Core;
-using HordeServer;
-using HordeServer.Collections;
+using Horde.Build;
+using Horde.Build.Collections;
 using HordeCommon.Rpc;
-using HordeServer.Models;
-using HordeServer.Services;
-using HordeServer.Utilities;
-using HordeServerTests.Stubs.Collections;
+using Horde.Build.Models;
+using Horde.Build.Services;
+using Horde.Build.Utilities;
+using Horde.Build.Tests.Stubs.Collections;
 using Microsoft.AspNetCore.Http;
 using Microsoft.AspNetCore.Http.Features;
 using Microsoft.Extensions.Configuration;
@@ -27,10 +27,10 @@ using Microsoft.VisualStudio.TestTools.UnitTesting;
 using MongoDB.Bson;
 using AgentCapabilities = HordeCommon.Rpc.Messages.AgentCapabilities;
 using ISession = Microsoft.AspNetCore.Http.ISession;
-using PoolId = HordeServer.Utilities.StringId<HordeServer.Models.IPool>;
+using PoolId = Horde.Build.Utilities.StringId<Horde.Build.Models.IPool>;
 
 
-namespace HordeServerTests
+namespace Horde.Build.Tests
 {
 	using LeaseId = ObjectId<ILease>;
 	using LogId = ObjectId<ILogFile>;
@@ -406,7 +406,7 @@ namespace HordeServerTests
 
 			// Set the session ID on the job batch to pass auth later
 			Deref(await JobCollection.TryAssignLeaseAsync(Fixture.Job1, 0, new PoolId("foo"),
-				new AgentId("test"), new ObjectId<HordeServer.Models.ISession>(SessionId),
+				new AgentId("test"), new ObjectId<Horde.Build.Models.ISession>(SessionId),
 				LeaseId.GenerateNewId(), LogId.GenerateNewId()));
 /*
 			TestAsyncStreamReader<UploadArtifactRequest> RequestStream = new TestAsyncStreamReader<UploadArtifactRequest>(Context);
diff --git a/Engine/Source/Programs/Horde/Horde.Build.Tests/SchedulerTests.cs b/Engine/Source/Programs/Horde/Horde.Build.Tests/SchedulerTests.cs
--- a/Engine/Source/Programs/Horde/Horde.Build.Tests/SchedulerTests.cs
+++ b/Engine/Source/Programs/Horde/Horde.Build.Tests/SchedulerTests.cs
@@ -4,24 +4,24 @@ using System;
 using System.Collections.Generic;
 using System.Threading.Tasks;
 using HordeCommon;
-using HordeServer.Collections.Impl;
-using HordeServer.Models;
-using HordeServer.Services;
+using Horde.Build.Collections.Impl;
+using Horde.Build.Models;
+using Horde.Build.Services;
 using Microsoft.Extensions.DependencyInjection;
 using Microsoft.VisualStudio.TestTools.UnitTesting;
-using PoolId = HordeServer.Utilities.StringId<HordeServer.Models.IPool>;
-using StreamId = HordeServer.Utilities.StringId<HordeServer.Models.IStream>;
-using ProjectId = HordeServer.Utilities.StringId<HordeServer.Models.IProject>;
-using TemplateRefId = HordeServer.Utilities.StringId<HordeServer.Models.TemplateRef>;
-using HordeServer.Api;
+using PoolId = Horde.Build.Utilities.StringId<Horde.Build.Models.IPool>;
+using StreamId = Horde.Build.Utilities.StringId<Horde.Build.Models.IStream>;
+using ProjectId = Horde.Build.Utilities.StringId<Horde.Build.Models.IProject>;
+using TemplateRefId = Horde.Build.Utilities.StringId<Horde.Build.Models.TemplateRef>;
+using Horde.Build.Api;
 using MongoDB.Bson;
 using System.Linq;
-using HordeServerTests.Stubs.Services;
-using HordeServer.Collections;
-using HordeServer.Utilities;
+using Horde.Build.Tests.Stubs.Services;
+using Horde.Build.Collections;
+using Horde.Build.Utilities;
 using System.Threading;
 
-namespace HordeServerTests
+namespace Horde.Build.Tests
 {
 	using JobId = ObjectId<IJob>;
 	using LogId = ObjectId<ILogFile>;
diff --git a/Engine/Source/Programs/Horde/Horde.Build.Tests/ServiceAccountAuthTest.cs b/Engine/Source/Programs/Horde/Horde.Build.Tests/ServiceAccountAuthTest.cs
--- a/Engine/Source/Programs/Horde/Horde.Build.Tests/ServiceAccountAuthTest.cs
+++ b/Engine/Source/Programs/Horde/Horde.Build.Tests/ServiceAccountAuthTest.cs
@@ -4,11 +4,11 @@ using System.Collections.Generic;
 using System.Linq;
 using System.Security.Claims;
 using System.Threading.Tasks;
-using HordeServer.Authentication;
-using HordeServer.Collections;
-using HordeServer.Collections.Impl;
-using HordeServer.Models;
-using HordeServer.Services;
+using Horde.Build.Authentication;
+using Horde.Build.Collections;
+using Horde.Build.Collections.Impl;
+using Horde.Build.Models;
+using Horde.Build.Services;
 using Microsoft.AspNetCore.Authentication;
 using Microsoft.AspNetCore.Http;
 using Microsoft.Extensions.Logging;
@@ -17,7 +17,7 @@ using Microsoft.Extensions.WebEncoders.Testing;
 using Microsoft.Net.Http.Headers;
 using Microsoft.VisualStudio.TestTools.UnitTesting;
 
-namespace HordeServerTests
+namespace Horde.Build.Tests
 {
 	[TestClass]
 	public class ServiceAccountAuthTest : DatabaseIntegrationTest
diff --git a/Engine/Source/Programs/Horde/Horde.Build.Tests/ServiceAccountCollectionTests.cs b/Engine/Source/Programs/Horde/Horde.Build.Tests/ServiceAccountCollectionTests.cs
--- a/Engine/Source/Programs/Horde/Horde.Build.Tests/ServiceAccountCollectionTests.cs
+++ b/Engine/Source/Programs/Horde/Horde.Build.Tests/ServiceAccountCollectionTests.cs
@@ -2,14 +2,14 @@
 
 using System.Collections.Generic;
 using System.Threading.Tasks;
-using HordeServer.Collections;
-using HordeServer.Collections.Impl;
-using HordeServer.Models;
-using HordeServer.Services;
+using Horde.Build.Collections;
+using Horde.Build.Collections.Impl;
+using Horde.Build.Models;
+using Horde.Build.Services;
 using Microsoft.VisualStudio.TestTools.UnitTesting;
 using MongoDB.Bson;
 
-namespace HordeServerTests
+namespace Horde.Build.Tests
 {
 	[TestClass]
 	public class ServiceAccountCollectionTests : DatabaseIntegrationTest
diff --git a/Engine/Source/Programs/Horde/Horde.Build.Tests/StorageTests.cs b/Engine/Source/Programs/Horde/Horde.Build.Tests/StorageTests.cs
--- a/Engine/Source/Programs/Horde/Horde.Build.Tests/StorageTests.cs
+++ b/Engine/Source/Programs/Horde/Horde.Build.Tests/StorageTests.cs
@@ -3,8 +3,8 @@
 using EpicGames.Core;
 using EpicGames.Horde.Storage;
 using EpicGames.Serialization;
-using HordeServer.Services;
-using HordeServer.Utilities;
+using Horde.Build.Services;
+using Horde.Build.Utilities;
 using Microsoft.Extensions.DependencyInjection;
 using Microsoft.VisualStudio.TestTools.UnitTesting;
 using System;
@@ -13,7 +13,7 @@ using System.IO;
 using System.Text;
 using System.Threading.Tasks;
 
-namespace HordeServerTests
+namespace Horde.Build.Tests
 {
 	[TestClass]
 	public class StorageTests : TestSetup
diff --git a/Engine/Source/Programs/Horde/Horde.Build.Tests/StreamServiceTest.cs b/Engine/Source/Programs/Horde/Horde.Build.Tests/StreamServiceTest.cs
--- a/Engine/Source/Programs/Horde/Horde.Build.Tests/StreamServiceTest.cs
+++ b/Engine/Source/Programs/Horde/Horde.Build.Tests/StreamServiceTest.cs
@@ -3,11 +3,11 @@
 using System;
 using System.Globalization;
 using System.Threading.Tasks;
-using HordeServer.Models;
+using Horde.Build.Models;
 using Microsoft.VisualStudio.TestTools.UnitTesting;
-using PoolId = HordeServer.Utilities.StringId<HordeServer.Models.IPool>;
+using PoolId = Horde.Build.Utilities.StringId<Horde.Build.Models.IPool>;
 
-namespace HordeServerTests
+namespace Horde.Build.Tests
 {
     [TestClass]
     public class StreamServiceTests : TestSetup
diff --git a/Engine/Source/Programs/Horde/Horde.Build.Tests/Stubs/Collections/JobCollectionStub.cs b/Engine/Source/Programs/Horde/Horde.Build.Tests/Stubs/Collections/JobCollectionStub.cs
--- a/Engine/Source/Programs/Horde/Horde.Build.Tests/Stubs/Collections/JobCollectionStub.cs
+++ b/Engine/Source/Programs/Horde/Horde.Build.Tests/Stubs/Collections/JobCollectionStub.cs
@@ -1,8 +1,8 @@
 // Copyright Epic Games, Inc. All Rights Reserved.
 
-using HordeServer.Collections;
-using HordeServer.Models;
-using HordeServer.Utilities;
+using Horde.Build.Collections;
+using Horde.Build.Models;
+using Horde.Build.Utilities;
 using MongoDB.Bson;
 using MongoDB.Driver;
 using System;
@@ -10,7 +10,7 @@ using System.Collections.Generic;
 using System.Text;
 using System.Threading.Tasks;
 
-namespace HordeServerTests.Stubs.Collections
+namespace Horde.Build.Tests.Stubs.Collections
 {
 	using JobId = ObjectId<IJob>;
 
diff --git a/Engine/Source/Programs/Horde/Horde.Build.Tests/Stubs/Collections/JobStepRefCollectionStub.cs b/Engine/Source/Programs/Horde/Horde.Build.Tests/Stubs/Collections/JobStepRefCollectionStub.cs
--- a/Engine/Source/Programs/Horde/Horde.Build.Tests/Stubs/Collections/JobStepRefCollectionStub.cs
+++ b/Engine/Source/Programs/Horde/Horde.Build.Tests/Stubs/Collections/JobStepRefCollectionStub.cs
@@ -1,9 +1,9 @@
 // Copyright Epic Games, Inc. All Rights Reserved.
 
-using HordeServer.Collections;
+using Horde.Build.Collections;
 using HordeCommon;
-using HordeServer.Models;
-using HordeServer.Utilities;
+using Horde.Build.Models;
+using Horde.Build.Utilities;
 using MongoDB.Bson;
 using Moq;
 using System;
@@ -11,11 +11,11 @@ using System.Collections.Generic;
 using System.Text;
 using System.Threading.Tasks;
 
-using StreamId = HordeServer.Utilities.StringId<HordeServer.Models.IStream>;
-using TemplateRefId = HordeServer.Utilities.StringId<HordeServer.Models.TemplateRef>;
-using PoolId = HordeServer.Utilities.StringId<HordeServer.Models.IPool>;
+using StreamId = Horde.Build.Utilities.StringId<Horde.Build.Models.IStream>;
+using TemplateRefId = Horde.Build.Utilities.StringId<Horde.Build.Models.TemplateRef>;
+using PoolId = Horde.Build.Utilities.StringId<Horde.Build.Models.IPool>;
 
-namespace HordeServerTests.Stubs.Collections
+namespace Horde.Build.Tests.Stubs.Collections
 {
 	using JobId = ObjectId<IJob>;
 	using LogId = ObjectId<ILogFile>;
diff --git a/Engine/Source/Programs/Horde/Horde.Build.Tests/Stubs/Collections/SingletonDocumentStub.cs b/Engine/Source/Programs/Horde/Horde.Build.Tests/Stubs/Collections/SingletonDocumentStub.cs
--- a/Engine/Source/Programs/Horde/Horde.Build.Tests/Stubs/Collections/SingletonDocumentStub.cs
+++ b/Engine/Source/Programs/Horde/Horde.Build.Tests/Stubs/Collections/SingletonDocumentStub.cs
@@ -1,13 +1,13 @@
 // Copyright Epic Games, Inc. All Rights Reserved.
 
-using HordeServer.Models;
-using HordeServer.Utilities;
+using Horde.Build.Models;
+using Horde.Build.Utilities;
 using System;
 using System.Collections.Generic;
 using System.Text;
 using System.Threading.Tasks;
 
-namespace HordeServerTests.Stubs.Collections
+namespace Horde.Build.Tests.Stubs.Collections
 {
 	class SingletonDocumentStub<T> : ISingletonDocument<T> where T : SingletonBase, new()
 	{
diff --git a/Engine/Source/Programs/Horde/Horde.Build.Tests/Stubs/Collections/StreamCollectionStub.cs b/Engine/Source/Programs/Horde/Horde.Build.Tests/Stubs/Collections/StreamCollectionStub.cs
--- a/Engine/Source/Programs/Horde/Horde.Build.Tests/Stubs/Collections/StreamCollectionStub.cs
+++ b/Engine/Source/Programs/Horde/Horde.Build.Tests/Stubs/Collections/StreamCollectionStub.cs
@@ -1,15 +1,15 @@
 // Copyright Epic Games, Inc. All Rights Reserved.
 
-using HordeServer.Collections;
-using HordeServer.Models;
-using HordeServer.Utilities;
+using Horde.Build.Collections;
+using Horde.Build.Models;
+using Horde.Build.Utilities;
 using MongoDB.Bson;
 using System;
 using System.Collections.Generic;
 using System.Text;
 using System.Threading.Tasks;
 
-namespace HordeServerTests.Stubs.Collections
+namespace Horde.Build.Tests.Stubs.Collections
 {
 	class StreamCollectionStub : IStreamCollection
 	{
diff --git a/Engine/Source/Programs/Horde/Horde.Build.Tests/Stubs/Services/PerforceServiceStub.cs b/Engine/Source/Programs/Horde/Horde.Build.Tests/Stubs/Services/PerforceServiceStub.cs
--- a/Engine/Source/Programs/Horde/Horde.Build.Tests/Stubs/Services/PerforceServiceStub.cs
+++ b/Engine/Source/Programs/Horde/Horde.Build.Tests/Stubs/Services/PerforceServiceStub.cs
@@ -1,10 +1,10 @@
 // Copyright Epic Games, Inc. All Rights Reserved.
 
 using EpicGames.Perforce;
-using HordeServer.Collections;
-using HordeServer.Models;
-using HordeServer.Services;
-using HordeServer.Utilities;
+using Horde.Build.Collections;
+using Horde.Build.Models;
+using Horde.Build.Services;
+using Horde.Build.Utilities;
 using System;
 using System.Collections.Generic;
 using System.Diagnostics.CodeAnalysis;
@@ -13,9 +13,9 @@ using System.Security.Cryptography.X509Certificates;
 using System.Text;
 using System.Threading.Tasks;
 
-using StreamId = HordeServer.Utilities.StringId<HordeServer.Models.IStream>;
+using StreamId = Horde.Build.Utilities.StringId<Horde.Build.Models.IStream>;
 
-namespace HordeServerTests.Stubs.Services
+namespace Horde.Build.Tests.Stubs.Services
 {
 	using UserId = ObjectId<IUser>;
 
diff --git a/Engine/Source/Programs/Horde/Horde.Build.Tests/SwaggerTest.cs b/Engine/Source/Programs/Horde/Horde.Build.Tests/SwaggerTest.cs
--- a/Engine/Source/Programs/Horde/Horde.Build.Tests/SwaggerTest.cs
+++ b/Engine/Source/Programs/Horde/Horde.Build.Tests/SwaggerTest.cs
@@ -3,7 +3,7 @@
 using System;
 using System.Text.Json;
 using System.Threading.Tasks;
-using HordeServerTests;
+using Horde.Build.Tests;
 using Microsoft.VisualStudio.TestTools.UnitTesting;
 
 namespace Horde.Build.Tests
diff --git a/Engine/Source/Programs/Horde/Horde.Build.Tests/TestSetup.cs b/Engine/Source/Programs/Horde/Horde.Build.Tests/TestSetup.cs
--- a/Engine/Source/Programs/Horde/Horde.Build.Tests/TestSetup.cs
+++ b/Engine/Source/Programs/Horde/Horde.Build.Tests/TestSetup.cs
@@ -10,26 +10,26 @@ using EpicGames.Core;
 using EpicGames.Horde.Storage;
 using HordeCommon;
 using Horde.Build.Storage.Services;
-using HordeServer;
-using HordeServer.Collections;
-using HordeServer.Collections.Impl;
-using HordeServer.Controllers;
-using HordeServer.Jobs;
-using HordeServer.Logs;
-using HordeServer.Logs.Builder;
-using HordeServer.Logs.Readers;
-using HordeServer.Logs.Storage.Impl;
-using HordeServer.Models;
-using HordeServer.Notifications;
-using HordeServer.Notifications.Impl;
-using HordeServer.Services;
-using HordeServer.Services.Impl;
-using HordeServer.Storage;
-using HordeServer.Storage.Backends;
-using HordeServer.Tasks.Impl;
-using HordeServer.Utilities;
-using HordeServerTests.Stubs.Collections;
-using HordeServerTests.Stubs.Services;
+using Horde.Build;
+using Horde.Build.Collections;
+using Horde.Build.Collections.Impl;
+using Horde.Build.Controllers;
+using Horde.Build.Jobs;
+using Horde.Build.Logs;
+using Horde.Build.Logs.Builder;
+using Horde.Build.Logs.Readers;
+using Horde.Build.Logs.Storage.Impl;
+using Horde.Build.Models;
+using Horde.Build.Notifications;
+using Horde.Build.Notifications.Impl;
+using Horde.Build.Services;
+using Horde.Build.Services.Impl;
+using Horde.Build.Storage;
+using Horde.Build.Storage.Backends;
+using Horde.Build.Tasks.Impl;
+using Horde.Build.Utilities;
+using Horde.Build.Tests.Stubs.Collections;
+using Horde.Build.Tests.Stubs.Services;
 using Microsoft.AspNetCore.Http;
 using Microsoft.AspNetCore.Mvc;
 using Microsoft.Extensions.Configuration;
@@ -42,7 +42,7 @@ using Microsoft.Extensions.Options;
 using Moq;
 using StatsdClient;
 
-namespace HordeServerTests
+namespace Horde.Build.Tests
 {
 	/// <summary>
 	/// Handles set up of collections, services, fixtures etc during testing
diff --git a/Engine/Source/Programs/Horde/Horde.Build/Acls/Acl.cs b/Engine/Source/Programs/Horde/Horde.Build/Acls/Acl.cs
--- a/Engine/Source/Programs/Horde/Horde.Build/Acls/Acl.cs
+++ b/Engine/Source/Programs/Horde/Horde.Build/Acls/Acl.cs
@@ -1,8 +1,8 @@
 // Copyright Epic Games, Inc. All Rights Reserved.
 
-using HordeServer.Api;
-using HordeServer.Services;
-using HordeServer.Utilities;
+using Horde.Build.Api;
+using Horde.Build.Services;
+using Horde.Build.Utilities;
 using MongoDB.Bson;
 using MongoDB.Bson.Serialization;
 using MongoDB.Bson.Serialization.Attributes;
@@ -13,7 +13,7 @@ using System.Linq;
 using System.Linq.Expressions;
 using System.Security.Claims;
 
-namespace HordeServer.Models
+namespace Horde.Build.Models
 {
 	/// <summary>
 	/// Set of actions that can be performed by a user. NOTE: This enum is sensitive to ordering. Do not change values.
diff --git a/Engine/Source/Programs/Horde/Horde.Build/Acls/AclService.cs b/Engine/Source/Programs/Horde/Horde.Build/Acls/AclService.cs
--- a/Engine/Source/Programs/Horde/Horde.Build/Acls/AclService.cs
+++ b/Engine/Source/Programs/Horde/Horde.Build/Acls/AclService.cs
@@ -1,9 +1,9 @@
 // Copyright Epic Games, Inc. All Rights Reserved.
 
-using HordeServer.Api;
-using HordeServer.Collections;
-using HordeServer.Models;
-using HordeServer.Utilities;
+using Horde.Build.Api;
+using Horde.Build.Collections;
+using Horde.Build.Models;
+using Horde.Build.Utilities;
 using Microsoft.AspNetCore.Authentication.OAuth.Claims;
 using Microsoft.Extensions.Options;
 using Microsoft.IdentityModel.Tokens;
@@ -18,7 +18,7 @@ using System.Security.Claims;
 using System.Security.Cryptography;
 using System.Threading.Tasks;
 
-namespace HordeServer.Services
+namespace Horde.Build.Services
 {
 	using UserId = ObjectId<IUser>;
 	using SessionId = ObjectId<ISession>;
diff --git a/Engine/Source/Programs/Horde/Horde.Build/Agents/AgentId.cs b/Engine/Source/Programs/Horde/Horde.Build/Agents/AgentId.cs
--- a/Engine/Source/Programs/Horde/Horde.Build/Agents/AgentId.cs
+++ b/Engine/Source/Programs/Horde/Horde.Build/Agents/AgentId.cs
@@ -1,7 +1,7 @@
 // Copyright Epic Games, Inc. All Rights Reserved.
 
 using EpicGames.Redis;
-using HordeServer.Models;
+using Horde.Build.Models;
 using MongoDB.Bson.Serialization;
 using MongoDB.Bson.Serialization.Attributes;
 using MongoDB.Bson.Serialization.Serializers;
@@ -17,7 +17,7 @@ using System.Text.Json;
 using System.Text.Json.Serialization;
 using System.Threading.Tasks;
 
-namespace HordeServer.Utilities
+namespace Horde.Build.Utilities
 {
 	/// <summary>
 	/// Normalized hostname of an agent
diff --git a/Engine/Source/Programs/Horde/Horde.Build/Agents/IAgent.cs b/Engine/Source/Programs/Horde/Horde.Build/Agents/IAgent.cs
--- a/Engine/Source/Programs/Horde/Horde.Build/Agents/IAgent.cs
+++ b/Engine/Source/Programs/Horde/Horde.Build/Agents/IAgent.cs
@@ -2,11 +2,11 @@
 
 using Google.Protobuf.WellKnownTypes;
 using EpicGames.Core;
-using HordeServer.Api;
+using Horde.Build.Api;
 using HordeCommon;
 using HordeCommon.Rpc.Tasks;
-using HordeServer.Services;
-using HordeServer.Utilities;
+using Horde.Build.Services;
+using Horde.Build.Utilities;
 using MongoDB.Bson;
 using MongoDB.Bson.Serialization;
 using MongoDB.Bson.Serialization.Attributes;
@@ -19,11 +19,11 @@ using System.Security.Cryptography;
 using Serilog;
 using System.Diagnostics;
 using System.Threading.Tasks;
-using HordeServer.Collections;
+using Horde.Build.Collections;
 using EpicGames.Horde.Compute;
 using EpicGames.Horde.Common;
 
-namespace HordeServer.Models
+namespace Horde.Build.Models
 {
 	using LeaseId = ObjectId<ILease>;
 	using LogId = ObjectId<ILogFile>;
diff --git a/Engine/Source/Programs/Horde/Horde.Build/Agents/IAgentCollection.cs b/Engine/Source/Programs/Horde/Horde.Build/Agents/IAgentCollection.cs
--- a/Engine/Source/Programs/Horde/Horde.Build/Agents/IAgentCollection.cs
+++ b/Engine/Source/Programs/Horde/Horde.Build/Agents/IAgentCollection.cs
@@ -1,12 +1,12 @@
 // Copyright Epic Games, Inc. All Rights Reserved.
 
 using Google.Protobuf.WellKnownTypes;
-using HordeServer.Api;
+using Horde.Build.Api;
 using HordeCommon;
-using HordeServer.Models;
+using Horde.Build.Models;
 using HordeCommon.Rpc.Tasks;
-using HordeServer.Services;
-using HordeServer.Utilities;
+using Horde.Build.Services;
+using Horde.Build.Utilities;
 using MongoDB.Bson;
 using MongoDB.Bson.Serialization.Attributes;
 using MongoDB.Driver;
@@ -16,7 +16,7 @@ using System.Linq;
 using System.Linq.Expressions;
 using System.Threading.Tasks;
 
-namespace HordeServer.Collections
+namespace Horde.Build.Collections
 {
 	using PoolId = StringId<IPool>;
 	using AgentSoftwareChannelName = StringId<AgentSoftwareChannels>;
diff --git a/Engine/Source/Programs/Horde/Horde.Build/Agents/IAgentSoftwareChannel.cs b/Engine/Source/Programs/Horde/Horde.Build/Agents/IAgentSoftwareChannel.cs
--- a/Engine/Source/Programs/Horde/Horde.Build/Agents/IAgentSoftwareChannel.cs
+++ b/Engine/Source/Programs/Horde/Horde.Build/Agents/IAgentSoftwareChannel.cs
@@ -1,9 +1,9 @@
 // Copyright Epic Games, Inc. All Rights Reserved.
 
-using HordeServer.Api;
-using HordeServer.Collections;
-using HordeServer.Services;
-using HordeServer.Utilities;
+using Horde.Build.Api;
+using Horde.Build.Collections;
+using Horde.Build.Services;
+using Horde.Build.Utilities;
 using MongoDB.Bson;
 using MongoDB.Bson.Serialization.Attributes;
 using System;
@@ -14,7 +14,7 @@ using System.Linq;
 using System.Text.Json;
 using System.Threading.Tasks;
 
-namespace HordeServer.Models
+namespace Horde.Build.Models
 {
 	using AgentSoftwareVersion = StringId<IAgentSoftwareCollection>;
 	using AgentSoftwareChannelName = StringId<AgentSoftwareChannels>;
diff --git a/Engine/Source/Programs/Horde/Horde.Build/Agents/IAgentSoftwareCollection.cs b/Engine/Source/Programs/Horde/Horde.Build/Agents/IAgentSoftwareCollection.cs
--- a/Engine/Source/Programs/Horde/Horde.Build/Agents/IAgentSoftwareCollection.cs
+++ b/Engine/Source/Programs/Horde/Horde.Build/Agents/IAgentSoftwareCollection.cs
@@ -2,8 +2,8 @@
 
 using EpicGames.Core;
 using HordeCommon;
-using HordeServer.Models;
-using HordeServer.Services;
+using Horde.Build.Models;
+using Horde.Build.Services;
 using MongoDB.Bson;
 using MongoDB.Bson.Serialization.Attributes;
 using MongoDB.Driver;
@@ -16,7 +16,7 @@ using System.Linq;
 using System.Text.Json;
 using System.Threading.Tasks;
 
-namespace HordeServer.Collections
+namespace Horde.Build.Collections
 {
 	/// <summary>
 	/// Interface for a collection of template documents
diff --git a/Engine/Source/Programs/Horde/Horde.Build/Agents/ILease.cs b/Engine/Source/Programs/Horde/Horde.Build/Agents/ILease.cs
--- a/Engine/Source/Programs/Horde/Horde.Build/Agents/ILease.cs
+++ b/Engine/Source/Programs/Horde/Horde.Build/Agents/ILease.cs
@@ -1,8 +1,8 @@
 // Copyright Epic Games, Inc. All Rights Reserved.
 
-using HordeServer.Api;
+using Horde.Build.Api;
 using HordeCommon;
-using HordeServer.Utilities;
+using Horde.Build.Utilities;
 using MongoDB.Bson;
 using MongoDB.Bson.Serialization.Attributes;
 using Serilog;
@@ -13,7 +13,7 @@ using System.Threading.Tasks;
 using Google.Protobuf.WellKnownTypes;
 using Google.Protobuf;
 
-namespace HordeServer.Models
+namespace Horde.Build.Models
 {
 	using LeaseId = ObjectId<ILease>;
 	using LogId = ObjectId<ILogFile>;
diff --git a/Engine/Source/Programs/Horde/Horde.Build/Agents/ILeaseCollection.cs b/Engine/Source/Programs/Horde/Horde.Build/Agents/ILeaseCollection.cs
--- a/Engine/Source/Programs/Horde/Horde.Build/Agents/ILeaseCollection.cs
+++ b/Engine/Source/Programs/Horde/Horde.Build/Agents/ILeaseCollection.cs
@@ -1,9 +1,9 @@
 // Copyright Epic Games, Inc. All Rights Reserved.
 
 using HordeCommon;
-using HordeServer.Models;
-using HordeServer.Services;
-using HordeServer.Utilities;
+using Horde.Build.Models;
+using Horde.Build.Services;
+using Horde.Build.Utilities;
 using MongoDB.Bson;
 using MongoDB.Bson.Serialization.Attributes;
 using MongoDB.Driver;
@@ -12,7 +12,7 @@ using System.Collections.Generic;
 using System.Linq;
 using System.Threading.Tasks;
 
-namespace HordeServer.Collections
+namespace Horde.Build.Collections
 {
 	using LeaseId = ObjectId<ILease>;
 	using LogId = ObjectId<ILogFile>;
diff --git a/Engine/Source/Programs/Horde/Horde.Build/Agents/ISession.cs b/Engine/Source/Programs/Horde/Horde.Build/Agents/ISession.cs
--- a/Engine/Source/Programs/Horde/Horde.Build/Agents/ISession.cs
+++ b/Engine/Source/Programs/Horde/Horde.Build/Agents/ISession.cs
@@ -1,7 +1,7 @@
 // Copyright Epic Games, Inc. All Rights Reserved.
 
-using HordeServer.Api;
-using HordeServer.Utilities;
+using Horde.Build.Api;
+using Horde.Build.Utilities;
 using MongoDB.Bson;
 using MongoDB.Bson.Serialization.Attributes;
 using System;
@@ -9,7 +9,7 @@ using System.Collections.Generic;
 using System.Linq;
 using System.Threading.Tasks;
 
-namespace HordeServer.Models
+namespace Horde.Build.Models
 {
 	using SessionId = ObjectId<ISession>;
 
diff --git a/Engine/Source/Programs/Horde/Horde.Build/Agents/ISessionCollection.cs b/Engine/Source/Programs/Horde/Horde.Build/Agents/ISessionCollection.cs
--- a/Engine/Source/Programs/Horde/Horde.Build/Agents/ISessionCollection.cs
+++ b/Engine/Source/Programs/Horde/Horde.Build/Agents/ISessionCollection.cs
@@ -1,8 +1,8 @@
 // Copyright Epic Games, Inc. All Rights Reserved.
 
-using HordeServer.Models;
-using HordeServer.Services;
-using HordeServer.Utilities;
+using Horde.Build.Models;
+using Horde.Build.Services;
+using Horde.Build.Utilities;
 using Microsoft.Extensions.DependencyInjection;
 using MongoDB.Bson;
 using MongoDB.Bson.Serialization.Attributes;
@@ -12,7 +12,7 @@ using System.Collections.Generic;
 using System.Linq;
 using System.Threading.Tasks;
 
-namespace HordeServer.Collections
+namespace Horde.Build.Collections
 {
 	using SessionId = ObjectId<ISession>;
 
diff --git a/Engine/Source/Programs/Horde/Horde.Build/Agents/ITaskSource.cs b/Engine/Source/Programs/Horde/Horde.Build/Agents/ITaskSource.cs
--- a/Engine/Source/Programs/Horde/Horde.Build/Agents/ITaskSource.cs
+++ b/Engine/Source/Programs/Horde/Horde.Build/Agents/ITaskSource.cs
@@ -4,9 +4,9 @@ using Google.Protobuf;
 using Google.Protobuf.Reflection;
 using Google.Protobuf.WellKnownTypes;
 using HordeCommon;
-using HordeServer.Models;
-using HordeServer.Services;
-using HordeServer.Utilities;
+using Horde.Build.Models;
+using Horde.Build.Services;
+using Horde.Build.Utilities;
 using Microsoft.Extensions.Logging;
 using MongoDB.Bson;
 using System;
@@ -19,7 +19,7 @@ using System.Text;
 using System.Threading;
 using System.Threading.Tasks;
 
-namespace HordeServer.Tasks
+namespace Horde.Build.Tasks
 {
 	using LeaseId = ObjectId<ILease>;
 
diff --git a/Engine/Source/Programs/Horde/Horde.Build/Agents/Impl/AgentCollection.cs b/Engine/Source/Programs/Horde/Horde.Build/Agents/Impl/AgentCollection.cs
--- a/Engine/Source/Programs/Horde/Horde.Build/Agents/Impl/AgentCollection.cs
+++ b/Engine/Source/Programs/Horde/Horde.Build/Agents/Impl/AgentCollection.cs
@@ -4,9 +4,9 @@ using EpicGames.Redis;
 using Google.Protobuf.WellKnownTypes;
 using HordeCommon;
 using HordeCommon.Rpc.Tasks;
-using HordeServer.Models;
-using HordeServer.Services;
-using HordeServer.Utilities;
+using Horde.Build.Models;
+using Horde.Build.Services;
+using Horde.Build.Utilities;
 using Microsoft.Extensions.Logging;
 using MongoDB.Bson;
 using MongoDB.Bson.Serialization.Attributes;
@@ -17,7 +17,7 @@ using System.Linq;
 using System.Linq.Expressions;
 using System.Threading.Tasks;
 
-namespace HordeServer.Collections.Impl
+namespace Horde.Build.Collections.Impl
 {
 	using AgentSoftwareChannelName = StringId<AgentSoftwareChannels>;
 	using PoolId = StringId<IPool>;
diff --git a/Engine/Source/Programs/Horde/Horde.Build/Agents/Impl/AgentService.cs b/Engine/Source/Programs/Horde/Horde.Build/Agents/Impl/AgentService.cs
--- a/Engine/Source/Programs/Horde/Horde.Build/Agents/Impl/AgentService.cs
+++ b/Engine/Source/Programs/Horde/Horde.Build/Agents/Impl/AgentService.cs
@@ -5,10 +5,10 @@ using EpicGames.Redis;
 using EpicGames.Serialization;
 using Google.Protobuf.WellKnownTypes;
 using HordeCommon;
-using HordeServer.Collections;
-using HordeServer.Models;
-using HordeServer.Tasks;
-using HordeServer.Utilities;
+using Horde.Build.Collections;
+using Horde.Build.Models;
+using Horde.Build.Tasks;
+using Horde.Build.Utilities;
 using Microsoft.Extensions.Hosting;
 using Microsoft.Extensions.Logging;
 using MongoDB.Bson;
@@ -25,7 +25,7 @@ using System.Threading.Tasks;
 using OpenTracing;
 using OpenTracing.Util;
 
-namespace HordeServer.Services
+namespace Horde.Build.Services
 {
 	using AgentSoftwareChannelName = StringId<AgentSoftwareChannels>;
 	using LeaseId = ObjectId<ILease>;
diff --git a/Engine/Source/Programs/Horde/Horde.Build/Agents/Impl/AgentSoftwareCollection.cs b/Engine/Source/Programs/Horde/Horde.Build/Agents/Impl/AgentSoftwareCollection.cs
--- a/Engine/Source/Programs/Horde/Horde.Build/Agents/Impl/AgentSoftwareCollection.cs
+++ b/Engine/Source/Programs/Horde/Horde.Build/Agents/Impl/AgentSoftwareCollection.cs
@@ -2,9 +2,9 @@
 
 using EpicGames.Core;
 using HordeCommon;
-using HordeServer.Models;
-using HordeServer.Services;
-using HordeServer.Utilities;
+using Horde.Build.Models;
+using Horde.Build.Services;
+using Horde.Build.Utilities;
 using MongoDB.Bson;
 using MongoDB.Bson.Serialization.Attributes;
 using MongoDB.Driver;
@@ -17,7 +17,7 @@ using System.Linq;
 using System.Text.Json;
 using System.Threading.Tasks;
 
-namespace HordeServer.Collections.Impl
+namespace Horde.Build.Collections.Impl
 {
 	using AgentSoftwareVersion = StringId<IAgentSoftwareCollection>;
 
diff --git a/Engine/Source/Programs/Horde/Horde.Build/Agents/Impl/AgentSoftwareService.cs b/Engine/Source/Programs/Horde/Horde.Build/Agents/Impl/AgentSoftwareService.cs
--- a/Engine/Source/Programs/Horde/Horde.Build/Agents/Impl/AgentSoftwareService.cs
+++ b/Engine/Source/Programs/Horde/Horde.Build/Agents/Impl/AgentSoftwareService.cs
@@ -1,9 +1,9 @@
 // Copyright Epic Games, Inc. All Rights Reserved.
 
 using HordeCommon;
-using HordeServer.Collections;
-using HordeServer.Models;
-using HordeServer.Utilities;
+using Horde.Build.Collections;
+using Horde.Build.Models;
+using Horde.Build.Utilities;
 using Microsoft.IdentityModel.Tokens;
 using MongoDB.Bson;
 using MongoDB.Bson.Serialization.Attributes;
@@ -23,7 +23,7 @@ using System.IO.Compression;
 using System.Security.Cryptography.X509Certificates;
 using Horde.Build.Commands;
 
-namespace HordeServer.Services
+namespace Horde.Build.Services
 {
 	using AgentSoftwareChannelName = StringId<AgentSoftwareChannels>;
 
diff --git a/Engine/Source/Programs/Horde/Horde.Build/Agents/Impl/ConformList.cs b/Engine/Source/Programs/Horde/Horde.Build/Agents/Impl/ConformList.cs
--- a/Engine/Source/Programs/Horde/Horde.Build/Agents/Impl/ConformList.cs
+++ b/Engine/Source/Programs/Horde/Horde.Build/Agents/Impl/ConformList.cs
@@ -1,6 +1,6 @@
 // Copyright Epic Games, Inc. All Rights Reserved.
 
-using HordeServer.Utilities;
+using Horde.Build.Utilities;
 using MongoDB.Bson;
 using MongoDB.Bson.Serialization.Attributes;
 using System;
@@ -8,7 +8,7 @@ using System.Collections.Generic;
 using System.Linq;
 using System.Threading.Tasks;
 
-namespace HordeServer.Models
+namespace Horde.Build.Models
 {
 	using LeaseId = ObjectId<ILease>;
 
diff --git a/Engine/Source/Programs/Horde/Horde.Build/Agents/Impl/Tasks/ConformTaskSource.cs b/Engine/Source/Programs/Horde/Horde.Build/Agents/Impl/Tasks/ConformTaskSource.cs
--- a/Engine/Source/Programs/Horde/Horde.Build/Agents/Impl/Tasks/ConformTaskSource.cs
+++ b/Engine/Source/Programs/Horde/Horde.Build/Agents/Impl/Tasks/ConformTaskSource.cs
@@ -5,11 +5,11 @@ using Google.Protobuf.Reflection;
 using Google.Protobuf.WellKnownTypes;
 using HordeCommon;
 using HordeCommon.Rpc.Tasks;
-using HordeServer.Api;
-using HordeServer.Collections;
-using HordeServer.Models;
-using HordeServer.Services;
-using HordeServer.Utilities;
+using Horde.Build.Api;
+using Horde.Build.Collections;
+using Horde.Build.Models;
+using Horde.Build.Services;
+using Horde.Build.Utilities;
 using Microsoft.Extensions.Hosting;
 using Microsoft.Extensions.Logging;
 using MongoDB.Bson;
@@ -19,7 +19,7 @@ using System.Linq;
 using System.Threading;
 using System.Threading.Tasks;
 
-namespace HordeServer.Tasks.Impl
+namespace Horde.Build.Tasks.Impl
 {
 	using JobId = ObjectId<IJob>;
 	using LeaseId = ObjectId<ILease>;
diff --git a/Engine/Source/Programs/Horde/Horde.Build/Agents/Impl/Tasks/RestartTaskSource.cs b/Engine/Source/Programs/Horde/Horde.Build/Agents/Impl/Tasks/RestartTaskSource.cs
--- a/Engine/Source/Programs/Horde/Horde.Build/Agents/Impl/Tasks/RestartTaskSource.cs
+++ b/Engine/Source/Programs/Horde/Horde.Build/Agents/Impl/Tasks/RestartTaskSource.cs
@@ -5,10 +5,10 @@ using Google.Protobuf.Reflection;
 using Google.Protobuf.WellKnownTypes;
 using HordeCommon;
 using HordeCommon.Rpc.Tasks;
-using HordeServer.Api;
-using HordeServer.Models;
-using HordeServer.Services;
-using HordeServer.Utilities;
+using Horde.Build.Api;
+using Horde.Build.Models;
+using Horde.Build.Services;
+using Horde.Build.Utilities;
 using MongoDB.Bson;
 using System;
 using System.Collections.Generic;
@@ -16,7 +16,7 @@ using System.Linq;
 using System.Threading;
 using System.Threading.Tasks;
 
-namespace HordeServer.Tasks.Impl
+namespace Horde.Build.Tasks.Impl
 {
 	using JobId = ObjectId<IJob>;
 	using LeaseId = ObjectId<ILease>;
diff --git a/Engine/Source/Programs/Horde/Horde.Build/Agents/Impl/Tasks/ShutdownTaskSource.cs b/Engine/Source/Programs/Horde/Horde.Build/Agents/Impl/Tasks/ShutdownTaskSource.cs
--- a/Engine/Source/Programs/Horde/Horde.Build/Agents/Impl/Tasks/ShutdownTaskSource.cs
+++ b/Engine/Source/Programs/Horde/Horde.Build/Agents/Impl/Tasks/ShutdownTaskSource.cs
@@ -6,10 +6,10 @@ using Google.Protobuf.Reflection;
 using Google.Protobuf.WellKnownTypes;
 using HordeCommon;
 using HordeCommon.Rpc.Tasks;
-using HordeServer.Api;
-using HordeServer.Models;
-using HordeServer.Services;
-using HordeServer.Utilities;
+using Horde.Build.Api;
+using Horde.Build.Models;
+using Horde.Build.Services;
+using Horde.Build.Utilities;
 using MongoDB.Bson;
 using System;
 using System.Collections.Generic;
@@ -17,7 +17,7 @@ using System.Linq;
 using System.Threading;
 using System.Threading.Tasks;
 
-namespace HordeServer.Tasks.Impl
+namespace Horde.Build.Tasks.Impl
 {
 	using JobId = ObjectId<IJob>;
 	using LeaseId = ObjectId<ILease>;
diff --git a/Engine/Source/Programs/Horde/Horde.Build/Agents/Impl/UpgradeService.cs b/Engine/Source/Programs/Horde/Horde.Build/Agents/Impl/UpgradeService.cs
--- a/Engine/Source/Programs/Horde/Horde.Build/Agents/Impl/UpgradeService.cs
+++ b/Engine/Source/Programs/Horde/Horde.Build/Agents/Impl/UpgradeService.cs
@@ -1,8 +1,8 @@
 // Copyright Epic Games, Inc. All Rights Reserved.
 
-using HordeServer.Collections;
-using HordeServer.Collections.Impl;
-using HordeServer.Models;
+using Horde.Build.Collections;
+using Horde.Build.Collections.Impl;
+using Horde.Build.Models;
 using Microsoft.Extensions.DependencyInjection;
 using Microsoft.Extensions.Logging;
 using MongoDB.Driver;
@@ -12,7 +12,7 @@ using System.Collections.Generic;
 using System.Linq;
 using System.Threading.Tasks;
 
-namespace HordeServer.Services
+namespace Horde.Build.Services
 {
 	/// <summary>
 	/// Service to perform schema upgrades on the database
diff --git a/Engine/Source/Programs/Horde/Horde.Build/Api/Acls.cs b/Engine/Source/Programs/Horde/Horde.Build/Api/Acls.cs
--- a/Engine/Source/Programs/Horde/Horde.Build/Api/Acls.cs
+++ b/Engine/Source/Programs/Horde/Horde.Build/Api/Acls.cs
@@ -1,6 +1,6 @@
 // Copyright Epic Games, Inc. All Rights Reserved.
 
-using HordeServer.Models;
+using Horde.Build.Models;
 using System;
 using System.Collections.Generic;
 using System.ComponentModel.DataAnnotations;
@@ -8,7 +8,7 @@ using System.Linq;
 using System.Security.Claims;
 using System.Threading.Tasks;
 
-namespace HordeServer.Api
+namespace Horde.Build.Api
 {
 	/// <summary>
 	/// New claim to create
diff --git a/Engine/Source/Programs/Horde/Horde.Build/Api/AgentSoftwareChannels.cs b/Engine/Source/Programs/Horde/Horde.Build/Api/AgentSoftwareChannels.cs
--- a/Engine/Source/Programs/Horde/Horde.Build/Api/AgentSoftwareChannels.cs
+++ b/Engine/Source/Programs/Horde/Horde.Build/Api/AgentSoftwareChannels.cs
@@ -1,12 +1,12 @@
 // Copyright Epic Games, Inc. All Rights Reserved.
 
-using HordeServer.Models;
+using Horde.Build.Models;
 using System;
 using System.Collections.Generic;
 using System.Linq;
 using System.Threading.Tasks;
 
-namespace HordeServer.Api
+namespace Horde.Build.Api
 {
 	/// <summary>
 	/// Information about an agent software channel
diff --git a/Engine/Source/Programs/Horde/Horde.Build/Api/Agents.cs b/Engine/Source/Programs/Horde/Horde.Build/Api/Agents.cs
--- a/Engine/Source/Programs/Horde/Horde.Build/Api/Agents.cs
+++ b/Engine/Source/Programs/Horde/Horde.Build/Api/Agents.cs
@@ -2,15 +2,15 @@
 
 using EpicGames.Core;
 using HordeCommon;
-using HordeServer.Models;
-using HordeServer.Utilities;
+using Horde.Build.Models;
+using Horde.Build.Utilities;
 using System;
 using System.Collections.Generic;
 using System.ComponentModel.DataAnnotations;
 using System.Linq;
 using System.Threading.Tasks;
 
-namespace HordeServer.Api
+namespace Horde.Build.Api
 {
 	using PoolId = StringId<IPool>;
 
diff --git a/Engine/Source/Programs/Horde/Horde.Build/Api/Events.cs b/Engine/Source/Programs/Horde/Horde.Build/Api/Events.cs
--- a/Engine/Source/Programs/Horde/Horde.Build/Api/Events.cs
+++ b/Engine/Source/Programs/Horde/Horde.Build/Api/Events.cs
@@ -2,8 +2,8 @@
 
 using EpicGames.Core;
 using HordeCommon;
-using HordeServer.Models;
-using HordeServer.Utilities;
+using Horde.Build.Models;
+using Horde.Build.Utilities;
 using Microsoft.AspNetCore.Mvc;
 using System;
 using System.Collections.Generic;
@@ -11,7 +11,7 @@ using System.ComponentModel.DataAnnotations;
 using System.Linq;
 using System.Threading.Tasks;
 
-namespace HordeServer.Api
+namespace Horde.Build.Api
 {
 	using StreamId = StringId<IStream>;
 	using TemplateRefId = StringId<TemplateRef>;
diff --git a/Engine/Source/Programs/Horde/Horde.Build/Api/HelixSwarm.cs b/Engine/Source/Programs/Horde/Horde.Build/Api/HelixSwarm.cs
--- a/Engine/Source/Programs/Horde/Horde.Build/Api/HelixSwarm.cs
+++ b/Engine/Source/Programs/Horde/Horde.Build/Api/HelixSwarm.cs
@@ -3,7 +3,7 @@
 using System.ComponentModel.DataAnnotations;
 using System.Diagnostics.CodeAnalysis;
 
-namespace HordeServer.Api
+namespace Horde.Build.Api
 {
 	/// <summary>
 	/// Request sent by the Helix Swarm project
diff --git a/Engine/Source/Programs/Horde/Horde.Build/Api/Issues.cs b/Engine/Source/Programs/Horde/Horde.Build/Api/Issues.cs
--- a/Engine/Source/Programs/Horde/Horde.Build/Api/Issues.cs
+++ b/Engine/Source/Programs/Horde/Horde.Build/Api/Issues.cs
@@ -1,9 +1,9 @@
 // Copyright Epic Games, Inc. All Rights Reserved.
 
 using EpicGames.Core;
-using HordeServer.Models;
-using HordeServer.Services;
-using HordeServer.Utilities;
+using Horde.Build.Models;
+using Horde.Build.Services;
+using Horde.Build.Utilities;
 using MongoDB.Bson;
 using System;
 using System.Collections.Generic;
@@ -12,7 +12,7 @@ using System.Linq;
 using System.Text.RegularExpressions;
 using System.Threading.Tasks;
 
-namespace HordeServer.Api
+namespace Horde.Build.Api
 {
 	using StreamId = StringId<IStream>;
 	using TemplateRefId = StringId<TemplateRef>;
diff --git a/Engine/Source/Programs/Horde/Horde.Build/Api/Jira.cs b/Engine/Source/Programs/Horde/Horde.Build/Api/Jira.cs
--- a/Engine/Source/Programs/Horde/Horde.Build/Api/Jira.cs
+++ b/Engine/Source/Programs/Horde/Horde.Build/Api/Jira.cs
@@ -5,10 +5,10 @@ using System.Collections.Generic;
 using System.ComponentModel.DataAnnotations;
 using System.Diagnostics.CodeAnalysis;
 using System.Linq;
-using HordeServer.Models;
-using HordeServer.Services;
+using Horde.Build.Models;
+using Horde.Build.Services;
 
-namespace HordeServer.Api
+namespace Horde.Build.Api
 {
 
 	/// <summary>
diff --git a/Engine/Source/Programs/Horde/Horde.Build/Api/LogFiles.cs b/Engine/Source/Programs/Horde/Horde.Build/Api/LogFiles.cs
--- a/Engine/Source/Programs/Horde/Horde.Build/Api/LogFiles.cs
+++ b/Engine/Source/Programs/Horde/Horde.Build/Api/LogFiles.cs
@@ -1,8 +1,8 @@
 // Copyright Epic Games, Inc. All Rights Reserved.
 
-using HordeServer.Logs;
-using HordeServer.Models;
-using HordeServer.Services;
+using Horde.Build.Logs;
+using Horde.Build.Models;
+using Horde.Build.Services;
 using System;
 using System.Collections.Generic;
 using System.ComponentModel.DataAnnotations;
@@ -10,7 +10,7 @@ using System.Linq;
 using System.Text.Json.Serialization;
 using System.Threading.Tasks;
 
-namespace HordeServer.Api
+namespace Horde.Build.Api
 {
 	/// <summary>
 	/// The type of data stored in this log file
diff --git a/Engine/Source/Programs/Horde/Horde.Build/Api/Server.cs b/Engine/Source/Programs/Horde/Horde.Build/Api/Server.cs
--- a/Engine/Source/Programs/Horde/Horde.Build/Api/Server.cs
+++ b/Engine/Source/Programs/Horde/Horde.Build/Api/Server.cs
@@ -5,7 +5,7 @@ using System.Runtime.InteropServices;
 using System.Reflection;
 using System.Diagnostics;
 
-namespace HordeServer.Api
+namespace Horde.Build.Api
 {
 	/// <summary>
 	/// Server Info
