Horde: Normalize namespace names to match assembly names.

diff --git a/Engine/Source/Programs/Horde/Horde.Build/Api/UgsMetadata.cs b/Engine/Source/Programs/Horde/Horde.Build/Api/UgsMetadata.cs
--- a/Engine/Source/Programs/Horde/Horde.Build/Api/UgsMetadata.cs
+++ b/Engine/Source/Programs/Horde/Horde.Build/Api/UgsMetadata.cs
@@ -1,9 +1,9 @@
 // Copyright Epic Games, Inc. All Rights Reserved.
 
 using EpicGames.Core;
-using HordeServer.Models;
-using HordeServer.Services;
-using HordeServer.Utilities;
+using Horde.Build.Models;
+using Horde.Build.Services;
+using Horde.Build.Utilities;
 using System;
 using System.Collections.Generic;
 using System.ComponentModel.DataAnnotations;
@@ -11,7 +11,7 @@ using System.Linq;
 using System.Text.Json.Serialization;
 using System.Threading.Tasks;
 
-namespace HordeServer.Api
+namespace Horde.Build.Api
 {
 	/// <summary>
 	/// Review by a user of a particular change
diff --git a/Engine/Source/Programs/Horde/Horde.Build/Api/Users.cs b/Engine/Source/Programs/Horde/Horde.Build/Api/Users.cs
--- a/Engine/Source/Programs/Horde/Horde.Build/Api/Users.cs
+++ b/Engine/Source/Programs/Horde/Horde.Build/Api/Users.cs
@@ -1,8 +1,8 @@
 // Copyright Epic Games, Inc. All Rights Reserved.
 
 using EpicGames.Core;
-using HordeServer.Models;
-using HordeServer.Services;
+using Horde.Build.Models;
+using Horde.Build.Services;
 using MongoDB.Bson;
 using System;
 using System.Collections.Generic;
@@ -10,7 +10,7 @@ using System.Linq;
 using System.Text.Json;
 using System.Threading.Tasks;
 
-namespace HordeServer.Api
+namespace Horde.Build.Api
 {
 	/// <summary>
 	/// Response describing the current user
diff --git a/Engine/Source/Programs/Horde/Horde.Build/Auditing/IAuditLog.cs b/Engine/Source/Programs/Horde/Horde.Build/Auditing/IAuditLog.cs
--- a/Engine/Source/Programs/Horde/Horde.Build/Auditing/IAuditLog.cs
+++ b/Engine/Source/Programs/Horde/Horde.Build/Auditing/IAuditLog.cs
@@ -1,12 +1,12 @@
 // Copyright Epic Games, Inc. All Rights Reserved.
 
 using Google.Protobuf.WellKnownTypes;
-using HordeServer.Api;
+using Horde.Build.Api;
 using HordeCommon;
-using HordeServer.Models;
+using Horde.Build.Models;
 using HordeCommon.Rpc.Tasks;
-using HordeServer.Services;
-using HordeServer.Utilities;
+using Horde.Build.Services;
+using Horde.Build.Utilities;
 using MongoDB.Bson;
 using MongoDB.Bson.Serialization.Attributes;
 using MongoDB.Driver;
@@ -22,7 +22,7 @@ using System.Text.Json;
 using System.Text;
 using System.IO.Pipelines;
 
-namespace HordeServer.Collections
+namespace Horde.Build.Collections
 {
 	/// <summary>
 	/// Message from an audit log
diff --git a/Engine/Source/Programs/Horde/Horde.Build/Authentication/AnonymousHandler.cs b/Engine/Source/Programs/Horde/Horde.Build/Authentication/AnonymousHandler.cs
--- a/Engine/Source/Programs/Horde/Horde.Build/Authentication/AnonymousHandler.cs
+++ b/Engine/Source/Programs/Horde/Horde.Build/Authentication/AnonymousHandler.cs
@@ -1,9 +1,9 @@
 // Copyright Epic Games, Inc. All Rights Reserved.
 
-using HordeServer.Collections;
-using HordeServer.Models;
-using HordeServer.Services;
-using HordeServer.Utilities;
+using Horde.Build.Collections;
+using Horde.Build.Models;
+using Horde.Build.Services;
+using Horde.Build.Utilities;
 using Microsoft.AspNetCore.Authentication;
 using Microsoft.Extensions.Logging;
 using Microsoft.Extensions.Options;
@@ -14,7 +14,7 @@ using System.Security.Claims;
 using System.Text.Encodings.Web;
 using System.Threading.Tasks;
 
-namespace HordeServer.Authentication
+namespace Horde.Build.Authentication
 {
 	class AnonymousAuthenticationOptions : AuthenticationSchemeOptions
 	{
diff --git a/Engine/Source/Programs/Horde/Horde.Build/Authentication/HordeJwtBearerHandler.cs b/Engine/Source/Programs/Horde/Horde.Build/Authentication/HordeJwtBearerHandler.cs
--- a/Engine/Source/Programs/Horde/Horde.Build/Authentication/HordeJwtBearerHandler.cs
+++ b/Engine/Source/Programs/Horde/Horde.Build/Authentication/HordeJwtBearerHandler.cs
@@ -1,7 +1,7 @@
 // Copyright Epic Games, Inc. All Rights Reserved.
 
-using HordeServer.Services;
-using HordeServer.Utilities;
+using Horde.Build.Services;
+using Horde.Build.Utilities;
 using Microsoft.AspNetCore.Authentication;
 using Microsoft.AspNetCore.Authentication.JwtBearer;
 using Microsoft.Extensions.Logging;
@@ -13,7 +13,7 @@ using System.Linq;
 using System.Text.Encodings.Web;
 using System.Threading.Tasks;
 
-namespace HordeServer.Authentication
+namespace Horde.Build.Authentication
 {
 	/// <summary>
 	/// JWT handler for server-issued bearer tokens. These tokens are signed using a randomly generated key per DB instance.
diff --git a/Engine/Source/Programs/Horde/Horde.Build/Authentication/HordeOpenIdConnectHandler.cs b/Engine/Source/Programs/Horde/Horde.Build/Authentication/HordeOpenIdConnectHandler.cs
--- a/Engine/Source/Programs/Horde/Horde.Build/Authentication/HordeOpenIdConnectHandler.cs
+++ b/Engine/Source/Programs/Horde/Horde.Build/Authentication/HordeOpenIdConnectHandler.cs
@@ -1,8 +1,8 @@
 // Copyright Epic Games, Inc. All Rights Reserved.
 
-using HordeServer.Collections;
-using HordeServer.Models;
-using HordeServer.Utilities;
+using Horde.Build.Collections;
+using Horde.Build.Models;
+using Horde.Build.Utilities;
 using Microsoft.AspNetCore.Authentication;
 using Microsoft.AspNetCore.Authentication.OAuth.Claims;
 using Microsoft.AspNetCore.Authentication.OpenIdConnect;
@@ -18,7 +18,7 @@ using System.Text.Encodings.Web;
 using System.Text.Json;
 using System.Threading.Tasks;
 
-namespace HordeServer.Authentication
+namespace Horde.Build.Authentication
 {
 	class HordeOpenIdConnectHandler : OpenIdConnectHandler
 	{
diff --git a/Engine/Source/Programs/Horde/Horde.Build/Authentication/OktaHandler.cs b/Engine/Source/Programs/Horde/Horde.Build/Authentication/OktaHandler.cs
--- a/Engine/Source/Programs/Horde/Horde.Build/Authentication/OktaHandler.cs
+++ b/Engine/Source/Programs/Horde/Horde.Build/Authentication/OktaHandler.cs
@@ -1,8 +1,8 @@
 // Copyright Epic Games, Inc. All Rights Reserved.
 
-using HordeServer.Collections;
-using HordeServer.Models;
-using HordeServer.Utilities;
+using Horde.Build.Collections;
+using Horde.Build.Models;
+using Horde.Build.Utilities;
 using Microsoft.AspNetCore.Authentication;
 using Microsoft.AspNetCore.Authentication.OAuth.Claims;
 using Microsoft.AspNetCore.Authentication.OpenIdConnect;
@@ -19,7 +19,7 @@ using System.Text.Encodings.Web;
 using System.Text.Json;
 using System.Threading.Tasks;
 
-namespace HordeServer.Authentication
+namespace Horde.Build.Authentication
 {
 	class OktaDefaults
 	{
diff --git a/Engine/Source/Programs/Horde/Horde.Build/Authentication/ServiceAccountAuthHandler.cs b/Engine/Source/Programs/Horde/Horde.Build/Authentication/ServiceAccountAuthHandler.cs
--- a/Engine/Source/Programs/Horde/Horde.Build/Authentication/ServiceAccountAuthHandler.cs
+++ b/Engine/Source/Programs/Horde/Horde.Build/Authentication/ServiceAccountAuthHandler.cs
@@ -9,11 +9,11 @@ using System.Linq;
 using System.Security.Claims;
 using System.Text.Encodings.Web;
 using System.Threading.Tasks;
-using HordeServer.Collections;
-using HordeServer.Models;
+using Horde.Build.Collections;
+using Horde.Build.Models;
 using Microsoft.Net.Http.Headers;
 
-namespace HordeServer.Authentication
+namespace Horde.Build.Authentication
 {
 	class ServiceAccountAuthOptions : AuthenticationSchemeOptions
 	{
diff --git a/Engine/Source/Programs/Horde/Horde.Build/Commands/Bundles/CreateCommand.cs b/Engine/Source/Programs/Horde/Horde.Build/Commands/Bundles/CreateCommand.cs
--- a/Engine/Source/Programs/Horde/Horde.Build/Commands/Bundles/CreateCommand.cs
+++ b/Engine/Source/Programs/Horde/Horde.Build/Commands/Bundles/CreateCommand.cs
@@ -16,7 +16,7 @@ using System.IO;
 using System.Linq;
 using System.Threading.Tasks;
 
-namespace HordeServer.Commands.Bundles
+namespace Horde.Build.Commands.Bundles
 {
 	abstract class BundleCommandBase : Command
 	{
diff --git a/Engine/Source/Programs/Horde/Horde.Build/Commands/Bundles/DeltaCommand.cs b/Engine/Source/Programs/Horde/Horde.Build/Commands/Bundles/DeltaCommand.cs
--- a/Engine/Source/Programs/Horde/Horde.Build/Commands/Bundles/DeltaCommand.cs
+++ b/Engine/Source/Programs/Horde/Horde.Build/Commands/Bundles/DeltaCommand.cs
@@ -15,7 +15,7 @@ using System.Linq;
 using System.Text;
 using System.Threading.Tasks;
 
-namespace HordeServer.Commands.Bundles
+namespace Horde.Build.Commands.Bundles
 {
 	[Command("bundle", "delta", "Desribe the delta between two bundles")]
 	class DeltaCommand : BundleCommandBase
diff --git a/Engine/Source/Programs/Horde/Horde.Build/Commands/Config/SchemasCommand.cs b/Engine/Source/Programs/Horde/Horde.Build/Commands/Config/SchemasCommand.cs
--- a/Engine/Source/Programs/Horde/Horde.Build/Commands/Config/SchemasCommand.cs
+++ b/Engine/Source/Programs/Horde/Horde.Build/Commands/Config/SchemasCommand.cs
@@ -1,13 +1,13 @@
 // Copyright Epic Games, Inc. All Rights Reserved.
 
 using EpicGames.Core;
-using HordeServer.Controllers;
-using HordeServer.Utilities;
+using Horde.Build.Controllers;
+using Horde.Build.Utilities;
 using Microsoft.Extensions.Logging;
 using System;
 using System.Threading.Tasks;
 
-namespace HordeServer.Commands.Config
+namespace Horde.Build.Commands.Config
 {
 	[Command("config", "schemas", "Writes JSON schemas for server settings")]
 	class SchemasCommand : Command
diff --git a/Engine/Source/Programs/Horde/Horde.Build/Commands/Perforce/ReplicateCommand.cs b/Engine/Source/Programs/Horde/Horde.Build/Commands/Perforce/ReplicateCommand.cs
--- a/Engine/Source/Programs/Horde/Horde.Build/Commands/Perforce/ReplicateCommand.cs
+++ b/Engine/Source/Programs/Horde/Horde.Build/Commands/Perforce/ReplicateCommand.cs
@@ -4,12 +4,12 @@ using EpicGames.Core;
 using EpicGames.Horde.Storage;
 using EpicGames.Horde.Storage.Impl;
 using EpicGames.Serialization;
-using HordeServer.Collections;
-using HordeServer.Commits;
-using HordeServer.Commits.Impl;
-using HordeServer.Controllers;
-using HordeServer.Models;
-using HordeServer.Utilities;
+using Horde.Build.Collections;
+using Horde.Build.Commits;
+using Horde.Build.Commits.Impl;
+using Horde.Build.Controllers;
+using Horde.Build.Models;
+using Horde.Build.Utilities;
 using Microsoft.Extensions.Configuration;
 using Microsoft.Extensions.DependencyInjection;
 using Microsoft.Extensions.Logging;
@@ -18,7 +18,7 @@ using System.Collections.Generic;
 using System.Linq;
 using System.Threading.Tasks;
 
-namespace HordeServer.Commands
+namespace Horde.Build.Commands
 {
 	using StreamId = StringId<IStream>;
 
diff --git a/Engine/Source/Programs/Horde/Horde.Build/Commands/ServerCommand.cs b/Engine/Source/Programs/Horde/Horde.Build/Commands/ServerCommand.cs
--- a/Engine/Source/Programs/Horde/Horde.Build/Commands/ServerCommand.cs
+++ b/Engine/Source/Programs/Horde/Horde.Build/Commands/ServerCommand.cs
@@ -1,7 +1,7 @@
 // Copyright Epic Games, Inc. All Rights Reserved.
 
 using EpicGames.Core;
-using HordeServer;
+using Horde.Build;
 using Microsoft.AspNetCore.Hosting;
 using Microsoft.AspNetCore.Server.Kestrel.Core;
 using Microsoft.Extensions.Configuration;
diff --git a/Engine/Source/Programs/Horde/Horde.Build/Compute/IComputeService.cs b/Engine/Source/Programs/Horde/Horde.Build/Compute/IComputeService.cs
--- a/Engine/Source/Programs/Horde/Horde.Build/Compute/IComputeService.cs
+++ b/Engine/Source/Programs/Horde/Horde.Build/Compute/IComputeService.cs
@@ -4,10 +4,10 @@ using EpicGames.Core;
 using EpicGames.Horde.Compute;
 using EpicGames.Horde.Storage;
 using EpicGames.Serialization;
-using HordeServer.Models;
-using HordeServer.Storage;
-using HordeServer.Tasks;
-using HordeServer.Utilities;
+using Horde.Build.Models;
+using Horde.Build.Storage;
+using Horde.Build.Tasks;
+using Horde.Build.Utilities;
 using Microsoft.Extensions.Hosting;
 using MongoDB.Bson;
 using System;
@@ -17,7 +17,7 @@ using System.Linq;
 using System.Threading;
 using System.Threading.Tasks;
 
-namespace HordeServer.Compute
+namespace Horde.Build.Compute
 {
 	using LeaseId = ObjectId<ILease>;
 
diff --git a/Engine/Source/Programs/Horde/Horde.Build/Compute/Impl/ComputeController.cs b/Engine/Source/Programs/Horde/Horde.Build/Compute/Impl/ComputeController.cs
--- a/Engine/Source/Programs/Horde/Horde.Build/Compute/Impl/ComputeController.cs
+++ b/Engine/Source/Programs/Horde/Horde.Build/Compute/Impl/ComputeController.cs
@@ -3,9 +3,9 @@
 using EpicGames.Core;
 using EpicGames.Horde.Compute;
 using EpicGames.Horde.Compute.Impl;
-using HordeServer.Models;
-using HordeServer.Services;
-using HordeServer.Utilities;
+using Horde.Build.Models;
+using Horde.Build.Services;
+using Horde.Build.Utilities;
 using Microsoft.AspNetCore.Authorization;
 using Microsoft.AspNetCore.Mvc;
 using System;
@@ -15,7 +15,7 @@ using System.Linq;
 using System.Threading;
 using System.Threading.Tasks;
 
-namespace HordeServer.Compute.Impl
+namespace Horde.Build.Compute.Impl
 {
 	/// <summary>
 	/// Controller for the /api/v1/compute endpoint
diff --git a/Engine/Source/Programs/Horde/Horde.Build/Compute/Impl/ComputeService.cs b/Engine/Source/Programs/Horde/Horde.Build/Compute/Impl/ComputeService.cs
--- a/Engine/Source/Programs/Horde/Horde.Build/Compute/Impl/ComputeService.cs
+++ b/Engine/Source/Programs/Horde/Horde.Build/Compute/Impl/ComputeService.cs
@@ -7,9 +7,9 @@ using Google.Protobuf.WellKnownTypes;
 using Grpc.Core;
 using HordeCommon;
 using HordeCommon.Rpc.Tasks;
-using HordeServer.Api;
-using HordeServer.Models;
-using HordeServer.Services;
+using Horde.Build.Api;
+using Horde.Build.Models;
+using Horde.Build.Services;
 using Microsoft.Extensions.Logging;
 using MongoDB.Bson;
 using System;
@@ -21,9 +21,9 @@ using System.Linq;
 using System.Threading;
 using System.Threading.Tasks;
 using Microsoft.Extensions.Options;
-using HordeServer.Tasks;
-using HordeServer.Utilities;
-using HordeServer.Storage;
+using Horde.Build.Tasks;
+using Horde.Build.Utilities;
+using Horde.Build.Storage;
 using StackExchange.Redis;
 using MongoDB.Bson.Serialization.Attributes;
 using MongoDB.Bson.Serialization;
@@ -33,12 +33,12 @@ using EpicGames.Horde.Compute;
 using Microsoft.Extensions.Hosting;
 using System.Threading.Channels;
 using EpicGames.Redis;
-using HordeServer.Collections;
+using Horde.Build.Collections;
 using EpicGames.Horde.Storage;
 using EpicGames.Horde.Common;
 using System.Text.Json;
 
-namespace HordeServer.Compute.Impl
+namespace Horde.Build.Compute.Impl
 {
 	using LeaseId = ObjectId<ILease>;
 
diff --git a/Engine/Source/Programs/Horde/Horde.Build/Compute/Impl/RedisMessageQueue.cs b/Engine/Source/Programs/Horde/Horde.Build/Compute/Impl/RedisMessageQueue.cs
--- a/Engine/Source/Programs/Horde/Horde.Build/Compute/Impl/RedisMessageQueue.cs
+++ b/Engine/Source/Programs/Horde/Horde.Build/Compute/Impl/RedisMessageQueue.cs
@@ -1,7 +1,7 @@
 // Copyright Epic Games, Inc. All Rights Reserved.
 
 using EpicGames.Redis;
-using HordeServer.Compute.Impl;
+using Horde.Build.Compute.Impl;
 using Microsoft.Extensions.Logging;
 using StackExchange.Redis;
 using System;
@@ -11,7 +11,7 @@ using System.Linq;
 using System.Threading;
 using System.Threading.Tasks;
 
-namespace HordeServer.Compute.Impl
+namespace Horde.Build.Compute.Impl
 {
 	/// <summary>
 	/// Interface for a distributed message queue
diff --git a/Engine/Source/Programs/Horde/Horde.Build/Compute/Impl/RedisTaskScheduler.cs b/Engine/Source/Programs/Horde/Horde.Build/Compute/Impl/RedisTaskScheduler.cs
--- a/Engine/Source/Programs/Horde/Horde.Build/Compute/Impl/RedisTaskScheduler.cs
+++ b/Engine/Source/Programs/Horde/Horde.Build/Compute/Impl/RedisTaskScheduler.cs
@@ -4,10 +4,10 @@ using EpicGames.Core;
 using EpicGames.Horde.Compute;
 using EpicGames.Redis;
 using EpicGames.Serialization;
-using HordeServer.Models;
-using HordeServer.Services;
-using HordeServer.Storage;
-using HordeServer.Utilities;
+using Horde.Build.Models;
+using Horde.Build.Services;
+using Horde.Build.Storage;
+using Horde.Build.Utilities;
 using Microsoft.Extensions.Caching.Memory;
 using Microsoft.Extensions.Logging;
 using StackExchange.Redis;
@@ -20,7 +20,7 @@ using System.Threading;
 using System.Threading.Channels;
 using System.Threading.Tasks;
 
-namespace HordeServer.Compute.Impl
+namespace Horde.Build.Compute.Impl
 {
 	using Condition = StackExchange.Redis.Condition;
 
diff --git a/Engine/Source/Programs/Horde/Horde.Build/Controllers/AccountController.cs b/Engine/Source/Programs/Horde/Horde.Build/Controllers/AccountController.cs
--- a/Engine/Source/Programs/Horde/Horde.Build/Controllers/AccountController.cs
+++ b/Engine/Source/Programs/Horde/Horde.Build/Controllers/AccountController.cs
@@ -1,8 +1,8 @@
 // Copyright Epic Games, Inc. All Rights Reserved.
 
-using HordeServer.Authentication;
-using HordeServer.Models;
-using HordeServer.Services;
+using Horde.Build.Authentication;
+using Horde.Build.Models;
+using Horde.Build.Services;
 using Microsoft.AspNetCore.Authentication;
 using Microsoft.AspNetCore.Authentication.Cookies;
 using Microsoft.AspNetCore.Authentication.OpenIdConnect;
@@ -17,7 +17,7 @@ using System.Text;
 using System.Threading.Tasks;
 using Microsoft.Extensions.Options;
 
-namespace HordeServer.Controllers
+namespace Horde.Build.Controllers
 {
 	/// <summary>
 	/// Controller managing account status
diff --git a/Engine/Source/Programs/Horde/Horde.Build/Controllers/AdminController.cs b/Engine/Source/Programs/Horde/Horde.Build/Controllers/AdminController.cs
--- a/Engine/Source/Programs/Horde/Horde.Build/Controllers/AdminController.cs
+++ b/Engine/Source/Programs/Horde/Horde.Build/Controllers/AdminController.cs
@@ -1,12 +1,12 @@
 // Copyright Epic Games, Inc. All Rights Reserved.
 
 using EpicGames.Core;
-using HordeServer.Api;
-using HordeServer.Collections;
+using Horde.Build.Api;
+using Horde.Build.Collections;
 using HordeCommon;
-using HordeServer.Models;
-using HordeServer.Services;
-using HordeServer.Utilities;
+using Horde.Build.Models;
+using Horde.Build.Services;
+using Horde.Build.Utilities;
 using Microsoft.AspNetCore.Authorization;
 using Microsoft.AspNetCore.Authorization.Infrastructure;
 using Microsoft.AspNetCore.Mvc;
@@ -21,7 +21,7 @@ using System.Runtime.InteropServices;
 using System.Security.Claims;
 using System.Threading.Tasks;
 
-namespace HordeServer.Controllers
+namespace Horde.Build.Controllers
 {
 	/// <summary>
 	/// Object containing settings for the server
diff --git a/Engine/Source/Programs/Horde/Horde.Build/Controllers/AgentSoftwareController.cs b/Engine/Source/Programs/Horde/Horde.Build/Controllers/AgentSoftwareController.cs
--- a/Engine/Source/Programs/Horde/Horde.Build/Controllers/AgentSoftwareController.cs
+++ b/Engine/Source/Programs/Horde/Horde.Build/Controllers/AgentSoftwareController.cs
@@ -1,7 +1,7 @@
 // Copyright Epic Games, Inc. All Rights Reserved.
 
-using HordeServer.Api;
-using HordeServer.Services;
+using Horde.Build.Api;
+using Horde.Build.Services;
 using Microsoft.AspNetCore.Authorization;
 using Microsoft.AspNetCore.Http;
 using Microsoft.AspNetCore.Mvc;
@@ -11,11 +11,11 @@ using System.IO;
 using System.Collections.Generic;
 using System.Linq;
 using System.Threading.Tasks;
-using HordeServer.Models;
+using Horde.Build.Models;
 using MongoDB.Bson;
-using HordeServer.Utilities;
+using Horde.Build.Utilities;
 
-namespace HordeServer.Controllers
+namespace Horde.Build.Controllers
 {
 	using AgentSoftwareChannelName = StringId<AgentSoftwareChannels>;
 
diff --git a/Engine/Source/Programs/Horde/Horde.Build/Controllers/AgentsController.cs b/Engine/Source/Programs/Horde/Horde.Build/Controllers/AgentsController.cs
--- a/Engine/Source/Programs/Horde/Horde.Build/Controllers/AgentsController.cs
+++ b/Engine/Source/Programs/Horde/Horde.Build/Controllers/AgentsController.cs
@@ -1,10 +1,10 @@
 // Copyright Epic Games, Inc. All Rights Reserved.
 
-using HordeServer.Api;
-using HordeServer.Collections;
-using HordeServer.Models;
-using HordeServer.Services;
-using HordeServer.Utilities;
+using Horde.Build.Api;
+using Horde.Build.Collections;
+using Horde.Build.Models;
+using Horde.Build.Services;
+using Horde.Build.Utilities;
 using Microsoft.AspNetCore.Authorization;
 using Microsoft.AspNetCore.Mvc;
 using Microsoft.Extensions.Logging;
@@ -15,7 +15,7 @@ using System.Linq;
 using System.Text;
 using System.Threading.Tasks;
 
-namespace HordeServer.Controllers
+namespace Horde.Build.Controllers
 {
 	using AgentSoftwareChannelName = StringId<AgentSoftwareChannels>;
 	using LeaseId = ObjectId<ILease>;
diff --git a/Engine/Source/Programs/Horde/Horde.Build/Controllers/ArtifactsController.cs b/Engine/Source/Programs/Horde/Horde.Build/Controllers/ArtifactsController.cs
--- a/Engine/Source/Programs/Horde/Horde.Build/Controllers/ArtifactsController.cs
+++ b/Engine/Source/Programs/Horde/Horde.Build/Controllers/ArtifactsController.cs
@@ -1,11 +1,11 @@
 // Copyright Epic Games, Inc. All Rights Reserved.
 
 using Amazon.S3.Transfer;
-using HordeServer.Api;
-using HordeServer.Jobs;
-using HordeServer.Models;
-using HordeServer.Services;
-using HordeServer.Utilities;
+using Horde.Build.Api;
+using Horde.Build.Jobs;
+using Horde.Build.Models;
+using Horde.Build.Services;
+using Horde.Build.Utilities;
 using Microsoft.AspNetCore.Authorization;
 using Microsoft.AspNetCore.Http;
 using Microsoft.AspNetCore.Mvc;
@@ -25,7 +25,7 @@ using System.Security.Claims;
 using System.Text;
 using System.Threading.Tasks;
 
-namespace HordeServer.Controllers
+namespace Horde.Build.Controllers
 {
 	using JobId = ObjectId<IJob>;
 
diff --git a/Engine/Source/Programs/Horde/Horde.Build/Controllers/ConfigController.cs b/Engine/Source/Programs/Horde/Horde.Build/Controllers/ConfigController.cs
--- a/Engine/Source/Programs/Horde/Horde.Build/Controllers/ConfigController.cs
+++ b/Engine/Source/Programs/Horde/Horde.Build/Controllers/ConfigController.cs
@@ -1,9 +1,9 @@
 // Copyright Epic Games, Inc. All Rights Reserved.
 
 using EpicGames.Core;
-using HordeServer.Api;
-using HordeServer.Models;
-using HordeServer.Services;
+using Horde.Build.Api;
+using Horde.Build.Models;
+using Horde.Build.Services;
 using Microsoft.AspNetCore.Authorization;
 using Microsoft.AspNetCore.Mvc;
 using Microsoft.Extensions.Options;
@@ -11,9 +11,9 @@ using System.Threading.Tasks;
 using System.Text.Json;
 using Microsoft.Extensions.Logging;
 using System.Runtime.InteropServices;
-using HordeServer.Utilities;
+using Horde.Build.Utilities;
 
-namespace HordeServer.Controllers
+namespace Horde.Build.Controllers
 {
 
 	/// <summary>
diff --git a/Engine/Source/Programs/Horde/Horde.Build/Controllers/CredentialsController.cs b/Engine/Source/Programs/Horde/Horde.Build/Controllers/CredentialsController.cs
--- a/Engine/Source/Programs/Horde/Horde.Build/Controllers/CredentialsController.cs
+++ b/Engine/Source/Programs/Horde/Horde.Build/Controllers/CredentialsController.cs
@@ -1,9 +1,9 @@
 // Copyright Epic Games, Inc. All Rights Reserved.
 
-using HordeServer.Api;
-using HordeServer.Models;
-using HordeServer.Services;
-using HordeServer.Utilities;
+using Horde.Build.Api;
+using Horde.Build.Models;
+using Horde.Build.Services;
+using Horde.Build.Utilities;
 using Microsoft.AspNetCore.Authorization;
 using Microsoft.AspNetCore.Mvc;
 using MongoDB.Bson;
@@ -13,7 +13,7 @@ using System.ComponentModel.DataAnnotations;
 using System.Linq;
 using System.Threading.Tasks;
 
-namespace HordeServer.Controllers
+namespace Horde.Build.Controllers
 {
 	/// <summary>
 	/// Controller for the /api/v1/credentials endpoint
diff --git a/Engine/Source/Programs/Horde/Horde.Build/Controllers/DashboardController.cs b/Engine/Source/Programs/Horde/Horde.Build/Controllers/DashboardController.cs
--- a/Engine/Source/Programs/Horde/Horde.Build/Controllers/DashboardController.cs
+++ b/Engine/Source/Programs/Horde/Horde.Build/Controllers/DashboardController.cs
@@ -2,14 +2,14 @@
 
 // Copyright Epic Games, Inc. All Rights Reserved.	
 
-using HordeServer.Authentication;
+using Horde.Build.Authentication;
 using Microsoft.AspNetCore.Authentication;
 using Microsoft.AspNetCore.Authentication.Cookies;
 using Microsoft.AspNetCore.Authorization;
 using Microsoft.AspNetCore.Mvc;
 using System.Threading.Tasks;
-using HordeServer;
-using HordeServer.Controllers;
+using Horde.Build;
+using Horde.Build.Controllers;
 using Microsoft.Extensions.Options;
 using System;
 using System.Text;
diff --git a/Engine/Source/Programs/Horde/Horde.Build/Controllers/DebugController.cs b/Engine/Source/Programs/Horde/Horde.Build/Controllers/DebugController.cs
--- a/Engine/Source/Programs/Horde/Horde.Build/Controllers/DebugController.cs
+++ b/Engine/Source/Programs/Horde/Horde.Build/Controllers/DebugController.cs
@@ -4,13 +4,13 @@
 //#define ENABLE_SECURE_DEBUG_CONTROLLER
 
 using EpicGames.Core;
-using HordeServer.Api;
+using Horde.Build.Api;
 using HordeCommon;
-using HordeServer.Collections;
-using HordeServer.Models;
-using HordeServer.Services;
-using HordeServer.Tasks.Impl;
-using HordeServer.Utilities;
+using Horde.Build.Collections;
+using Horde.Build.Models;
+using Horde.Build.Services;
+using Horde.Build.Tasks.Impl;
+using Horde.Build.Utilities;
 using Microsoft.AspNetCore.Authorization;
 using Microsoft.AspNetCore.Mvc;
 using Microsoft.Extensions.Primitives;
@@ -31,9 +31,9 @@ using System.Web;
 using MongoDB.Driver;
 using System.Text.Json;
 using System.Threading;
-using HordeServer.Services.Impl;
+using Horde.Build.Services.Impl;
 using Microsoft.Net.Http.Headers;
-using HordeServer.Storage;
+using Horde.Build.Storage;
 using Microsoft.Extensions.Hosting;
 using Microsoft.Extensions.Options;
 using StatsdClient;
@@ -41,14 +41,14 @@ using StatsdClient;
 using Amazon.S3;
 using JetBrains.Profiler.SelfApi;
 
-namespace HordeServer.Controllers
+namespace Horde.Build.Controllers
 {
 	using LogId = ObjectId<ILogFile>;
 	using PoolId = StringId<IPool>;
 	using ProjectId = StringId<IProject>;
 	using TemplateRefId = StringId<TemplateRef>;
 	using AgentSoftwareVersion = StringId<IAgentSoftwareCollection>;
-	using IStream = HordeServer.Models.IStream;
+	using IStream = Horde.Build.Models.IStream;
 	using P4Debugging = Perforce.P4.P4Debugging;
 
 #if ENABLE_PUBLIC_DEBUG_CONTROLLER
diff --git a/Engine/Source/Programs/Horde/Horde.Build/Controllers/DevicesController.cs b/Engine/Source/Programs/Horde/Horde.Build/Controllers/DevicesController.cs
--- a/Engine/Source/Programs/Horde/Horde.Build/Controllers/DevicesController.cs
+++ b/Engine/Source/Programs/Horde/Horde.Build/Controllers/DevicesController.cs
@@ -1,10 +1,10 @@
 // Copyright Epic Games, Inc. All Rights Reserved.
 
-using HordeServer.Api;
-using HordeServer.Collections;
-using HordeServer.Models;
-using HordeServer.Services;
-using HordeServer.Utilities;
+using Horde.Build.Api;
+using Horde.Build.Collections;
+using Horde.Build.Models;
+using Horde.Build.Services;
+using Horde.Build.Utilities;
 using Microsoft.AspNetCore.Authorization;
 using Microsoft.AspNetCore.Mvc;
 using Microsoft.Extensions.Logging;
@@ -16,7 +16,7 @@ using System.Linq;
 using System.Text.Json;
 using System.Threading.Tasks;
 
-namespace HordeServer.Controllers
+namespace Horde.Build.Controllers
 {
 	using DeviceId = StringId<IDevice>;
 	using DevicePlatformId = StringId<IDevicePlatform>;
diff --git a/Engine/Source/Programs/Horde/Horde.Build/Controllers/ErrorController.cs b/Engine/Source/Programs/Horde/Horde.Build/Controllers/ErrorController.cs
--- a/Engine/Source/Programs/Horde/Horde.Build/Controllers/ErrorController.cs
+++ b/Engine/Source/Programs/Horde/Horde.Build/Controllers/ErrorController.cs
@@ -1,13 +1,13 @@
 // Copyright Epic Games, Inc. All Rights Reserved.
 
 using EpicGames.Core;
-using HordeServer.Utilities;
+using Horde.Build.Utilities;
 using Microsoft.AspNetCore.Diagnostics;
 using Microsoft.AspNetCore.Http;
 using Microsoft.AspNetCore.Mvc;
 using Microsoft.Extensions.Logging;
 
-namespace HordeServer.Controllers
+namespace Horde.Build.Controllers
 {
 	/// <summary>
 	/// Controller which returns exception metadata to the caller
diff --git a/Engine/Source/Programs/Horde/Horde.Build/Controllers/IssuesController.cs b/Engine/Source/Programs/Horde/Horde.Build/Controllers/IssuesController.cs
--- a/Engine/Source/Programs/Horde/Horde.Build/Controllers/IssuesController.cs
+++ b/Engine/Source/Programs/Horde/Horde.Build/Controllers/IssuesController.cs
@@ -1,12 +1,12 @@
 // Copyright Epic Games, Inc. All Rights Reserved.
 
 using EpicGames.Core;
-using HordeServer.Api;
-using HordeServer.Collections;
+using Horde.Build.Api;
+using Horde.Build.Collections;
 using HordeCommon;
-using HordeServer.Models;
-using HordeServer.Services;
-using HordeServer.Utilities;
+using Horde.Build.Models;
+using Horde.Build.Services;
+using Horde.Build.Utilities;
 using Microsoft.AspNetCore.Authorization;
 using Microsoft.AspNetCore.Mvc;
 using Microsoft.AspNetCore.Server.Kestrel.Core.Features;
@@ -19,7 +19,7 @@ using System.Linq;
 using System.Text;
 using System.Threading.Tasks;
 
-namespace HordeServer.Controllers
+namespace Horde.Build.Controllers
 {
 	using JobId = ObjectId<IJob>;
 	using LogId = ObjectId<ILogFile>;
diff --git a/Engine/Source/Programs/Horde/Horde.Build/Controllers/JiraController.cs b/Engine/Source/Programs/Horde/Horde.Build/Controllers/JiraController.cs
--- a/Engine/Source/Programs/Horde/Horde.Build/Controllers/JiraController.cs
+++ b/Engine/Source/Programs/Horde/Horde.Build/Controllers/JiraController.cs
@@ -3,12 +3,12 @@
 using Microsoft.AspNetCore.Authorization;
 using Microsoft.AspNetCore.Mvc;
 using System.Threading.Tasks;
-using HordeServer.Services;
+using Horde.Build.Services;
 using Microsoft.Extensions.Logging;
-using HordeServer.Api;
+using Horde.Build.Api;
 using System.Collections.Generic;
-using HordeServer.Models;
-using HordeServer.Utilities;
+using Horde.Build.Models;
+using Horde.Build.Utilities;
 
 namespace Horde.Build
 {
diff --git a/Engine/Source/Programs/Horde/Horde.Build/Controllers/JobsController.cs b/Engine/Source/Programs/Horde/Horde.Build/Controllers/JobsController.cs
--- a/Engine/Source/Programs/Horde/Horde.Build/Controllers/JobsController.cs
+++ b/Engine/Source/Programs/Horde/Horde.Build/Controllers/JobsController.cs
@@ -2,12 +2,11 @@
 
 using EpicGames.Core;
 using Horde.Build.Utilities;
-using HordeServer.Api;
-using HordeServer.Collections;
+using Horde.Build.Api;
+using Horde.Build.Collections;
 using HordeCommon;
-using HordeServer.Models;
-using HordeServer.Services;
-using HordeServer.Utilities;
+using Horde.Build.Models;
+using Horde.Build.Services;
 using Microsoft.AspNetCore.Authorization;
 using Microsoft.AspNetCore.Mvc;
 using MongoDB.Bson;
@@ -21,13 +20,13 @@ using System.Text.RegularExpressions;
 using System.Threading.Tasks;
 
 using Microsoft.Extensions.Logging;
-using HordeServer.Notifications;
+using Horde.Build.Notifications;
 using OpenTracing.Util;
 using OpenTracing;
-using HordeServer.Jobs;
+using Horde.Build.Jobs;
 using EpicGames.Perforce;
 
-namespace HordeServer.Controllers
+namespace Horde.Build.Controllers
 {
 	using JobId = ObjectId<IJob>;
 	using StreamId = StringId<IStream>;
diff --git a/Engine/Source/Programs/Horde/Horde.Build/Controllers/LeasesController.cs b/Engine/Source/Programs/Horde/Horde.Build/Controllers/LeasesController.cs
--- a/Engine/Source/Programs/Horde/Horde.Build/Controllers/LeasesController.cs
+++ b/Engine/Source/Programs/Horde/Horde.Build/Controllers/LeasesController.cs
@@ -1,9 +1,9 @@
 // Copyright Epic Games, Inc. All Rights Reserved.
 
-using HordeServer.Api;
-using HordeServer.Models;
-using HordeServer.Services;
-using HordeServer.Utilities;
+using Horde.Build.Api;
+using Horde.Build.Models;
+using Horde.Build.Services;
+using Horde.Build.Utilities;
 using Microsoft.AspNetCore.Authentication;
 using Microsoft.AspNetCore.Authorization;
 using Microsoft.AspNetCore.Mvc;
@@ -24,7 +24,7 @@ using Microsoft.Extensions.Logging;
 using OpenTracing;
 using OpenTracing.Util;
 
-namespace HordeServer.Controllers
+namespace Horde.Build.Controllers
 {
 	using LeaseId = ObjectId<ILease>;
 	using PoolId = StringId<IPool>;
diff --git a/Engine/Source/Programs/Horde/Horde.Build/Controllers/LifetimeController.cs b/Engine/Source/Programs/Horde/Horde.Build/Controllers/LifetimeController.cs
--- a/Engine/Source/Programs/Horde/Horde.Build/Controllers/LifetimeController.cs
+++ b/Engine/Source/Programs/Horde/Horde.Build/Controllers/LifetimeController.cs
@@ -1,10 +1,10 @@
 // Copyright Epic Games, Inc. All Rights Reserved.
 
-using HordeServer.Services;
+using Horde.Build.Services;
 using Microsoft.AspNetCore.Mvc;
 using System.Threading.Tasks;
 
-namespace HordeServer.Controllers
+namespace Horde.Build.Controllers
 {
 	/// <summary>
 	/// Controller for app lifetime related routes
diff --git a/Engine/Source/Programs/Horde/Horde.Build/Controllers/LogsController.cs b/Engine/Source/Programs/Horde/Horde.Build/Controllers/LogsController.cs
--- a/Engine/Source/Programs/Horde/Horde.Build/Controllers/LogsController.cs
+++ b/Engine/Source/Programs/Horde/Horde.Build/Controllers/LogsController.cs
@@ -11,12 +11,12 @@ using System.Security.Claims;
 using System.Text;
 using System.Text.Json;
 using System.Threading.Tasks;
-using HordeServer.Api;
-using HordeServer.Collections;
-using HordeServer.Logs;
-using HordeServer.Models;
-using HordeServer.Services;
-using HordeServer.Utilities;
+using Horde.Build.Api;
+using Horde.Build.Collections;
+using Horde.Build.Logs;
+using Horde.Build.Models;
+using Horde.Build.Services;
+using Horde.Build.Utilities;
 using Microsoft.AspNetCore.Authorization;
 using Microsoft.AspNetCore.Http;
 using Microsoft.AspNetCore.Mvc;
@@ -24,7 +24,7 @@ using Microsoft.Extensions.Logging;
 using Microsoft.Extensions.Logging.Abstractions;
 using MongoDB.Bson;
 
-namespace HordeServer.Controllers
+namespace Horde.Build.Controllers
 {
 	using JobId = ObjectId<IJob>;
 	using LogId = ObjectId<ILogFile>;
diff --git a/Engine/Source/Programs/Horde/Horde.Build/Controllers/PerforceController.cs b/Engine/Source/Programs/Horde/Horde.Build/Controllers/PerforceController.cs
--- a/Engine/Source/Programs/Horde/Horde.Build/Controllers/PerforceController.cs
+++ b/Engine/Source/Programs/Horde/Horde.Build/Controllers/PerforceController.cs
@@ -5,11 +5,11 @@ using System.Net;
 using Microsoft.Extensions.Logging;
 using System.Threading.Tasks;
 using System.Collections.Generic;
-using HordeServer.Services;
+using Horde.Build.Services;
 using Microsoft.AspNetCore.Authorization;
-using HordeServer.Models;
+using Horde.Build.Models;
 
-namespace HordeServer.Controllers
+namespace Horde.Build.Controllers
 {
 	/// <summary>
 	/// Controller for Perforce triggers and callbacks
diff --git a/Engine/Source/Programs/Horde/Horde.Build/Controllers/ReportsController.cs b/Engine/Source/Programs/Horde/Horde.Build/Controllers/ReportsController.cs
--- a/Engine/Source/Programs/Horde/Horde.Build/Controllers/ReportsController.cs
+++ b/Engine/Source/Programs/Horde/Horde.Build/Controllers/ReportsController.cs
@@ -4,15 +4,15 @@ using System;
 using System.Collections.Generic;
 using System.Linq;
 using System.Threading.Tasks;
-using HordeServer.Api;
-using HordeServer.Collections;
-using HordeServer.Models;
-using HordeServer.Services;
+using Horde.Build.Api;
+using Horde.Build.Collections;
+using Horde.Build.Models;
+using Horde.Build.Services;
 using Microsoft.AspNetCore.Authorization;
 using Microsoft.AspNetCore.Mvc;
 using EpicGames.Core;
 
-namespace HordeServer.Controllers
+namespace Horde.Build.Controllers
 {
 	/// <summary>
 	/// Controller for the /api/v1/reports endpoint, used for the reports pages
diff --git a/Engine/Source/Programs/Horde/Horde.Build/Controllers/SchemaController.cs b/Engine/Source/Programs/Horde/Horde.Build/Controllers/SchemaController.cs
--- a/Engine/Source/Programs/Horde/Horde.Build/Controllers/SchemaController.cs
+++ b/Engine/Source/Programs/Horde/Horde.Build/Controllers/SchemaController.cs
@@ -1,8 +1,8 @@
 // Copyright Epic Games, Inc. All Rights Reserved.
 
 using EpicGames.Core;
-using HordeServer.Models;
-using HordeServer.Utilities;
+using Horde.Build.Models;
+using Horde.Build.Utilities;
 using Microsoft.AspNetCore.Mvc;
 using Microsoft.Extensions.Primitives;
 using System;
@@ -15,7 +15,7 @@ using System.Text.Json;
 using System.Text.Json.Serialization;
 using System.Threading.Tasks;
 
-namespace HordeServer.Controllers
+namespace Horde.Build.Controllers
 {
 	/// <summary>
 	/// Controller for the /api/v1/schema endpoint
diff --git a/Engine/Source/Programs/Horde/Horde.Build/Controllers/SwarmController.cs b/Engine/Source/Programs/Horde/Horde.Build/Controllers/SwarmController.cs
--- a/Engine/Source/Programs/Horde/Horde.Build/Controllers/SwarmController.cs
+++ b/Engine/Source/Programs/Horde/Horde.Build/Controllers/SwarmController.cs
@@ -8,7 +8,7 @@ using System.Collections.Generic;
 using System.Linq;
 using System.Threading.Tasks;
 
-namespace HordeServer.Controllers
+namespace Horde.Build.Controllers
 {
 	/// <summary>
 	/// Controller for the /api/v1/swarm endpoint
diff --git a/Engine/Source/Programs/Horde/Horde.Build/Controllers/UsersController.cs b/Engine/Source/Programs/Horde/Horde.Build/Controllers/UsersController.cs
--- a/Engine/Source/Programs/Horde/Horde.Build/Controllers/UsersController.cs
+++ b/Engine/Source/Programs/Horde/Horde.Build/Controllers/UsersController.cs
@@ -1,11 +1,11 @@
 // Copyright Epic Games, Inc. All Rights Reserved.
 
-using HordeServer.Api;
-using HordeServer.Collections;
-using HordeServer.Models;
-using HordeServer.Notifications.Impl;
-using HordeServer.Services;
-using HordeServer.Utilities;
+using Horde.Build.Api;
+using Horde.Build.Collections;
+using Horde.Build.Models;
+using Horde.Build.Notifications.Impl;
+using Horde.Build.Services;
+using Horde.Build.Utilities;
 using Microsoft.AspNetCore.Authorization;
 using Microsoft.AspNetCore.Mvc;
 using Microsoft.Extensions.Options;
@@ -15,7 +15,7 @@ using System.Collections.Generic;
 using System.Linq;
 using System.Threading.Tasks;
 
-namespace HordeServer.Controllers
+namespace Horde.Build.Controllers
 {
 	using UserId = ObjectId<IUser>;
 
diff --git a/Engine/Source/Programs/Horde/Horde.Build/Devices/IDevice.cs b/Engine/Source/Programs/Horde/Horde.Build/Devices/IDevice.cs
--- a/Engine/Source/Programs/Horde/Horde.Build/Devices/IDevice.cs
+++ b/Engine/Source/Programs/Horde/Horde.Build/Devices/IDevice.cs
@@ -5,9 +5,9 @@ using MongoDB.Bson;
 using MongoDB.Bson.Serialization.Attributes;
 
 using System.Collections.Generic;
-using HordeServer.Utilities;
+using Horde.Build.Utilities;
 
-namespace HordeServer.Models
+namespace Horde.Build.Models
 {
 	using DeviceId = StringId<IDevice>;
 	using DevicePlatformId = StringId<IDevicePlatform>;
diff --git a/Engine/Source/Programs/Horde/Horde.Build/Devices/IDeviceCollection.cs b/Engine/Source/Programs/Horde/Horde.Build/Devices/IDeviceCollection.cs
--- a/Engine/Source/Programs/Horde/Horde.Build/Devices/IDeviceCollection.cs
+++ b/Engine/Source/Programs/Horde/Horde.Build/Devices/IDeviceCollection.cs
@@ -1,12 +1,12 @@
 // Copyright Epic Games, Inc. All Rights Reserved.
 
-using HordeServer.Models;
-using HordeServer.Utilities;
+using Horde.Build.Models;
+using Horde.Build.Utilities;
 using MongoDB.Bson;
 using System.Collections.Generic;
 using System.Threading.Tasks;
 
-namespace HordeServer.Collections
+namespace Horde.Build.Collections
 {
 	using DeviceId = StringId<IDevice>;
 	using DevicePlatformId = StringId<IDevicePlatform>;
diff --git a/Engine/Source/Programs/Horde/Horde.Build/Devices/Impl/DeviceCollection.cs b/Engine/Source/Programs/Horde/Horde.Build/Devices/Impl/DeviceCollection.cs
--- a/Engine/Source/Programs/Horde/Horde.Build/Devices/Impl/DeviceCollection.cs
+++ b/Engine/Source/Programs/Horde/Horde.Build/Devices/Impl/DeviceCollection.cs
@@ -1,19 +1,19 @@
 // Copyright Epic Games, Inc. All Rights Reserved.
 
-using HordeServer.Services;
+using Horde.Build.Services;
 using MongoDB.Bson.Serialization.Attributes;
 using System.Threading.Tasks;
 
-using HordeServer.Models;
+using Horde.Build.Models;
 using MongoDB.Driver;
 using System;
-using HordeServer.Utilities;
+using Horde.Build.Utilities;
 using System.Linq;
 
 using MongoDB.Bson;
 using System.Collections.Generic;
 
-namespace HordeServer.Collections.Impl
+namespace Horde.Build.Collections.Impl
 {
 	using DeviceId = StringId<IDevice>;
 	using DevicePlatformId = StringId<IDevicePlatform>;
diff --git a/Engine/Source/Programs/Horde/Horde.Build/Devices/Impl/DeviceService.cs b/Engine/Source/Programs/Horde/Horde.Build/Devices/Impl/DeviceService.cs
--- a/Engine/Source/Programs/Horde/Horde.Build/Devices/Impl/DeviceService.cs
+++ b/Engine/Source/Programs/Horde/Horde.Build/Devices/Impl/DeviceService.cs
@@ -2,10 +2,10 @@
 
 using System.Threading.Tasks;
 
-using HordeServer.Models;
-using HordeServer.Collections;
+using Horde.Build.Models;
+using Horde.Build.Collections;
 using System;
-using HordeServer.Utilities;
+using Horde.Build.Utilities;
 using System.Threading;
 
 using Microsoft.Extensions.Logging;
@@ -14,12 +14,12 @@ using MongoDB.Bson;
 using System.Collections.Generic;
 using System.Security.Claims;
 using System.Diagnostics.CodeAnalysis;
-using HordeServer.Notifications;
+using Horde.Build.Notifications;
 using System.Linq;
-using HordeServer.Api;
+using Horde.Build.Api;
 using Microsoft.Extensions.Hosting;
 
-namespace HordeServer.Services
+namespace Horde.Build.Services
 {
 	using DeviceId = StringId<IDevice>;
 	using DevicePlatformId = StringId<IDevicePlatform>;
diff --git a/Engine/Source/Programs/Horde/Horde.Build/Fleet/Autoscale/AutoscaleService.cs b/Engine/Source/Programs/Horde/Horde.Build/Fleet/Autoscale/AutoscaleService.cs
--- a/Engine/Source/Programs/Horde/Horde.Build/Fleet/Autoscale/AutoscaleService.cs
+++ b/Engine/Source/Programs/Horde/Horde.Build/Fleet/Autoscale/AutoscaleService.cs
@@ -7,9 +7,9 @@ using Amazon.Runtime;
 using Amazon.Runtime.CredentialManagement;
 using Google.Protobuf.WellKnownTypes;
 using HordeCommon.Rpc.Tasks;
-using HordeServer.Collections;
-using HordeServer.Models;
-using HordeServer.Utilities;
+using Horde.Build.Collections;
+using Horde.Build.Models;
+using Horde.Build.Utilities;
 using Microsoft.Extensions.Logging;
 using MongoDB.Bson;
 using System;
@@ -27,7 +27,7 @@ using OpenTracing.Util;
 using StatsdClient;
 using Microsoft.Extensions.Hosting;
 
-namespace HordeServer.Services
+namespace Horde.Build.Services
 {
 	using PoolId = StringId<IPool>;
 
diff --git a/Engine/Source/Programs/Horde/Horde.Build/Fleet/Autoscale/AutoscaleServiceV2.cs b/Engine/Source/Programs/Horde/Horde.Build/Fleet/Autoscale/AutoscaleServiceV2.cs
--- a/Engine/Source/Programs/Horde/Horde.Build/Fleet/Autoscale/AutoscaleServiceV2.cs
+++ b/Engine/Source/Programs/Horde/Horde.Build/Fleet/Autoscale/AutoscaleServiceV2.cs
@@ -7,11 +7,11 @@ using System.Linq;
 using System.Threading;
 using System.Threading.Tasks;
 using HordeCommon;
-using HordeServer;
-using HordeServer.Collections;
-using HordeServer.Models;
-using HordeServer.Services;
-using HordeServer.Utilities;
+using Horde.Build;
+using Horde.Build.Collections;
+using Horde.Build.Models;
+using Horde.Build.Services;
+using Horde.Build.Utilities;
 using Microsoft.Extensions.Hosting;
 using Microsoft.Extensions.Logging;
 using Microsoft.Extensions.Options;
diff --git a/Engine/Source/Programs/Horde/Horde.Build/Fleet/Autoscale/IPoolSizeStrategy.cs b/Engine/Source/Programs/Horde/Horde.Build/Fleet/Autoscale/IPoolSizeStrategy.cs
--- a/Engine/Source/Programs/Horde/Horde.Build/Fleet/Autoscale/IPoolSizeStrategy.cs
+++ b/Engine/Source/Programs/Horde/Horde.Build/Fleet/Autoscale/IPoolSizeStrategy.cs
@@ -3,8 +3,8 @@
 using System.Collections.Generic;
 using System.Linq;
 using System.Threading.Tasks;
-using HordeServer.Models;
-using HordeServer.Utilities;
+using Horde.Build.Models;
+using Horde.Build.Utilities;
 
 namespace Horde.Build.Fleet.Autoscale
 {
diff --git a/Engine/Source/Programs/Horde/Horde.Build/Fleet/Autoscale/JobQueueStrategy.cs b/Engine/Source/Programs/Horde/Horde.Build/Fleet/Autoscale/JobQueueStrategy.cs
--- a/Engine/Source/Programs/Horde/Horde.Build/Fleet/Autoscale/JobQueueStrategy.cs
+++ b/Engine/Source/Programs/Horde/Horde.Build/Fleet/Autoscale/JobQueueStrategy.cs
@@ -5,11 +5,11 @@ using System.Collections.Generic;
 using System.Linq;
 using System.Threading.Tasks;
 using HordeCommon;
-using HordeServer.Api;
-using HordeServer.Collections;
-using HordeServer.Models;
-using HordeServer.Services;
-using HordeServer.Utilities;
+using Horde.Build.Api;
+using Horde.Build.Collections;
+using Horde.Build.Models;
+using Horde.Build.Services;
+using Horde.Build.Utilities;
 
 namespace Horde.Build.Fleet.Autoscale
 {
diff --git a/Engine/Source/Programs/Horde/Horde.Build/Fleet/Autoscale/LeaseUtilizationStrategy.cs b/Engine/Source/Programs/Horde/Horde.Build/Fleet/Autoscale/LeaseUtilizationStrategy.cs
--- a/Engine/Source/Programs/Horde/Horde.Build/Fleet/Autoscale/LeaseUtilizationStrategy.cs
+++ b/Engine/Source/Programs/Horde/Horde.Build/Fleet/Autoscale/LeaseUtilizationStrategy.cs
@@ -8,9 +8,9 @@ using System.Threading.Tasks;
 using Google.Protobuf.WellKnownTypes;
 using HordeCommon;
 using HordeCommon.Rpc.Tasks;
-using HordeServer.Collections;
-using HordeServer.Models;
-using HordeServer.Utilities;
+using Horde.Build.Collections;
+using Horde.Build.Models;
+using Horde.Build.Utilities;
 using OpenTracing;
 using OpenTracing.Util;
 
diff --git a/Engine/Source/Programs/Horde/Horde.Build/Fleet/IFleetManager.cs b/Engine/Source/Programs/Horde/Horde.Build/Fleet/IFleetManager.cs
--- a/Engine/Source/Programs/Horde/Horde.Build/Fleet/IFleetManager.cs
+++ b/Engine/Source/Programs/Horde/Horde.Build/Fleet/IFleetManager.cs
@@ -1,12 +1,12 @@
 // Copyright Epic Games, Inc. All Rights Reserved.
 
-using HordeServer.Models;
+using Horde.Build.Models;
 using System;
 using System.Collections.Generic;
 using System.Linq;
 using System.Threading.Tasks;
 
-namespace HordeServer.Services
+namespace Horde.Build.Services
 {
 	/// <summary>
 	/// Service to manage a fleet of machines
diff --git a/Engine/Source/Programs/Horde/Horde.Build/Fleet/ITelemetryCollection.cs b/Engine/Source/Programs/Horde/Horde.Build/Fleet/ITelemetryCollection.cs
--- a/Engine/Source/Programs/Horde/Horde.Build/Fleet/ITelemetryCollection.cs
+++ b/Engine/Source/Programs/Horde/Horde.Build/Fleet/ITelemetryCollection.cs
@@ -1,13 +1,13 @@
 // Copyright Epic Games, Inc. All Rights Reserved.
 
-using HordeServer.Models;
-using HordeServer.Utilities;
+using Horde.Build.Models;
+using Horde.Build.Utilities;
 using System;
 using System.Collections.Generic;
 using System.Linq;
 using System.Threading.Tasks;
 
-namespace HordeServer.Collections
+namespace Horde.Build.Collections
 {
 	using PoolId = StringId<IPool>;
 	using StreamId = StringId<IStream>;
diff --git a/Engine/Source/Programs/Horde/Horde.Build/Fleet/Impl/Managers/Fleet/AwsFleetManager.cs b/Engine/Source/Programs/Horde/Horde.Build/Fleet/Impl/Managers/Fleet/AwsFleetManager.cs
--- a/Engine/Source/Programs/Horde/Horde.Build/Fleet/Impl/Managers/Fleet/AwsFleetManager.cs
+++ b/Engine/Source/Programs/Horde/Horde.Build/Fleet/Impl/Managers/Fleet/AwsFleetManager.cs
@@ -3,18 +3,18 @@
 using Amazon;
 using Amazon.EC2;
 using Amazon.EC2.Model;
-using HordeServer.Collections;
-using HordeServer.Models;
+using Horde.Build.Collections;
+using Horde.Build.Models;
 using Microsoft.Extensions.Logging;
 using System;
 using System.Collections.Generic;
 using System.Linq;
 using System.Threading.Tasks;
-using HordeServer.Utilities;
+using Horde.Build.Utilities;
 using OpenTracing;
 using OpenTracing.Util;
 
-namespace HordeServer.Services.Impl
+namespace Horde.Build.Services.Impl
 {
 	/// <summary>
 	/// Fleet manager for handling AWS EC2 instances
diff --git a/Engine/Source/Programs/Horde/Horde.Build/Fleet/Impl/Managers/Fleet/DefaultFleetManager.cs b/Engine/Source/Programs/Horde/Horde.Build/Fleet/Impl/Managers/Fleet/DefaultFleetManager.cs
--- a/Engine/Source/Programs/Horde/Horde.Build/Fleet/Impl/Managers/Fleet/DefaultFleetManager.cs
+++ b/Engine/Source/Programs/Horde/Horde.Build/Fleet/Impl/Managers/Fleet/DefaultFleetManager.cs
@@ -1,13 +1,13 @@
 // Copyright Epic Games, Inc. All Rights Reserved.
 
-using HordeServer.Models;
+using Horde.Build.Models;
 using Microsoft.Extensions.Logging;
 using System;
 using System.Collections.Generic;
 using System.Linq;
 using System.Threading.Tasks;
 
-namespace HordeServer.Services.Impl
+namespace Horde.Build.Services.Impl
 {
 	/// <summary>
 	/// Default implementation of <see cref="IFleetManager"/>
diff --git a/Engine/Source/Programs/Horde/Horde.Build/Issues/IIssue.cs b/Engine/Source/Programs/Horde/Horde.Build/Issues/IIssue.cs
--- a/Engine/Source/Programs/Horde/Horde.Build/Issues/IIssue.cs
+++ b/Engine/Source/Programs/Horde/Horde.Build/Issues/IIssue.cs
@@ -1,11 +1,11 @@
 // Copyright Epic Games, Inc. All Rights Reserved.
 
 using EpicGames.Core;
-using HordeServer.Api;
-using HordeServer.Collections;
+using Horde.Build.Api;
+using Horde.Build.Collections;
 using HordeCommon;
-using HordeServer.Services;
-using HordeServer.Utilities;
+using Horde.Build.Services;
+using Horde.Build.Utilities;
 using MongoDB.Bson;
 using MongoDB.Bson.Serialization.Attributes;
 using MongoDB.Bson.Serialization.Serializers;
@@ -15,7 +15,7 @@ using System.Diagnostics.CodeAnalysis;
 using System.Linq;
 using System.Runtime.CompilerServices;
 
-namespace HordeServer.Models
+namespace Horde.Build.Models
 {
 	using StreamId = StringId<IStream>;
 	using TemplateRefId = StringId<TemplateRef>;
diff --git a/Engine/Source/Programs/Horde/Horde.Build/Issues/IIssueDetails.cs b/Engine/Source/Programs/Horde/Horde.Build/Issues/IIssueDetails.cs
--- a/Engine/Source/Programs/Horde/Horde.Build/Issues/IIssueDetails.cs
+++ b/Engine/Source/Programs/Horde/Horde.Build/Issues/IIssueDetails.cs
@@ -1,14 +1,14 @@
 // Copyright Epic Games, Inc. All Rights Reserved.
 
-using HordeServer.Services;
-using HordeServer.Utilities;
+using Horde.Build.Services;
+using Horde.Build.Utilities;
 using MongoDB.Bson;
 using System;
 using System.Collections.Generic;
 using System.Linq;
 using System.Threading.Tasks;
 
-namespace HordeServer.Models
+namespace Horde.Build.Models
 {
 	using UserId = ObjectId<IUser>;
 
diff --git a/Engine/Source/Programs/Horde/Horde.Build/Issues/IIssueHandler.cs b/Engine/Source/Programs/Horde/Horde.Build/Issues/IIssueHandler.cs
--- a/Engine/Source/Programs/Horde/Horde.Build/Issues/IIssueHandler.cs
+++ b/Engine/Source/Programs/Horde/Horde.Build/Issues/IIssueHandler.cs
@@ -1,15 +1,15 @@
 // Copyright Epic Games, Inc. All Rights Reserved.
 
 using EpicGames.Core;
-using HordeServer.Models;
-using HordeServer.Collections;
+using Horde.Build.Models;
+using Horde.Build.Collections;
 using System;
 using System.Collections.Generic;
 using System.Linq;
 using System.Threading.Tasks;
 using System.Diagnostics.CodeAnalysis;
 
-namespace HordeServer.IssueHandlers
+namespace Horde.Build.IssueHandlers
 {
 	/// <summary>
 	/// Information about a changelist and a value ranking the likelihood that it caused an issue
diff --git a/Engine/Source/Programs/Horde/Horde.Build/Issues/IIssueService.cs b/Engine/Source/Programs/Horde/Horde.Build/Issues/IIssueService.cs
--- a/Engine/Source/Programs/Horde/Horde.Build/Issues/IIssueService.cs
+++ b/Engine/Source/Programs/Horde/Horde.Build/Issues/IIssueService.cs
@@ -1,7 +1,7 @@
 // Copyright Epic Games, Inc. All Rights Reserved.
 
-using HordeServer.Models;
-using HordeServer.Utilities;
+using Horde.Build.Models;
+using Horde.Build.Utilities;
 using Microsoft.Extensions.Hosting;
 using MongoDB.Bson;
 using System;
@@ -9,7 +9,7 @@ using System.Collections.Generic;
 using System.Linq;
 using System.Threading.Tasks;
 
-namespace HordeServer.Services
+namespace Horde.Build.Services
 {
 	using LogId = ObjectId<ILogFile>;
 	using StreamId = StringId<IStream>;
diff --git a/Engine/Source/Programs/Horde/Horde.Build/Issues/IIssueStep.cs b/Engine/Source/Programs/Horde/Horde.Build/Issues/IIssueStep.cs
--- a/Engine/Source/Programs/Horde/Horde.Build/Issues/IIssueStep.cs
+++ b/Engine/Source/Programs/Horde/Horde.Build/Issues/IIssueStep.cs
@@ -1,15 +1,15 @@
 // Copyright Epic Games, Inc. All Rights Reserved.
 
-using HordeServer.Utilities;
+using Horde.Build.Utilities;
 using MongoDB.Bson;
 using System;
 using System.Collections.Generic;
 using System.Linq;
 using System.Threading.Tasks;
 
-using SubResourceId = HordeServer.Models.SubResourceId;
+using SubResourceId = Horde.Build.Models.SubResourceId;
 
-namespace HordeServer.Models
+namespace Horde.Build.Models
 {
 	using JobId = ObjectId<IJob>;
 	using LogId = ObjectId<ILogFile>;
diff --git a/Engine/Source/Programs/Horde/Horde.Build/Issues/Impl/Handlers/CompileIssueHandler.cs b/Engine/Source/Programs/Horde/Horde.Build/Issues/Impl/Handlers/CompileIssueHandler.cs
--- a/Engine/Source/Programs/Horde/Horde.Build/Issues/Impl/Handlers/CompileIssueHandler.cs
+++ b/Engine/Source/Programs/Horde/Horde.Build/Issues/Impl/Handlers/CompileIssueHandler.cs
@@ -1,10 +1,10 @@
 // Copyright Epic Games, Inc. All Rights Reserved.
 
 using HordeCommon;
-using HordeServer.Collections;
+using Horde.Build.Collections;
 using HordeCommon.Rpc;
-using HordeServer.Services;
-using HordeServer.Utilities;
+using Horde.Build.Services;
+using Horde.Build.Utilities;
 using Microsoft.Extensions.Logging;
 using MongoDB.Bson;
 using MongoDB.Bson.Serialization;
@@ -18,10 +18,10 @@ using System.Security.Policy;
 using System.Text;
 using System.Threading.Tasks;
 using System.Text.Json;
-using HordeServer.Models;
+using Horde.Build.Models;
 using EpicGames.Core;
 
-namespace HordeServer.IssueHandlers.Impl
+namespace Horde.Build.IssueHandlers.Impl
 {
 	/// <summary>
 	/// Instance of a particular compile error
diff --git a/Engine/Source/Programs/Horde/Horde.Build/Issues/Impl/Handlers/ContentIssueHandler.cs b/Engine/Source/Programs/Horde/Horde.Build/Issues/Impl/Handlers/ContentIssueHandler.cs
--- a/Engine/Source/Programs/Horde/Horde.Build/Issues/Impl/Handlers/ContentIssueHandler.cs
+++ b/Engine/Source/Programs/Horde/Horde.Build/Issues/Impl/Handlers/ContentIssueHandler.cs
@@ -1,9 +1,9 @@
 // Copyright Epic Games, Inc. All Rights Reserved.
 
 using HordeCommon;
-using HordeServer.Collections;
-using HordeServer.Models;
-using HordeServer.Utilities;
+using Horde.Build.Collections;
+using Horde.Build.Models;
+using Horde.Build.Utilities;
 using Microsoft.Extensions.Logging;
 using MongoDB.Bson;
 using MongoDB.Bson.Serialization;
@@ -18,7 +18,7 @@ using System.Text;
 using System.Threading.Tasks;
 using EpicGames.Core;
 
-namespace HordeServer.IssueHandlers.Impl
+namespace Horde.Build.IssueHandlers.Impl
 {
 	/// <summary>
 	/// Instance of a particular compile error
diff --git a/Engine/Source/Programs/Horde/Horde.Build/Issues/Impl/Handlers/CopyrightIssueHandler.cs b/Engine/Source/Programs/Horde/Horde.Build/Issues/Impl/Handlers/CopyrightIssueHandler.cs
--- a/Engine/Source/Programs/Horde/Horde.Build/Issues/Impl/Handlers/CopyrightIssueHandler.cs
+++ b/Engine/Source/Programs/Horde/Horde.Build/Issues/Impl/Handlers/CopyrightIssueHandler.cs
@@ -1,10 +1,10 @@
 // Copyright Epic Games, Inc. All Rights Reserved.
 
 using HordeCommon;
-using HordeServer.Collections;
+using Horde.Build.Collections;
 using HordeCommon.Rpc;
-using HordeServer.Models;
-using HordeServer.Utilities;
+using Horde.Build.Models;
+using Horde.Build.Utilities;
 using Microsoft.Extensions.Logging;
 using MongoDB.Bson;
 using MongoDB.Bson.Serialization;
@@ -19,7 +19,7 @@ using System.Text;
 using System.Threading.Tasks;
 using EpicGames.Core;
 
-namespace HordeServer.IssueHandlers.Impl
+namespace Horde.Build.IssueHandlers.Impl
 {
 	/// <summary>
 	/// Instance of a particular compile error
diff --git a/Engine/Source/Programs/Horde/Horde.Build/Issues/Impl/Handlers/DefaultIssueHandler.cs b/Engine/Source/Programs/Horde/Horde.Build/Issues/Impl/Handlers/DefaultIssueHandler.cs
--- a/Engine/Source/Programs/Horde/Horde.Build/Issues/Impl/Handlers/DefaultIssueHandler.cs
+++ b/Engine/Source/Programs/Horde/Horde.Build/Issues/Impl/Handlers/DefaultIssueHandler.cs
@@ -1,14 +1,14 @@
 // Copyright Epic Games, Inc. All Rights Reserved.
 
-using HordeServer.Collections;
-using HordeServer.Models;
+using Horde.Build.Collections;
+using Horde.Build.Models;
 using System;
 using System.Collections.Generic;
 using System.Diagnostics.CodeAnalysis;
 using System.Linq;
 using System.Threading.Tasks;
 
-namespace HordeServer.IssueHandlers.Impl
+namespace Horde.Build.IssueHandlers.Impl
 {
 	/// <summary>
 	/// Instance of a particular compile error
diff --git a/Engine/Source/Programs/Horde/Horde.Build/Issues/Impl/Handlers/GauntletIssueHandler.cs b/Engine/Source/Programs/Horde/Horde.Build/Issues/Impl/Handlers/GauntletIssueHandler.cs
--- a/Engine/Source/Programs/Horde/Horde.Build/Issues/Impl/Handlers/GauntletIssueHandler.cs
+++ b/Engine/Source/Programs/Horde/Horde.Build/Issues/Impl/Handlers/GauntletIssueHandler.cs
@@ -1,8 +1,8 @@
 // Copyright Epic Games, Inc. All Rights Reserved.
 
 using HordeCommon;
-using HordeServer.Models;
-using HordeServer.Utilities;
+using Horde.Build.Models;
+using Horde.Build.Utilities;
 using Microsoft.AspNetCore.Authentication;
 using Microsoft.AspNetCore.Connections.Features;
 using Microsoft.Extensions.Logging;
@@ -14,10 +14,10 @@ using System.Security.Cryptography.X509Certificates;
 using System.Text.Json;
 using System.Threading.Tasks;
 using System.Diagnostics.CodeAnalysis;
-using HordeServer.Collections;
+using Horde.Build.Collections;
 using EpicGames.Core;
 
-namespace HordeServer.IssueHandlers.Impl
+namespace Horde.Build.IssueHandlers.Impl
 {
 	/// <summary>
 	/// Instance of a particular Gauntlet error
diff --git a/Engine/Source/Programs/Horde/Horde.Build/Issues/Impl/Handlers/SymbolIssueHandler.cs b/Engine/Source/Programs/Horde/Horde.Build/Issues/Impl/Handlers/SymbolIssueHandler.cs
--- a/Engine/Source/Programs/Horde/Horde.Build/Issues/Impl/Handlers/SymbolIssueHandler.cs
+++ b/Engine/Source/Programs/Horde/Horde.Build/Issues/Impl/Handlers/SymbolIssueHandler.cs
@@ -1,9 +1,9 @@
 // Copyright Epic Games, Inc. All Rights Reserved.
 
 using HordeCommon;
-using HordeServer.Collections;
-using HordeServer.Models;
-using HordeServer.Utilities;
+using Horde.Build.Collections;
+using Horde.Build.Models;
+using Horde.Build.Utilities;
 using Microsoft.Extensions.Logging;
 using MongoDB.Bson;
 using MongoDB.Bson.Serialization;
@@ -15,7 +15,7 @@ using System.Text.RegularExpressions;
 using System.Threading.Tasks;
 using EpicGames.Core;
 
-namespace HordeServer.IssueHandlers.Impl
+namespace Horde.Build.IssueHandlers.Impl
 {
 	/// <summary>
 	/// Instance of a particular compile error
diff --git a/Engine/Source/Programs/Horde/Horde.Build/Issues/Impl/Handlers/SystemicIssueHandler.cs b/Engine/Source/Programs/Horde/Horde.Build/Issues/Impl/Handlers/SystemicIssueHandler.cs
--- a/Engine/Source/Programs/Horde/Horde.Build/Issues/Impl/Handlers/SystemicIssueHandler.cs
+++ b/Engine/Source/Programs/Horde/Horde.Build/Issues/Impl/Handlers/SystemicIssueHandler.cs
@@ -1,10 +1,10 @@
 // Copyright Epic Games, Inc. All Rights Reserved.
 
 using HordeCommon;
-using HordeServer.Collections;
+using Horde.Build.Collections;
 using HordeCommon.Rpc;
-using HordeServer.Services;
-using HordeServer.Utilities;
+using Horde.Build.Services;
+using Horde.Build.Utilities;
 using Microsoft.Extensions.Logging;
 using MongoDB.Bson;
 using MongoDB.Bson.Serialization;
@@ -18,10 +18,10 @@ using System.Security.Policy;
 using System.Text;
 using System.Threading.Tasks;
 using System.Text.Json;
-using HordeServer.Models;
+using Horde.Build.Models;
 using EpicGames.Core;
 
-namespace HordeServer.IssueHandlers.Impl
+namespace Horde.Build.IssueHandlers.Impl
 {
 	/// <summary>
 	/// Instance of a particular systemic error
diff --git a/Engine/Source/Programs/Horde/Horde.Build/Issues/Impl/IssueService.cs b/Engine/Source/Programs/Horde/Horde.Build/Issues/Impl/IssueService.cs
--- a/Engine/Source/Programs/Horde/Horde.Build/Issues/Impl/IssueService.cs
+++ b/Engine/Source/Programs/Horde/Horde.Build/Issues/Impl/IssueService.cs
@@ -3,11 +3,11 @@
 using EpicGames.Core;
 using EpicGames.Redis.Utility;
 using HordeCommon;
-using HordeServer.Collections;
-using HordeServer.IssueHandlers;
-using HordeServer.Models;
-using HordeServer.Services;
-using HordeServer.Utilities;
+using Horde.Build.Collections;
+using Horde.Build.IssueHandlers;
+using Horde.Build.Models;
+using Horde.Build.Services;
+using Horde.Build.Utilities;
 using Microsoft.Extensions.Hosting;
 using Microsoft.Extensions.Logging;
 using MongoDB.Bson;
@@ -26,12 +26,12 @@ using System.Text;
 using System.Text.RegularExpressions;
 using System.Threading;
 using System.Threading.Tasks;
-using HordeServer.IssueHandlers.Impl;
-using HordeServer.Notifications;
+using Horde.Build.IssueHandlers.Impl;
+using Horde.Build.Notifications;
 using OpenTracing;
 using OpenTracing.Util;
 
-namespace HordeServer.Services.Impl
+namespace Horde.Build.Services.Impl
 {
 	using LogId = ObjectId<ILogFile>;
 	using StreamId = StringId<IStream>;
diff --git a/Engine/Source/Programs/Horde/Horde.Build/Jira/IJiraService.cs b/Engine/Source/Programs/Horde/Horde.Build/Jira/IJiraService.cs
--- a/Engine/Source/Programs/Horde/Horde.Build/Jira/IJiraService.cs
+++ b/Engine/Source/Programs/Horde/Horde.Build/Jira/IJiraService.cs
@@ -3,7 +3,7 @@
 using System.Collections.Generic;
 using System.Threading.Tasks;
 
-namespace HordeServer.Services
+namespace Horde.Build.Services
 {
 
 	/// <summary>
diff --git a/Engine/Source/Programs/Horde/Horde.Build/Jira/Impl/JiraService.cs b/Engine/Source/Programs/Horde/Horde.Build/Jira/Impl/JiraService.cs
--- a/Engine/Source/Programs/Horde/Horde.Build/Jira/Impl/JiraService.cs
+++ b/Engine/Source/Programs/Horde/Horde.Build/Jira/Impl/JiraService.cs
@@ -13,7 +13,7 @@ using System.Text.Json;
 using System.Net.Http.Headers;
 using System.Collections.Concurrent;
 
-namespace HordeServer.Services
+namespace Horde.Build.Services
 {
 	/// <summary>
 	/// Jira service functinality
diff --git a/Engine/Source/Programs/Horde/Horde.Build/Jobs/IArtifact.cs b/Engine/Source/Programs/Horde/Horde.Build/Jobs/IArtifact.cs
--- a/Engine/Source/Programs/Horde/Horde.Build/Jobs/IArtifact.cs
+++ b/Engine/Source/Programs/Horde/Horde.Build/Jobs/IArtifact.cs
@@ -1,11 +1,11 @@
 // Copyright Epic Games, Inc. All Rights Reserved.
 
-using HordeServer.Api;
-using HordeServer.Utilities;
+using Horde.Build.Api;
+using Horde.Build.Utilities;
 using MongoDB.Bson;
 using MongoDB.Bson.Serialization.Attributes;
 
-namespace HordeServer.Models
+namespace Horde.Build.Models
 {
 	using JobId = ObjectId<IJob>;
 
diff --git a/Engine/Source/Programs/Horde/Horde.Build/Jobs/IArtifactCollection.cs b/Engine/Source/Programs/Horde/Horde.Build/Jobs/IArtifactCollection.cs
--- a/Engine/Source/Programs/Horde/Horde.Build/Jobs/IArtifactCollection.cs
+++ b/Engine/Source/Programs/Horde/Horde.Build/Jobs/IArtifactCollection.cs
@@ -1,12 +1,12 @@
 // Copyright Epic Games, Inc. All Rights Reserved.
 
-using HordeServer.Models;
-using HordeServer.Utilities;
+using Horde.Build.Models;
+using Horde.Build.Utilities;
 using MongoDB.Bson;
 using System.Collections.Generic;
 using System.Threading.Tasks;
 
-namespace HordeServer.Jobs
+namespace Horde.Build.Jobs
 {
 	using JobId = ObjectId<IJob>;
 
diff --git a/Engine/Source/Programs/Horde/Horde.Build/Jobs/IGraph.cs b/Engine/Source/Programs/Horde/Horde.Build/Jobs/IGraph.cs
--- a/Engine/Source/Programs/Horde/Horde.Build/Jobs/IGraph.cs
+++ b/Engine/Source/Programs/Horde/Horde.Build/Jobs/IGraph.cs
@@ -1,9 +1,9 @@
 // Copyright Epic Games, Inc. All Rights Reserved.
 
 using EpicGames.Core;
-using HordeServer.Api;
+using Horde.Build.Api;
 using HordeCommon;
-using HordeServer.Utilities;
+using Horde.Build.Utilities;
 using MongoDB.Bson;
 using MongoDB.Bson.Serialization.Attributes;
 using Newtonsoft.Json;
@@ -13,7 +13,7 @@ using System.Diagnostics;
 using System.Diagnostics.CodeAnalysis;
 using System.Linq;
 
-namespace HordeServer.Models
+namespace Horde.Build.Models
 {
 	/// <summary>
 	/// Represents a node in the graph
diff --git a/Engine/Source/Programs/Horde/Horde.Build/Jobs/IGraphCollection.cs b/Engine/Source/Programs/Horde/Horde.Build/Jobs/IGraphCollection.cs
--- a/Engine/Source/Programs/Horde/Horde.Build/Jobs/IGraphCollection.cs
+++ b/Engine/Source/Programs/Horde/Horde.Build/Jobs/IGraphCollection.cs
@@ -1,10 +1,10 @@
 // Copyright Epic Games, Inc. All Rights Reserved.
 
 using EpicGames.Core;
-using HordeServer.Api;
+using Horde.Build.Api;
 using HordeCommon;
-using HordeServer.Models;
-using HordeServer.Services;
+using Horde.Build.Models;
+using Horde.Build.Services;
 using MongoDB.Bson;
 using MongoDB.Bson.Serialization.Attributes;
 using MongoDB.Driver;
@@ -16,7 +16,7 @@ using System.Linq;
 using System.Threading;
 using System.Threading.Tasks;
 
-namespace HordeServer.Collections
+namespace Horde.Build.Collections
 {
 	/// <summary>
 	/// Interface for a collection of graph documents
diff --git a/Engine/Source/Programs/Horde/Horde.Build/Jobs/IJobCollection.cs b/Engine/Source/Programs/Horde/Horde.Build/Jobs/IJobCollection.cs
--- a/Engine/Source/Programs/Horde/Horde.Build/Jobs/IJobCollection.cs
+++ b/Engine/Source/Programs/Horde/Horde.Build/Jobs/IJobCollection.cs
@@ -1,11 +1,11 @@
 // Copyright Epic Games, Inc. All Rights Reserved.
 
 using EpicGames.Core;
-using HordeServer.Api;
+using Horde.Build.Api;
 using HordeCommon;
-using HordeServer.Models;
-using HordeServer.Services;
-using HordeServer.Utilities;
+using Horde.Build.Models;
+using Horde.Build.Services;
+using Horde.Build.Utilities;
 using Microsoft.Extensions.Logging;
 using MongoDB.Bson;
 using MongoDB.Bson.Serialization.Attributes;
@@ -17,7 +17,7 @@ using System.IO;
 using System.Linq;
 using System.Threading.Tasks;
 
-namespace HordeServer.Collections
+namespace Horde.Build.Collections
 {
 	using JobId = ObjectId<IJob>;
 	using LeaseId = ObjectId<ILease>;
diff --git a/Engine/Source/Programs/Horde/Horde.Build/Jobs/IJobStepRefCollection.cs b/Engine/Source/Programs/Horde/Horde.Build/Jobs/IJobStepRefCollection.cs
--- a/Engine/Source/Programs/Horde/Horde.Build/Jobs/IJobStepRefCollection.cs
+++ b/Engine/Source/Programs/Horde/Horde.Build/Jobs/IJobStepRefCollection.cs
@@ -1,8 +1,8 @@
 // Copyright Epic Games, Inc. All Rights Reserved.
 
 using HordeCommon;
-using HordeServer.Models;
-using HordeServer.Services;
+using Horde.Build.Models;
+using Horde.Build.Services;
 using MongoDB.Bson;
 using MongoDB.Bson.Serialization.Attributes;
 using MongoDB.Driver;
@@ -10,9 +10,9 @@ using System;
 using System.Collections.Generic;
 using System.Linq;
 using System.Threading.Tasks;
-using HordeServer.Utilities;
+using Horde.Build.Utilities;
 
-namespace HordeServer.Collections
+namespace Horde.Build.Collections
 {
 	using LogId = ObjectId<ILogFile>;
 	using PoolId = StringId<IPool>;
diff --git a/Engine/Source/Programs/Horde/Horde.Build/Jobs/IJobTiming.cs b/Engine/Source/Programs/Horde/Horde.Build/Jobs/IJobTiming.cs
--- a/Engine/Source/Programs/Horde/Horde.Build/Jobs/IJobTiming.cs
+++ b/Engine/Source/Programs/Horde/Horde.Build/Jobs/IJobTiming.cs
@@ -6,7 +6,7 @@ using System.Diagnostics.CodeAnalysis;
 using System.Linq;
 using System.Threading.Tasks;
 
-namespace HordeServer.Models
+namespace Horde.Build.Models
 {
 	/// <summary>
 	/// Information about the timing for an individual step
diff --git a/Engine/Source/Programs/Horde/Horde.Build/Jobs/ITemplateCollection.cs b/Engine/Source/Programs/Horde/Horde.Build/Jobs/ITemplateCollection.cs
--- a/Engine/Source/Programs/Horde/Horde.Build/Jobs/ITemplateCollection.cs
+++ b/Engine/Source/Programs/Horde/Horde.Build/Jobs/ITemplateCollection.cs
@@ -2,10 +2,10 @@
 
 using EpicGames.Core;
 using HordeCommon;
-using HordeServer.Api;
-using HordeServer.Models;
-using HordeServer.Services;
-using HordeServer.Utilities;
+using Horde.Build.Api;
+using Horde.Build.Models;
+using Horde.Build.Services;
+using Horde.Build.Utilities;
 using MongoDB.Bson;
 using MongoDB.Bson.Serialization;
 using MongoDB.Bson.Serialization.Attributes;
@@ -15,7 +15,7 @@ using System.Collections.Generic;
 using System.Linq;
 using System.Threading.Tasks;
 
-namespace HordeServer.Collections
+namespace Horde.Build.Collections
 {
 	/// <summary>
 	/// Interface for a collection of template documents
diff --git a/Engine/Source/Programs/Horde/Horde.Build/Jobs/Impl/ArtifactCollection.cs b/Engine/Source/Programs/Horde/Horde.Build/Jobs/Impl/ArtifactCollection.cs
--- a/Engine/Source/Programs/Horde/Horde.Build/Jobs/Impl/ArtifactCollection.cs
+++ b/Engine/Source/Programs/Horde/Horde.Build/Jobs/Impl/ArtifactCollection.cs
@@ -4,11 +4,11 @@ using Amazon.Extensions.NETCore.Setup;
 using Amazon.Runtime;
 using Amazon.S3;
 using Amazon.S3.Model;
-using HordeServer.Api;
-using HordeServer.Jobs;
-using HordeServer.Models;
-using HordeServer.Storage;
-using HordeServer.Utilities;
+using Horde.Build.Api;
+using Horde.Build.Jobs;
+using Horde.Build.Models;
+using Horde.Build.Storage;
+using Horde.Build.Utilities;
 using Microsoft.Extensions.Configuration;
 using Microsoft.Extensions.Hosting;
 using Microsoft.Extensions.Logging;
@@ -25,7 +25,7 @@ using System.Text;
 using System.Threading;
 using System.Threading.Tasks;
 
-namespace HordeServer.Services
+namespace Horde.Build.Services
 {
 	using JobId = ObjectId<IJob>;
 
diff --git a/Engine/Source/Programs/Horde/Horde.Build/Jobs/Impl/GraphCollection.cs b/Engine/Source/Programs/Horde/Horde.Build/Jobs/Impl/GraphCollection.cs
--- a/Engine/Source/Programs/Horde/Horde.Build/Jobs/Impl/GraphCollection.cs
+++ b/Engine/Source/Programs/Horde/Horde.Build/Jobs/Impl/GraphCollection.cs
@@ -1,10 +1,10 @@
 // Copyright Epic Games, Inc. All Rights Reserved.
 
 using EpicGames.Core;
-using HordeServer.Api;
+using Horde.Build.Api;
 using HordeCommon;
-using HordeServer.Models;
-using HordeServer.Services;
+using Horde.Build.Models;
+using Horde.Build.Services;
 using MongoDB.Bson;
 using MongoDB.Bson.Serialization.Attributes;
 using MongoDB.Driver;
@@ -16,7 +16,7 @@ using System.Linq;
 using System.Threading;
 using System.Threading.Tasks;
 
-namespace HordeServer.Collections.Impl
+namespace Horde.Build.Collections.Impl
 {
 	/// <summary>
 	/// Collection of graph documents
diff --git a/Engine/Source/Programs/Horde/Horde.Build/Jobs/Impl/JobCollection.cs b/Engine/Source/Programs/Horde/Horde.Build/Jobs/Impl/JobCollection.cs
--- a/Engine/Source/Programs/Horde/Horde.Build/Jobs/Impl/JobCollection.cs
+++ b/Engine/Source/Programs/Horde/Horde.Build/Jobs/Impl/JobCollection.cs
@@ -1,11 +1,11 @@
 // Copyright Epic Games, Inc. All Rights Reserved.
 
 using EpicGames.Core;
-using HordeServer.Api;
+using Horde.Build.Api;
 using HordeCommon;
-using HordeServer.Models;
-using HordeServer.Services;
-using HordeServer.Utilities;
+using Horde.Build.Models;
+using Horde.Build.Services;
+using Horde.Build.Utilities;
 using Microsoft.Extensions.Logging;
 using MongoDB.Bson;
 using MongoDB.Bson.Serialization.Attributes;
@@ -21,7 +21,7 @@ using OpenTracing;
 using MongoDB.Bson.Serialization;
 using MongoDB.Bson.IO;
 
-namespace HordeServer.Collections.Impl
+namespace Horde.Build.Collections.Impl
 {
 	using JobId = ObjectId<IJob>;
 	using LeaseId = ObjectId<ILease>;
diff --git a/Engine/Source/Programs/Horde/Horde.Build/Jobs/Impl/JobService.cs b/Engine/Source/Programs/Horde/Horde.Build/Jobs/Impl/JobService.cs
--- a/Engine/Source/Programs/Horde/Horde.Build/Jobs/Impl/JobService.cs
+++ b/Engine/Source/Programs/Horde/Horde.Build/Jobs/Impl/JobService.cs
@@ -1,12 +1,12 @@
 // Copyright Epic Games, Inc. All Rights Reserved.
 
 using EpicGames.Core;
-using HordeServer.Api;
-using HordeServer.Collections;
+using Horde.Build.Api;
+using Horde.Build.Collections;
 using HordeCommon;
-using HordeServer.Models;
+using Horde.Build.Models;
 using HordeCommon.Rpc;
-using HordeServer.Utilities;
+using Horde.Build.Utilities;
 using Microsoft.Extensions.Logging;
 using Microsoft.Extensions.Options;
 using MongoDB.Bson;
@@ -23,12 +23,11 @@ using System.Text.RegularExpressions;
 using System.Threading;
 using System.Threading.Tasks;
 using System.Globalization;
-using Horde.Build.Utilities;
-using HordeServer.Tasks.Impl;
+using Horde.Build.Tasks.Impl;
 using OpenTracing.Util;
 using OpenTracing;
 
-namespace HordeServer.Services
+namespace Horde.Build.Services
 {
 	using JobId = ObjectId<IJob>;
 	using LeaseId = ObjectId<ILease>;
@@ -269,7 +268,7 @@ namespace HordeServer.Services
 			if (UpdateIssues != null) TraceScope.Span.SetTag("UpdateIssues", UpdateIssues.Value);
 			if (JobTriggers != null) TraceScope.Span.SetTag("JobTriggers.Count", JobTriggers.Count);
 
-			JobId JobIdValue = JobId ?? HordeServer.Utilities.ObjectId<IJob>.GenerateNewId();
+			JobId JobIdValue = JobId ?? Horde.Build.Utilities.ObjectId<IJob>.GenerateNewId();
 			using IDisposable Scope = Logger.BeginScope("CreateJobAsync({JobId})", JobIdValue);
 
 			if (PreflightChange != null && ShouldClonePreflightChange(Stream.Id))
diff --git a/Engine/Source/Programs/Horde/Horde.Build/Jobs/Impl/JobStepRefCollection.cs b/Engine/Source/Programs/Horde/Horde.Build/Jobs/Impl/JobStepRefCollection.cs
--- a/Engine/Source/Programs/Horde/Horde.Build/Jobs/Impl/JobStepRefCollection.cs
+++ b/Engine/Source/Programs/Horde/Horde.Build/Jobs/Impl/JobStepRefCollection.cs
@@ -1,8 +1,8 @@
 // Copyright Epic Games, Inc. All Rights Reserved.
 
 using HordeCommon;
-using HordeServer.Models;
-using HordeServer.Services;
+using Horde.Build.Models;
+using Horde.Build.Services;
 using MongoDB.Bson;
 using MongoDB.Bson.Serialization.Attributes;
 using MongoDB.Driver;
@@ -10,9 +10,9 @@ using System;
 using System.Collections.Generic;
 using System.Linq;
 using System.Threading.Tasks;
-using HordeServer.Utilities;
+using Horde.Build.Utilities;
 
-namespace HordeServer.Collections.Impl
+namespace Horde.Build.Collections.Impl
 {
 	using LogId = ObjectId<ILogFile>;
 	using PoolId = StringId<IPool>;
diff --git a/Engine/Source/Programs/Horde/Horde.Build/Jobs/Impl/JobTaskSource.cs b/Engine/Source/Programs/Horde/Horde.Build/Jobs/Impl/JobTaskSource.cs
--- a/Engine/Source/Programs/Horde/Horde.Build/Jobs/Impl/JobTaskSource.cs
+++ b/Engine/Source/Programs/Horde/Horde.Build/Jobs/Impl/JobTaskSource.cs
@@ -5,14 +5,14 @@ using EpicGames.Horde.Common;
 using Google.Protobuf;
 using Google.Protobuf.Reflection;
 using Google.Protobuf.WellKnownTypes;
-using HordeServer.Api;
-using HordeServer.Collections;
+using Horde.Build.Api;
+using Horde.Build.Collections;
 using HordeCommon;
-using HordeServer.Models;
+using Horde.Build.Models;
 using HordeCommon.Rpc;
 using HordeCommon.Rpc.Tasks;
-using HordeServer.Tasks;
-using HordeServer.Utilities;
+using Horde.Build.Tasks;
+using Horde.Build.Utilities;
 using Microsoft.Extensions.Hosting;
 using Microsoft.Extensions.Logging;
 using Microsoft.Extensions.Options;
@@ -27,11 +27,11 @@ using System.Runtime.CompilerServices;
 using System.Text;
 using System.Threading;
 using System.Threading.Tasks;
-using HordeServer.Services;
+using Horde.Build.Services;
 using EpicGames.Horde.Compute;
 using System.Globalization;
 
-namespace HordeServer.Tasks.Impl
+namespace Horde.Build.Tasks.Impl
 {
 	using JobId = ObjectId<IJob>;
 	using LeaseId = ObjectId<ILease>;
diff --git a/Engine/Source/Programs/Horde/Horde.Build/Jobs/Impl/TemplateCollection.cs b/Engine/Source/Programs/Horde/Horde.Build/Jobs/Impl/TemplateCollection.cs
--- a/Engine/Source/Programs/Horde/Horde.Build/Jobs/Impl/TemplateCollection.cs
+++ b/Engine/Source/Programs/Horde/Horde.Build/Jobs/Impl/TemplateCollection.cs
@@ -2,10 +2,10 @@
 
 using EpicGames.Core;
 using HordeCommon;
-using HordeServer.Api;
-using HordeServer.Models;
-using HordeServer.Services;
-using HordeServer.Utilities;
+using Horde.Build.Api;
+using Horde.Build.Models;
+using Horde.Build.Services;
+using Horde.Build.Utilities;
 using Microsoft.Extensions.Caching.Memory;
 using MongoDB.Bson;
 using MongoDB.Bson.Serialization;
@@ -16,7 +16,7 @@ using System.Collections.Generic;
 using System.Linq;
 using System.Threading.Tasks;
 
-namespace HordeServer.Collections.Impl
+namespace Horde.Build.Collections.Impl
 {
 	/// <summary>
 	/// Collection of template documents
diff --git a/Engine/Source/Programs/Horde/Horde.Build/Jobs/SubResourceId.cs b/Engine/Source/Programs/Horde/Horde.Build/Jobs/SubResourceId.cs
--- a/Engine/Source/Programs/Horde/Horde.Build/Jobs/SubResourceId.cs
+++ b/Engine/Source/Programs/Horde/Horde.Build/Jobs/SubResourceId.cs
@@ -10,7 +10,7 @@ using System.Globalization;
 using System.Linq;
 using System.Threading.Tasks;
 
-namespace HordeServer.Models
+namespace Horde.Build.Models
 {
 	/// <summary>
 	/// Identifier for subresources. Assigning unique ids to subresources prevents against race conditions using indices when subresources are added and removed.
diff --git a/Engine/Source/Programs/Horde/Horde.Build/Logs/Builder/ILogBuilder.cs b/Engine/Source/Programs/Horde/Horde.Build/Logs/Builder/ILogBuilder.cs
--- a/Engine/Source/Programs/Horde/Horde.Build/Logs/Builder/ILogBuilder.cs
+++ b/Engine/Source/Programs/Horde/Horde.Build/Logs/Builder/ILogBuilder.cs
@@ -1,15 +1,15 @@
 // Copyright Epic Games, Inc. All Rights Reserved.
 
-using HordeServer.Api;
-using HordeServer.Models;
-using HordeServer.Utilities;
+using Horde.Build.Api;
+using Horde.Build.Models;
+using Horde.Build.Utilities;
 using MongoDB.Bson;
 using System;
 using System.Collections.Generic;
 using System.Linq;
 using System.Threading.Tasks;
 
-namespace HordeServer.Logs
+namespace Horde.Build.Logs
 {
 	using LogId = ObjectId<ILogFile>;
 
diff --git a/Engine/Source/Programs/Horde/Horde.Build/Logs/Builder/Impl/LocalLogBuilder.cs b/Engine/Source/Programs/Horde/Horde.Build/Logs/Builder/Impl/LocalLogBuilder.cs
--- a/Engine/Source/Programs/Horde/Horde.Build/Logs/Builder/Impl/LocalLogBuilder.cs
+++ b/Engine/Source/Programs/Horde/Horde.Build/Logs/Builder/Impl/LocalLogBuilder.cs
@@ -1,8 +1,8 @@
 // Copyright Epic Games, Inc. All Rights Reserved.
 
-using HordeServer.Api;
-using HordeServer.Models;
-using HordeServer.Utilities;
+using Horde.Build.Api;
+using Horde.Build.Models;
+using Horde.Build.Utilities;
 using MongoDB.Bson;
 using System;
 using System.Collections.Concurrent;
@@ -11,7 +11,7 @@ using System.IO;
 using System.Linq;
 using System.Threading.Tasks;
 
-namespace HordeServer.Logs.Builder
+namespace Horde.Build.Logs.Builder
 {
 	using LogId = ObjectId<ILogFile>;
 
diff --git a/Engine/Source/Programs/Horde/Horde.Build/Logs/Builder/Impl/RedisLogBuilder.cs b/Engine/Source/Programs/Horde/Horde.Build/Logs/Builder/Impl/RedisLogBuilder.cs
--- a/Engine/Source/Programs/Horde/Horde.Build/Logs/Builder/Impl/RedisLogBuilder.cs
+++ b/Engine/Source/Programs/Horde/Horde.Build/Logs/Builder/Impl/RedisLogBuilder.cs
@@ -1,8 +1,8 @@
 // Copyright Epic Games, Inc. All Rights Reserved.
 
 using EpicGames.Core;
-using HordeServer.Api;
-using HordeServer.Utilities;
+using Horde.Build.Api;
+using Horde.Build.Utilities;
 using Microsoft.Extensions.Logging;
 using Microsoft.Extensions.Primitives;
 using MongoDB.Bson;
@@ -15,9 +15,9 @@ using System.Text.RegularExpressions;
 using System.Threading.Tasks;
 using OpenTracing;
 using OpenTracing.Util;
-using HordeServer.Models;
+using Horde.Build.Models;
 
-namespace HordeServer.Logs.Builder
+namespace Horde.Build.Logs.Builder
 {
 	using LogId = ObjectId<ILogFile>;
 	using Condition = StackExchange.Redis.Condition;
diff --git a/Engine/Source/Programs/Horde/Horde.Build/Logs/Data/LogChunkData.cs b/Engine/Source/Programs/Horde/Horde.Build/Logs/Data/LogChunkData.cs
--- a/Engine/Source/Programs/Horde/Horde.Build/Logs/Data/LogChunkData.cs
+++ b/Engine/Source/Programs/Horde/Horde.Build/Logs/Data/LogChunkData.cs
@@ -1,8 +1,8 @@
 // Copyright Epic Games, Inc. All Rights Reserved.
 
 using EpicGames.Core;
-using HordeServer.Api;
-using HordeServer.Utilities;
+using Horde.Build.Api;
+using Horde.Build.Utilities;
 using System;
 using System.Collections.Generic;
 using System.IO;
@@ -10,7 +10,7 @@ using System.Linq;
 using System.Text;
 using System.Threading.Tasks;
 
-namespace HordeServer.Logs
+namespace Horde.Build.Logs
 {
 	/// <summary>
 	/// Data for a log chunk
diff --git a/Engine/Source/Programs/Horde/Horde.Build/Logs/Data/LogIndexData.cs b/Engine/Source/Programs/Horde/Horde.Build/Logs/Data/LogIndexData.cs
--- a/Engine/Source/Programs/Horde/Horde.Build/Logs/Data/LogIndexData.cs
+++ b/Engine/Source/Programs/Horde/Horde.Build/Logs/Data/LogIndexData.cs
@@ -1,7 +1,7 @@
 // Copyright Epic Games, Inc. All Rights Reserved.
 
 using EpicGames.Core;
-using HordeServer.Utilities;
+using Horde.Build.Utilities;
 using OpenTracing;
 using OpenTracing.Util;
 using System;
@@ -12,7 +12,7 @@ using System.Reflection;
 using System.Text;
 using System.Threading.Tasks;
 
-namespace HordeServer.Logs
+namespace Horde.Build.Logs
 {
 	/// <summary>
 	/// Contains a source block of text to be indexed
diff --git a/Engine/Source/Programs/Horde/Horde.Build/Logs/Data/LogSubChunkData.cs b/Engine/Source/Programs/Horde/Horde.Build/Logs/Data/LogSubChunkData.cs
--- a/Engine/Source/Programs/Horde/Horde.Build/Logs/Data/LogSubChunkData.cs
+++ b/Engine/Source/Programs/Horde/Horde.Build/Logs/Data/LogSubChunkData.cs
@@ -1,8 +1,8 @@
 // Copyright Epic Games, Inc. All Rights Reserved.
 
 using EpicGames.Core;
-using HordeServer.Api;
-using HordeServer.Utilities;
+using Horde.Build.Api;
+using Horde.Build.Utilities;
 using ICSharpCode.SharpZipLib.BZip2;
 using Microsoft.AspNetCore.Routing.Constraints;
 using System;
@@ -13,7 +13,7 @@ using System.Runtime.CompilerServices;
 using System.Text;
 using System.Threading.Tasks;
 
-namespace HordeServer.Logs
+namespace Horde.Build.Logs
 {
 	/// <summary>
 	/// Pending data for a sub-chunk
diff --git a/Engine/Source/Programs/Horde/Horde.Build/Logs/Data/LogText.cs b/Engine/Source/Programs/Horde/Horde.Build/Logs/Data/LogText.cs
--- a/Engine/Source/Programs/Horde/Horde.Build/Logs/Data/LogText.cs
+++ b/Engine/Source/Programs/Horde/Horde.Build/Logs/Data/LogText.cs
@@ -12,7 +12,7 @@ using System.Text;
 using System.Text.Json;
 using System.Threading.Tasks;
 
-namespace HordeServer.Logs
+namespace Horde.Build.Logs
 {
 	/// <summary>
 	/// Interface for log data
diff --git a/Engine/Source/Programs/Horde/Horde.Build/Logs/Data/LogToken.cs b/Engine/Source/Programs/Horde/Horde.Build/Logs/Data/LogToken.cs
--- a/Engine/Source/Programs/Horde/Horde.Build/Logs/Data/LogToken.cs
+++ b/Engine/Source/Programs/Horde/Horde.Build/Logs/Data/LogToken.cs
@@ -5,7 +5,7 @@ using System.Collections.Generic;
 using System.Linq;
 using System.Threading.Tasks;
 
-namespace HordeServer.Logs
+namespace Horde.Build.Logs
 {
 	/// <summary>
 	/// Functionality for decomposing log text into tokens
diff --git a/Engine/Source/Programs/Horde/Horde.Build/Logs/ILogFile.cs b/Engine/Source/Programs/Horde/Horde.Build/Logs/ILogFile.cs
--- a/Engine/Source/Programs/Horde/Horde.Build/Logs/ILogFile.cs
+++ b/Engine/Source/Programs/Horde/Horde.Build/Logs/ILogFile.cs
@@ -1,9 +1,9 @@
 // Copyright Epic Games, Inc. All Rights Reserved.
 
 using EpicGames.Core;
-using HordeServer.Api;
-using HordeServer.Services;
-using HordeServer.Utilities;
+using Horde.Build.Api;
+using Horde.Build.Services;
+using Horde.Build.Utilities;
 using MongoDB.Bson;
 using MongoDB.Bson.Serialization.Attributes;
 using System;
@@ -12,7 +12,7 @@ using System.Diagnostics.CodeAnalysis;
 using System.Linq;
 using System.Threading.Tasks;
 
-namespace HordeServer.Models
+namespace Horde.Build.Models
 {
 	using JobId = ObjectId<IJob>;
 	using LogId = ObjectId<ILogFile>;
diff --git a/Engine/Source/Programs/Horde/Horde.Build/Logs/ILogFileCollection.cs b/Engine/Source/Programs/Horde/Horde.Build/Logs/ILogFileCollection.cs
--- a/Engine/Source/Programs/Horde/Horde.Build/Logs/ILogFileCollection.cs
+++ b/Engine/Source/Programs/Horde/Horde.Build/Logs/ILogFileCollection.cs
@@ -1,9 +1,9 @@
 // Copyright Epic Games, Inc. All Rights Reserved.
 
-using HordeServer.Api;
-using HordeServer.Models;
-using HordeServer.Services;
-using HordeServer.Utilities;
+using Horde.Build.Api;
+using Horde.Build.Models;
+using Horde.Build.Services;
+using Horde.Build.Utilities;
 using MongoDB.Bson;
 using MongoDB.Bson.Serialization.Attributes;
 using MongoDB.Driver;
@@ -13,7 +13,7 @@ using System.Linq;
 using System.Net;
 using System.Threading.Tasks;
 
-namespace HordeServer.Collections
+namespace Horde.Build.Collections
 {
 	using JobId = ObjectId<IJob>;
 	using LogId = ObjectId<ILogFile>;
diff --git a/Engine/Source/Programs/Horde/Horde.Build/Logs/Impl/LogFileCollection.cs b/Engine/Source/Programs/Horde/Horde.Build/Logs/Impl/LogFileCollection.cs
--- a/Engine/Source/Programs/Horde/Horde.Build/Logs/Impl/LogFileCollection.cs
+++ b/Engine/Source/Programs/Horde/Horde.Build/Logs/Impl/LogFileCollection.cs
@@ -1,9 +1,9 @@
 // Copyright Epic Games, Inc. All Rights Reserved.
 
-using HordeServer.Api;
-using HordeServer.Models;
-using HordeServer.Services;
-using HordeServer.Utilities;
+using Horde.Build.Api;
+using Horde.Build.Models;
+using Horde.Build.Services;
+using Horde.Build.Utilities;
 using MongoDB.Bson;
 using MongoDB.Bson.Serialization.Attributes;
 using MongoDB.Driver;
@@ -13,7 +13,7 @@ using System.Linq;
 using System.Net;
 using System.Threading.Tasks;
 
-namespace HordeServer.Collections.Impl
+namespace Horde.Build.Collections.Impl
 {
 	using JobId = ObjectId<IJob>;
 	using LogId = ObjectId<ILogFile>;
diff --git a/Engine/Source/Programs/Horde/Horde.Build/Logs/Impl/LogFileService.cs b/Engine/Source/Programs/Horde/Horde.Build/Logs/Impl/LogFileService.cs
--- a/Engine/Source/Programs/Horde/Horde.Build/Logs/Impl/LogFileService.cs
+++ b/Engine/Source/Programs/Horde/Horde.Build/Logs/Impl/LogFileService.cs
@@ -2,12 +2,12 @@
 
 using EpicGames.Core;
 using HordeCommon;
-using HordeServer.Api;
-using HordeServer.Collections;
-using HordeServer.Logs;
-using HordeServer.Models;
-using HordeServer.Storage;
-using HordeServer.Utilities;
+using Horde.Build.Api;
+using Horde.Build.Collections;
+using Horde.Build.Logs;
+using Horde.Build.Models;
+using Horde.Build.Storage;
+using Horde.Build.Utilities;
 using Microsoft.Extensions.Caching.Memory;
 using Microsoft.Extensions.Logging;
 using MongoDB.Bson;
@@ -35,7 +35,7 @@ using OpenTracing;
 using Microsoft.Extensions.Hosting;
 using Microsoft.Extensions.Logging.Abstractions;
 
-namespace HordeServer.Services
+namespace Horde.Build.Services
 {
 	using JobId = ObjectId<IJob>;
 	using LogId = ObjectId<ILogFile>;
diff --git a/Engine/Source/Programs/Horde/Horde.Build/Logs/Storage/ILogStorage.cs b/Engine/Source/Programs/Horde/Horde.Build/Logs/Storage/ILogStorage.cs
--- a/Engine/Source/Programs/Horde/Horde.Build/Logs/Storage/ILogStorage.cs
+++ b/Engine/Source/Programs/Horde/Horde.Build/Logs/Storage/ILogStorage.cs
@@ -1,9 +1,9 @@
 // Copyright Epic Games, Inc. All Rights Reserved.
 
-using HordeServer.Models;
-using HordeServer.Services;
-using HordeServer.Storage;
-using HordeServer.Utilities;
+using Horde.Build.Models;
+using Horde.Build.Services;
+using Horde.Build.Storage;
+using Horde.Build.Utilities;
 using Microsoft.Extensions.Caching.Memory;
 using MongoDB.Bson;
 using System;
@@ -11,7 +11,7 @@ using System.Collections.Generic;
 using System.Linq;
 using System.Threading.Tasks;
 
-namespace HordeServer.Logs
+namespace Horde.Build.Logs
 {
 	using LogId = ObjectId<ILogFile>;
 
diff --git a/Engine/Source/Programs/Horde/Horde.Build/Logs/Storage/Impl/LocalLogStorage.cs b/Engine/Source/Programs/Horde/Horde.Build/Logs/Storage/Impl/LocalLogStorage.cs
--- a/Engine/Source/Programs/Horde/Horde.Build/Logs/Storage/Impl/LocalLogStorage.cs
+++ b/Engine/Source/Programs/Horde/Horde.Build/Logs/Storage/Impl/LocalLogStorage.cs
@@ -1,7 +1,7 @@
 // Copyright Epic Games, Inc. All Rights Reserved.
 
-using HordeServer.Models;
-using HordeServer.Utilities;
+using Horde.Build.Models;
+using Horde.Build.Utilities;
 using Microsoft.Extensions.Caching.Memory;
 using MongoDB.Bson;
 using System;
@@ -9,7 +9,7 @@ using System.Collections.Generic;
 using System.Linq;
 using System.Threading.Tasks;
 
-namespace HordeServer.Logs.Storage
+namespace Horde.Build.Logs.Storage
 {
 	using LogId = ObjectId<ILogFile>;
 
diff --git a/Engine/Source/Programs/Horde/Horde.Build/Logs/Storage/Impl/NullLogStorage.cs b/Engine/Source/Programs/Horde/Horde.Build/Logs/Storage/Impl/NullLogStorage.cs
--- a/Engine/Source/Programs/Horde/Horde.Build/Logs/Storage/Impl/NullLogStorage.cs
+++ b/Engine/Source/Programs/Horde/Horde.Build/Logs/Storage/Impl/NullLogStorage.cs
@@ -1,14 +1,14 @@
 // Copyright Epic Games, Inc. All Rights Reserved.
 
-using HordeServer.Models;
-using HordeServer.Utilities;
+using Horde.Build.Models;
+using Horde.Build.Utilities;
 using MongoDB.Bson;
 using System;
 using System.Collections.Generic;
 using System.Linq;
 using System.Threading.Tasks;
 
-namespace HordeServer.Logs.Storage.Impl
+namespace Horde.Build.Logs.Storage.Impl
 {
 	using LogId = ObjectId<ILogFile>;
 
diff --git a/Engine/Source/Programs/Horde/Horde.Build/Logs/Storage/Impl/PersistentLogStorage.cs b/Engine/Source/Programs/Horde/Horde.Build/Logs/Storage/Impl/PersistentLogStorage.cs
--- a/Engine/Source/Programs/Horde/Horde.Build/Logs/Storage/Impl/PersistentLogStorage.cs
+++ b/Engine/Source/Programs/Horde/Horde.Build/Logs/Storage/Impl/PersistentLogStorage.cs
@@ -1,10 +1,10 @@
 // Copyright Epic Games, Inc. All Rights Reserved.
 
 using EpicGames.Core;
-using HordeServer.Models;
-using HordeServer.Services;
-using HordeServer.Storage;
-using HordeServer.Utilities;
+using Horde.Build.Models;
+using Horde.Build.Services;
+using Horde.Build.Storage;
+using Horde.Build.Utilities;
 using Microsoft.Extensions.Logging;
 using MongoDB.Bson;
 using System;
@@ -12,7 +12,7 @@ using System.Collections.Generic;
 using System.Linq;
 using System.Threading.Tasks;
 
-namespace HordeServer.Logs.Readers
+namespace Horde.Build.Logs.Readers
 {
 	using LogId = ObjectId<ILogFile>;
 
diff --git a/Engine/Source/Programs/Horde/Horde.Build/Logs/Storage/Impl/RedisLogStorage.cs b/Engine/Source/Programs/Horde/Horde.Build/Logs/Storage/Impl/RedisLogStorage.cs
--- a/Engine/Source/Programs/Horde/Horde.Build/Logs/Storage/Impl/RedisLogStorage.cs
+++ b/Engine/Source/Programs/Horde/Horde.Build/Logs/Storage/Impl/RedisLogStorage.cs
@@ -1,7 +1,7 @@
 // Copyright Epic Games, Inc. All Rights Reserved.
 
-using HordeServer.Models;
-using HordeServer.Utilities;
+using Horde.Build.Models;
+using Horde.Build.Utilities;
 using Microsoft.Extensions.Logging;
 using MongoDB.Bson;
 using StackExchange.Redis;
@@ -11,7 +11,7 @@ using System.Diagnostics.CodeAnalysis;
 using System.Linq;
 using System.Threading.Tasks;
 
-namespace HordeServer.Logs.Storage.Impl
+namespace Horde.Build.Logs.Storage.Impl
 {
 	using LogId = ObjectId<ILogFile>;
 
diff --git a/Engine/Source/Programs/Horde/Horde.Build/Models/Globals.cs b/Engine/Source/Programs/Horde/Horde.Build/Models/Globals.cs
--- a/Engine/Source/Programs/Horde/Horde.Build/Models/Globals.cs
+++ b/Engine/Source/Programs/Horde/Horde.Build/Models/Globals.cs
@@ -1,8 +1,8 @@
 // Copyright Epic Games, Inc. All Rights Reserved.
 
 using Google.Protobuf.WellKnownTypes;
-using HordeServer.Services;
-using HordeServer.Utilities;
+using Horde.Build.Services;
+using Horde.Build.Utilities;
 using Microsoft.IdentityModel.Tokens;
 using MongoDB.Bson;
 using MongoDB.Bson.Serialization.Attributes;
@@ -14,7 +14,7 @@ using System.Linq;
 using System.Security.Cryptography;
 using System.Threading.Tasks;
 
-namespace HordeServer.Models
+namespace Horde.Build.Models
 {
 	using UserId = ObjectId<IUser>;
 
diff --git a/Engine/Source/Programs/Horde/Horde.Build/Notices/Impl/NoticeCollection.cs b/Engine/Source/Programs/Horde/Horde.Build/Notices/Impl/NoticeCollection.cs
--- a/Engine/Source/Programs/Horde/Horde.Build/Notices/Impl/NoticeCollection.cs
+++ b/Engine/Source/Programs/Horde/Horde.Build/Notices/Impl/NoticeCollection.cs
@@ -1,19 +1,19 @@
 // Copyright Epic Games, Inc. All Rights Reserved.
 
-using HordeServer.Services;
+using Horde.Build.Services;
 using MongoDB.Bson.Serialization.Attributes;
 using System.Threading.Tasks;
 
-using HordeServer.Models;
+using Horde.Build.Models;
 using MongoDB.Driver;
 using System;
-using HordeServer.Utilities;
+using Horde.Build.Utilities;
 using System.Linq;
 
 using MongoDB.Bson;
 using System.Collections.Generic;
 
-namespace HordeServer.Collections.Impl
+namespace Horde.Build.Collections.Impl
 {
 	using UserId = ObjectId<IUser>;
 
diff --git a/Engine/Source/Programs/Horde/Horde.Build/Notices/Impl/NoticeService.cs b/Engine/Source/Programs/Horde/Horde.Build/Notices/Impl/NoticeService.cs
--- a/Engine/Source/Programs/Horde/Horde.Build/Notices/Impl/NoticeService.cs
+++ b/Engine/Source/Programs/Horde/Horde.Build/Notices/Impl/NoticeService.cs
@@ -2,10 +2,10 @@
 
 using System.Threading.Tasks;
 
-using HordeServer.Models;
-using HordeServer.Collections;
+using Horde.Build.Models;
+using Horde.Build.Collections;
 using System;
-using HordeServer.Utilities;
+using Horde.Build.Utilities;
 using System.Threading;
 
 using Microsoft.Extensions.Logging;
@@ -14,12 +14,12 @@ using MongoDB.Bson;
 using System.Collections.Generic;
 using System.Security.Claims;
 using System.Diagnostics.CodeAnalysis;
-using HordeServer.Notifications;
+using Horde.Build.Notifications;
 using System.Linq;
-using HordeServer.Api;
+using Horde.Build.Api;
 using Microsoft.Extensions.Hosting;
 
-namespace HordeServer.Services
+namespace Horde.Build.Services
 {
 	using UserId = ObjectId<IUser>;
 
diff --git a/Engine/Source/Programs/Horde/Horde.Build/Notifications/INotificationService.cs b/Engine/Source/Programs/Horde/Horde.Build/Notifications/INotificationService.cs
--- a/Engine/Source/Programs/Horde/Horde.Build/Notifications/INotificationService.cs
+++ b/Engine/Source/Programs/Horde/Horde.Build/Notifications/INotificationService.cs
@@ -1,14 +1,14 @@
 // Copyright Epic Games, Inc. All Rights Reserved.
 
-using HordeServer.Api;
-using HordeServer.Models;
+using Horde.Build.Api;
+using Horde.Build.Models;
 using MongoDB.Bson;
 using System.Collections.Generic;
 using System.Diagnostics.CodeAnalysis;
 using System.Security.Claims;
 using System.Threading.Tasks;
 
-namespace HordeServer.Notifications
+namespace Horde.Build.Notifications
 {
 	/// <summary>
 	/// Marker interface for (serializable) notifications
diff --git a/Engine/Source/Programs/Horde/Horde.Build/Notifications/INotificationSink.cs b/Engine/Source/Programs/Horde/Horde.Build/Notifications/INotificationSink.cs
--- a/Engine/Source/Programs/Horde/Horde.Build/Notifications/INotificationSink.cs
+++ b/Engine/Source/Programs/Horde/Horde.Build/Notifications/INotificationSink.cs
@@ -1,13 +1,13 @@
 // Copyright Epic Games, Inc. All Rights Reserved.
 
 using HordeCommon;
-using HordeServer.Api;
-using HordeServer.Models;
+using Horde.Build.Api;
+using Horde.Build.Models;
 using System;
 using System.Collections.Generic;
 using System.Threading.Tasks;
 
-namespace HordeServer.Notifications
+namespace Horde.Build.Notifications
 {
 	/// <summary>
 	/// Implements a notification method
diff --git a/Engine/Source/Programs/Horde/Horde.Build/Notifications/Impl/NotificationService.cs b/Engine/Source/Programs/Horde/Horde.Build/Notifications/Impl/NotificationService.cs
--- a/Engine/Source/Programs/Horde/Horde.Build/Notifications/Impl/NotificationService.cs
+++ b/Engine/Source/Programs/Horde/Horde.Build/Notifications/Impl/NotificationService.cs
@@ -13,13 +13,13 @@ using System.Text;
 using System.Threading;
 using System.Threading.Tasks;
 using EpicGames.Core;
-using HordeServer.Api;
-using HordeServer.Collections;
+using Horde.Build.Api;
+using Horde.Build.Collections;
 using HordeCommon;
-using HordeServer.Models;
-using HordeServer.Utilities;
-using HordeServer.Utilities.BlockKit;
-using HordeServer.Utilities.Slack.BlockKit;
+using Horde.Build.Models;
+using Horde.Build.Utilities;
+using Horde.Build.Utilities.BlockKit;
+using Horde.Build.Utilities.Slack.BlockKit;
 using Microsoft.Extensions.Hosting;
 using Microsoft.Extensions.Logging;
 using Microsoft.Extensions.Options;
@@ -29,11 +29,11 @@ using Newtonsoft.Json;
 using Newtonsoft.Json.Linq;
 using StatsdClient;
 using JsonSerializer = System.Text.Json.JsonSerializer;
-using HordeServer.Notifications;
-using HordeServer.Services;
+using Horde.Build.Notifications;
+using Horde.Build.Services;
 using StackExchange.Redis;
 
-namespace HordeServer.Notifications.Impl
+namespace Horde.Build.Notifications.Impl
 {
 	using UserId = ObjectId<IUser>;
 
diff --git a/Engine/Source/Programs/Horde/Horde.Build/Notifications/Impl/SlackNotificationSink.cs b/Engine/Source/Programs/Horde/Horde.Build/Notifications/Impl/SlackNotificationSink.cs
--- a/Engine/Source/Programs/Horde/Horde.Build/Notifications/Impl/SlackNotificationSink.cs
+++ b/Engine/Source/Programs/Horde/Horde.Build/Notifications/Impl/SlackNotificationSink.cs
@@ -2,13 +2,13 @@
 
 using EpicGames.Core;
 using HordeCommon;
-using HordeServer.Api;
-using HordeServer.Collections;
-using HordeServer.Models;
-using HordeServer.Services;
-using HordeServer.Utilities;
-using HordeServer.Utilities.BlockKit;
-using HordeServer.Utilities.Slack.BlockKit;
+using Horde.Build.Api;
+using Horde.Build.Collections;
+using Horde.Build.Models;
+using Horde.Build.Services;
+using Horde.Build.Utilities;
+using Horde.Build.Utilities.BlockKit;
+using Horde.Build.Utilities.Slack.BlockKit;
 using Microsoft.Extensions.Caching.Memory;
 using Microsoft.Extensions.Hosting;
 using Microsoft.Extensions.Logging;
@@ -34,7 +34,7 @@ using System.Text.RegularExpressions;
 using System.Threading;
 using System.Threading.Tasks;
 
-namespace HordeServer.Notifications.Impl
+namespace Horde.Build.Notifications.Impl
 {
 	using UserId = ObjectId<IUser>;
 
diff --git a/Engine/Source/Programs/Horde/Horde.Build/Perforce/ICommit.cs b/Engine/Source/Programs/Horde/Horde.Build/Perforce/ICommit.cs
--- a/Engine/Source/Programs/Horde/Horde.Build/Perforce/ICommit.cs
+++ b/Engine/Source/Programs/Horde/Horde.Build/Perforce/ICommit.cs
@@ -3,8 +3,8 @@
 using EpicGames.Core;
 using EpicGames.Horde.Storage;
 using EpicGames.Perforce;
-using HordeServer.Models;
-using HordeServer.Utilities;
+using Horde.Build.Models;
+using Horde.Build.Utilities;
 using MongoDB.Bson;
 using System;
 using System.Collections.Generic;
@@ -12,7 +12,7 @@ using System.Linq;
 using System.Text;
 using System.Threading.Tasks;
 
-namespace HordeServer.Commits
+namespace Horde.Build.Commits
 {
 	using CommitId = ObjectId<ICommit>;
 	using StreamId = StringId<IStream>;
diff --git a/Engine/Source/Programs/Horde/Horde.Build/Perforce/ICommitCollection.cs b/Engine/Source/Programs/Horde/Horde.Build/Perforce/ICommitCollection.cs
--- a/Engine/Source/Programs/Horde/Horde.Build/Perforce/ICommitCollection.cs
+++ b/Engine/Source/Programs/Horde/Horde.Build/Perforce/ICommitCollection.cs
@@ -1,15 +1,15 @@
 ﻿// Copyright Epic Games, Inc. All Rights Reserved.
 
-using HordeServer.Models;
+using Horde.Build.Models;
 using System;
 using System.Collections.Generic;
 using System.Linq;
 using System.Threading.Tasks;
 using EpicGames.Core;
 using MongoDB.Bson;
-using HordeServer.Utilities;
+using Horde.Build.Utilities;
 
-namespace HordeServer.Commits
+namespace Horde.Build.Commits
 {
 	using StreamId = StringId<IStream>;
 	using CommitId = ObjectId<ICommit>;
diff --git a/Engine/Source/Programs/Horde/Horde.Build/Perforce/ICommitService.cs b/Engine/Source/Programs/Horde/Horde.Build/Perforce/ICommitService.cs
--- a/Engine/Source/Programs/Horde/Horde.Build/Perforce/ICommitService.cs
+++ b/Engine/Source/Programs/Horde/Horde.Build/Perforce/ICommitService.cs
@@ -1,15 +1,15 @@
 // Copyright Epic Games, Inc. All Rights Reserved.
 
-using HordeServer.Models;
+using Horde.Build.Models;
 using System;
 using System.Collections.Generic;
 using System.Linq;
 using System.Threading.Tasks;
 using EpicGames.Core;
 using MongoDB.Bson;
-using HordeServer.Utilities;
+using Horde.Build.Utilities;
 
-namespace HordeServer.Commits
+namespace Horde.Build.Commits
 {
 	/// <summary>
 	/// Provides information about commits
