Horde: Normalize namespace names to match assembly names.

diff --git a/Engine/Source/Programs/Horde/Horde.Build/Perforce/IPerforceService.cs b/Engine/Source/Programs/Horde/Horde.Build/Perforce/IPerforceService.cs
--- a/Engine/Source/Programs/Horde/Horde.Build/Perforce/IPerforceService.cs
+++ b/Engine/Source/Programs/Horde/Horde.Build/Perforce/IPerforceService.cs
@@ -2,7 +2,7 @@
 
 using EpicGames.Core;
 using EpicGames.Perforce;
-using HordeServer.Models;
+using Horde.Build.Models;
 using Microsoft.Extensions.Logging;
 using Polly;
 using Polly.Extensions.Http;
@@ -21,7 +21,7 @@ using System.Threading.Tasks;
 
 using P4 = Perforce.P4;
 
-namespace HordeServer.Services
+namespace Horde.Build.Services
 {
 	/// <summary>
 	/// Information about a Perforce user
diff --git a/Engine/Source/Programs/Horde/Horde.Build/Perforce/Impl/CommitCollection.cs b/Engine/Source/Programs/Horde/Horde.Build/Perforce/Impl/CommitCollection.cs
--- a/Engine/Source/Programs/Horde/Horde.Build/Perforce/Impl/CommitCollection.cs
+++ b/Engine/Source/Programs/Horde/Horde.Build/Perforce/Impl/CommitCollection.cs
@@ -1,9 +1,9 @@
 // Copyright Epic Games, Inc. All Rights Reserved.
 
 using EpicGames.Core;
-using HordeServer.Models;
-using HordeServer.Services;
-using HordeServer.Utilities;
+using Horde.Build.Models;
+using Horde.Build.Services;
+using Horde.Build.Utilities;
 using MongoDB.Bson;
 using MongoDB.Driver;
 using System;
@@ -13,7 +13,7 @@ using System.Threading.Tasks;
 using MongoDB.Bson.Serialization.Attributes;
 using EpicGames.Horde.Storage;
 
-namespace HordeServer.Commits.Impl
+namespace Horde.Build.Commits.Impl
 {
 	using CommitId = ObjectId<ICommit>;
 	using StreamId = StringId<IStream>;
diff --git a/Engine/Source/Programs/Horde/Horde.Build/Perforce/Impl/CommitService.cs b/Engine/Source/Programs/Horde/Horde.Build/Perforce/Impl/CommitService.cs
--- a/Engine/Source/Programs/Horde/Horde.Build/Perforce/Impl/CommitService.cs
+++ b/Engine/Source/Programs/Horde/Horde.Build/Perforce/Impl/CommitService.cs
@@ -1,9 +1,9 @@
 // Copyright Epic Games, Inc. All Rights Reserved.
 
 using EpicGames.Core;
-using HordeServer.Collections;
-using HordeServer.Models;
-using HordeServer.Utilities;
+using Horde.Build.Collections;
+using Horde.Build.Models;
+using Horde.Build.Utilities;
 using Microsoft.Extensions.Logging;
 using MongoDB.Bson;
 using Perforce;
@@ -13,8 +13,8 @@ using System.Linq;
 using System.Text.RegularExpressions;
 using System.Threading;
 using System.Threading.Tasks;
-using HordeServer.Storage;
-using HordeServer.Services;
+using Horde.Build.Storage;
+using Horde.Build.Services;
 using EpicGames.Serialization;
 using MongoDB.Bson.Serialization.Attributes;
 using System.Collections.Concurrent;
@@ -32,13 +32,12 @@ using HordeCommon;
 using EpicGames.Horde.Storage;
 using System.Text;
 using Microsoft.Extensions.Options;
-using Horde.Build.Utilities;
 using EpicGames.Horde.Bundles;
 using EpicGames.Horde.Bundles.Nodes;
 using Microsoft.Extensions.Caching.Memory;
 using System.Runtime;
 
-namespace HordeServer.Commits.Impl
+namespace Horde.Build.Commits.Impl
 {
 	using P4 = Perforce.P4;
 	using CommitId = ObjectId<ICommit>;
@@ -841,6 +840,21 @@ namespace HordeServer.Commits.Impl
 			}
 		}
 
+		/// <summary>
+		/// Creates a new bundle or reads an existing one
+		/// </summary>
+		async ValueTask<Bundle<DirectoryNode>> CreateOrReadBundle(StreamId StreamId, int? BaseChange, string? Filter, bool RevisionsOnly)
+		{
+			if (BaseChange == null)
+			{
+				return Bundle.Create<DirectoryNode>(StorageClient, Options.NamespaceId, Options.Bundle, MemoryCache);
+			}
+			else
+			{
+				return await Bundle.ReadAsync<DirectoryNode>(StorageClient, Options.NamespaceId, GetStreamBucketId(StreamId), new CommitKey(StreamId, BaseChange.Value, Filter, RevisionsOnly).GetRefId(), Options.Bundle, MemoryCache);
+			}
+		}
+
 		/// <summary>
 		/// Replicates the contents of a stream to Horde storage, optionally using the given change as a starting point
 		/// </summary>
@@ -852,22 +866,15 @@ namespace HordeServer.Commits.Impl
 		/// <returns>Root tree object</returns>
 		public async Task WriteCommitTreeAsync(IStream Stream, int Change, int? BaseChange, string? Filter, bool RevisionsOnly)
 		{
-			// Create a client to replicate from this stream, and connect to the server
+			// Create a client to replicate from this stream
 			ReplicationClient ClientInfo = await FindOrAddReplicationClientAsync(Stream);
+
+			// Connect to the server and flush the workspace
 			using IPerforceConnection Perforce = await PerforceConnection.CreateAsync(ClientInfo.Settings, Logger);
+			await FlushWorkspaceAsync(ClientInfo, Perforce, BaseChange ?? 0);
 
-			// Get the initial directory state
-			Bundle<DirectoryNode> CommitBundle;
-			if (BaseChange == null)
-			{
-				await FlushWorkspaceAsync(ClientInfo, Perforce, 0);
-				CommitBundle = Bundle.Create<DirectoryNode>(StorageClient, Options.NamespaceId, Options.Bundle, MemoryCache);
-			}
-			else
-			{
-				await FlushWorkspaceAsync(ClientInfo, Perforce, BaseChange.Value);
-				CommitBundle = await Bundle.ReadAsync<DirectoryNode>(StorageClient, Options.NamespaceId, GetStreamBucketId(Stream.Id), new CommitKey(Stream.Id, BaseChange.Value, Filter, RevisionsOnly).GetRefId(), Options.Bundle, MemoryCache);
-			}
+			// Get the initial bundle state
+			using Bundle<DirectoryNode> CommitBundle = await CreateOrReadBundle(Stream.Id, BaseChange, Filter, RevisionsOnly);
 
 			// Apply all the updates
 			Logger.LogInformation("Updating client {Client} to changelist {Change}", ClientInfo.Client.Name, Change);
diff --git a/Engine/Source/Programs/Horde/Horde.Build/Program.cs b/Engine/Source/Programs/Horde/Horde.Build/Program.cs
--- a/Engine/Source/Programs/Horde/Horde.Build/Program.cs
+++ b/Engine/Source/Programs/Horde/Horde.Build/Program.cs
@@ -21,8 +21,8 @@ using System.Threading.Tasks;
 using EpicGames.Core;
 using EpicGames.Perforce;
 using HordeCommon;
-using HordeServer.Models;
-using HordeServer.Utilities;
+using Horde.Build.Models;
+using Horde.Build.Utilities;
 using Microsoft.AspNetCore.Hosting;
 using Microsoft.AspNetCore.Server.Kestrel.Core;
 using Microsoft.Extensions.Configuration;
@@ -46,7 +46,7 @@ using OpenTracing.Propagation;
 using Microsoft.Extensions.DependencyInjection;
 using Horde.Build.Commands;
 
-namespace HordeServer
+namespace Horde.Build
 {
 	static class LoggerExtensions
 	{
diff --git a/Engine/Source/Programs/Horde/Horde.Build/Projects/Impl/ProjectService.cs b/Engine/Source/Programs/Horde/Horde.Build/Projects/Impl/ProjectService.cs
--- a/Engine/Source/Programs/Horde/Horde.Build/Projects/Impl/ProjectService.cs
+++ b/Engine/Source/Programs/Horde/Horde.Build/Projects/Impl/ProjectService.cs
@@ -1,9 +1,9 @@
 // Copyright Epic Games, Inc. All Rights Reserved.
 
-using HordeServer.Api;
-using HordeServer.Collections;
-using HordeServer.Models;
-using HordeServer.Utilities;
+using Horde.Build.Api;
+using Horde.Build.Collections;
+using Horde.Build.Models;
+using Horde.Build.Utilities;
 using Microsoft.Extensions.Options;
 using MongoDB.Bson;
 using MongoDB.Driver;
@@ -13,7 +13,7 @@ using System.Linq;
 using System.Security.Claims;
 using System.Threading.Tasks;
 
-namespace HordeServer.Services
+namespace Horde.Build.Services
 {
 	using ProjectId = StringId<IProject>;
 
diff --git a/Engine/Source/Programs/Horde/Horde.Build/Projects/Impl/StreamService.cs b/Engine/Source/Programs/Horde/Horde.Build/Projects/Impl/StreamService.cs
--- a/Engine/Source/Programs/Horde/Horde.Build/Projects/Impl/StreamService.cs
+++ b/Engine/Source/Programs/Horde/Horde.Build/Projects/Impl/StreamService.cs
@@ -1,9 +1,9 @@
 // Copyright Epic Games, Inc. All Rights Reserved.
 
-using HordeServer.Api;
+using Horde.Build.Api;
 using HordeCommon;
-using HordeServer.Models;
-using HordeServer.Utilities;
+using Horde.Build.Models;
+using Horde.Build.Utilities;
 using Microsoft.Extensions.Options;
 using MongoDB.Bson;
 using MongoDB.Driver;
@@ -14,10 +14,10 @@ using System.Security.Claims;
 using System.Text;
 using System.Threading.Tasks;
 using System.Linq.Expressions;
-using HordeServer.Collections;
+using Horde.Build.Collections;
 using Microsoft.Extensions.Caching.Memory;
 
-namespace HordeServer.Services
+namespace Horde.Build.Services
 {
 	using JobId = ObjectId<IJob>;
 	using ProjectId = StringId<IProject>;
diff --git a/Engine/Source/Programs/Horde/Horde.Build/Projects/ScheduleService.cs b/Engine/Source/Programs/Horde/Horde.Build/Projects/ScheduleService.cs
--- a/Engine/Source/Programs/Horde/Horde.Build/Projects/ScheduleService.cs
+++ b/Engine/Source/Programs/Horde/Horde.Build/Projects/ScheduleService.cs
@@ -3,11 +3,10 @@
 using Amazon.S3.Model.Internal.MarshallTransformations;
 using EpicGames.Core;
 using HordeCommon;
-using HordeServer.Api;
-using HordeServer.Collections;
-using HordeServer.Controllers;
-using HordeServer.Models;
-using HordeServer.Utilities;
+using Horde.Build.Api;
+using Horde.Build.Collections;
+using Horde.Build.Controllers;
+using Horde.Build.Models;
 using Horde.Build.Utilities;
 using Microsoft.AspNetCore.Server.Kestrel.Core.Features;
 using Microsoft.Extensions.Hosting;
@@ -31,7 +30,7 @@ using EpicGames.Serialization;
 using System.Text.RegularExpressions;
 using EpicGames.Redis.Utility;
 
-namespace HordeServer.Services
+namespace Horde.Build.Services
 {
 	using JobId = ObjectId<IJob>;
 	using StreamId = StringId<IStream>;
diff --git a/Engine/Source/Programs/Horde/Horde.Build/Secrets/CredentialService.cs b/Engine/Source/Programs/Horde/Horde.Build/Secrets/CredentialService.cs
--- a/Engine/Source/Programs/Horde/Horde.Build/Secrets/CredentialService.cs
+++ b/Engine/Source/Programs/Horde/Horde.Build/Secrets/CredentialService.cs
@@ -1,7 +1,7 @@
 // Copyright Epic Games, Inc. All Rights Reserved.
 
-using HordeServer.Api;
-using HordeServer.Models;
+using Horde.Build.Api;
+using Horde.Build.Models;
 using MongoDB.Bson;
 using MongoDB.Driver;
 using System;
@@ -10,7 +10,7 @@ using System.Linq;
 using System.Security.Claims;
 using System.Threading.Tasks;
 
-namespace HordeServer.Services
+namespace Horde.Build.Services
 {
 	/// <summary>
 	/// Service for accessing and modifying credentials
diff --git a/Engine/Source/Programs/Horde/Horde.Build/Server/ConfigService.cs b/Engine/Source/Programs/Horde/Horde.Build/Server/ConfigService.cs
--- a/Engine/Source/Programs/Horde/Horde.Build/Server/ConfigService.cs
+++ b/Engine/Source/Programs/Horde/Horde.Build/Server/ConfigService.cs
@@ -1,12 +1,12 @@
 // Copyright Epic Games, Inc. All Rights Reserved.
 
 using EpicGames.Core;
-using HordeServer.Api;
-using HordeServer.Collections;
-using HordeServer.Controllers;
-using HordeServer.Models;
-using HordeServer.Notifications;
-using HordeServer.Utilities;
+using Horde.Build.Api;
+using Horde.Build.Collections;
+using Horde.Build.Controllers;
+using Horde.Build.Models;
+using Horde.Build.Notifications;
+using Horde.Build.Utilities;
 using Microsoft.AspNetCore.StaticFiles;
 using Microsoft.Extensions.Configuration;
 using Microsoft.Extensions.Logging;
@@ -25,12 +25,12 @@ using System.Threading.Tasks;
 using System.Reflection;
 
 using System.Globalization;
-using HordeServer.Storage;
+using Horde.Build.Storage;
 using System.Text.Json.Serialization;
 using HordeCommon;
 using Microsoft.Extensions.Hosting;
 
-namespace HordeServer.Services
+namespace Horde.Build.Services
 {
 	using PoolId = StringId<IPool>;
 	using ProjectId = StringId<IProject>;
diff --git a/Engine/Source/Programs/Horde/Horde.Build/Server/DatabaseService.cs b/Engine/Source/Programs/Horde/Horde.Build/Server/DatabaseService.cs
--- a/Engine/Source/Programs/Horde/Horde.Build/Server/DatabaseService.cs
+++ b/Engine/Source/Programs/Horde/Horde.Build/Server/DatabaseService.cs
@@ -1,10 +1,10 @@
 // Copyright Epic Games, Inc. All Rights Reserved.
 
 using EpicGames.Core;
-using HordeServer.Api;
+using Horde.Build.Api;
 using HordeCommon;
-using HordeServer.Models;
-using HordeServer.Utilities;
+using Horde.Build.Models;
+using Horde.Build.Utilities;
 using Microsoft.Extensions.Logging;
 using Microsoft.Extensions.Options;
 using Microsoft.IdentityModel.Tokens;
@@ -26,18 +26,18 @@ using System.Threading.Tasks;
 using System.Net.Sockets;
 using MongoDB.Driver.Core.Events;
 using System.Threading;
-using HordeServer.Collections;
+using Horde.Build.Collections;
 using System.Diagnostics;
 using System.Text.RegularExpressions;
 using System.Net.NetworkInformation;
 using OpenTracing.Util;
 using OpenTracing;
-using HordeServer.Utiltiies;
+using Horde.Build.Utiltiies;
 using System.Globalization;
 
-namespace HordeServer.Services
+namespace Horde.Build.Services
 {
-	using IStream = HordeServer.Models.IStream;
+	using IStream = Horde.Build.Models.IStream;
 	using PoolId = StringId<IPool>;
 	using ProjectId = StringId<IProject>;
 	using TemplateRefId = StringId<TemplateRef>;
diff --git a/Engine/Source/Programs/Horde/Horde.Build/Server/DowntimeService.cs b/Engine/Source/Programs/Horde/Horde.Build/Server/DowntimeService.cs
--- a/Engine/Source/Programs/Horde/Horde.Build/Server/DowntimeService.cs
+++ b/Engine/Source/Programs/Horde/Horde.Build/Server/DowntimeService.cs
@@ -1,8 +1,8 @@
 // Copyright Epic Games, Inc. All Rights Reserved.
 
 using HordeCommon;
-using HordeServer.Models;
-using HordeServer.Utilities;
+using Horde.Build.Models;
+using Horde.Build.Utilities;
 using Microsoft.Extensions.Hosting;
 using Microsoft.Extensions.Logging;
 using Microsoft.Extensions.Options;
@@ -13,7 +13,7 @@ using System.Security.Cryptography.X509Certificates;
 using System.Threading;
 using System.Threading.Tasks;
 
-namespace HordeServer.Services
+namespace Horde.Build.Services
 {
 	/// <summary>
 	/// Interface for a service which keeps track of whether we're during downtime
diff --git a/Engine/Source/Programs/Horde/Horde.Build/Server/LifetimeService.cs b/Engine/Source/Programs/Horde/Horde.Build/Server/LifetimeService.cs
--- a/Engine/Source/Programs/Horde/Horde.Build/Server/LifetimeService.cs
+++ b/Engine/Source/Programs/Horde/Horde.Build/Server/LifetimeService.cs
@@ -11,11 +11,11 @@ using System.Threading;
 using System.Threading.Tasks;
 using Grpc.Core;
 using Grpc.Core.Interceptors;
-using HordeServer.Utilities;
+using Horde.Build.Utilities;
 using StackExchange.Redis;
 using System.Diagnostics;
 
-namespace HordeServer.Services
+namespace Horde.Build.Services
 {
 	/// <summary>
 	/// Service containing an async task that allows long polling operations to complete early if the server is shutting down
diff --git a/Engine/Source/Programs/Horde/Horde.Build/Server/MetricService.cs b/Engine/Source/Programs/Horde/Horde.Build/Server/MetricService.cs
--- a/Engine/Source/Programs/Horde/Horde.Build/Server/MetricService.cs
+++ b/Engine/Source/Programs/Horde/Horde.Build/Server/MetricService.cs
@@ -1,6 +1,6 @@
 // Copyright Epic Games, Inc. All Rights Reserved.
 
-using HordeServer.Utilities;
+using Horde.Build.Utilities;
 using Microsoft.Extensions.Logging;
 using System;
 using System.Threading;
@@ -9,7 +9,7 @@ using StatsdClient;
 using Microsoft.Extensions.Hosting;
 using HordeCommon;
 
-namespace HordeServer.Services
+namespace Horde.Build.Services
 {
 	/// <summary>
 	/// Periodically send metrics for the CLR and other services that cannot be collected on a per-request basis
diff --git a/Engine/Source/Programs/Horde/Horde.Build/Server/RequestTrackerService.cs b/Engine/Source/Programs/Horde/Horde.Build/Server/RequestTrackerService.cs
--- a/Engine/Source/Programs/Horde/Horde.Build/Server/RequestTrackerService.cs
+++ b/Engine/Source/Programs/Horde/Horde.Build/Server/RequestTrackerService.cs
@@ -10,7 +10,7 @@ using System.Threading.Tasks;
 using Microsoft.AspNetCore.Http;
 using Microsoft.Extensions.Logging;
 
-namespace HordeServer.Utilities
+namespace Horde.Build.Utilities
 {
 	/// <summary>
 	/// Tracks open HTTP requests to ASP.NET
diff --git a/Engine/Source/Programs/Horde/Horde.Build/ServerSettings.cs b/Engine/Source/Programs/Horde/Horde.Build/ServerSettings.cs
--- a/Engine/Source/Programs/Horde/Horde.Build/ServerSettings.cs
+++ b/Engine/Source/Programs/Horde/Horde.Build/ServerSettings.cs
@@ -1,12 +1,12 @@
 // Copyright Epic Games, Inc. All Rights Reserved.
 
 using EpicGames.Core;
-using HordeServer.Logs.Readers;
-using HordeServer.Models;
-using HordeServer.Services;
-using HordeServer.Storage;
-using HordeServer.Storage.Backends;
-using HordeServer.Utilities;
+using Horde.Build.Logs.Readers;
+using Horde.Build.Models;
+using Horde.Build.Services;
+using Horde.Build.Storage;
+using Horde.Build.Storage.Backends;
+using Horde.Build.Utilities;
 using System;
 using System.Collections.Generic;
 using System.ComponentModel.DataAnnotations;
@@ -17,7 +17,7 @@ using Horde.Build.Fleet.Autoscale;
 using TimeZoneConverter;
 using EpicGames.Horde.Storage;
 
-namespace HordeServer
+namespace Horde.Build
 {
 	/// <summary>
 	/// Types of storage to use for log data
diff --git a/Engine/Source/Programs/Horde/Horde.Build/Services/RpcService.cs b/Engine/Source/Programs/Horde/Horde.Build/Services/RpcService.cs
--- a/Engine/Source/Programs/Horde/Horde.Build/Services/RpcService.cs
+++ b/Engine/Source/Programs/Horde/Horde.Build/Services/RpcService.cs
@@ -5,9 +5,9 @@ using Google.Protobuf.WellKnownTypes;
 using Grpc.Core;
 using HordeCommon;
 using HordeCommon.Rpc;
-using HordeServer.Collections;
-using HordeServer.Models;
-using HordeServer.Utilities;
+using Horde.Build.Collections;
+using Horde.Build.Models;
+using Horde.Build.Utilities;
 using Microsoft.AspNetCore.Authorization;
 using Microsoft.Extensions.Logging;
 using MongoDB.Bson;
@@ -21,15 +21,15 @@ using System.Runtime.CompilerServices;
 using System.Security.Claims;
 using System.Text;
 using System.Threading.Tasks;
-using HordeServer.Tasks.Impl;
+using Horde.Build.Tasks.Impl;
 using System.Threading;
-using HordeServer.Jobs;
+using Horde.Build.Jobs;
 
-namespace HordeServer.Services
+namespace Horde.Build.Services
 {
 	using AgentSoftwareChannelName = StringId<AgentSoftwareChannels>;
 	using AgentSoftwareVersion = StringId<IAgentSoftwareCollection>;
-	using IStream = HordeServer.Models.IStream;
+	using IStream = Horde.Build.Models.IStream;
 	using JobId = ObjectId<IJob>;
 	using LogId = ObjectId<ILogFile>;
 	using SessionId = ObjectId<ISession>;
@@ -865,7 +865,7 @@ namespace HordeServer.Services
 			{
 				throw new StructuredRpcException(StatusCode.NotFound, "Resource not found");
 			}
-			if (!HordeServer.Services.LogFileService.AuthorizeForSession(LogFile, Context.GetHttpContext().User))
+			if (!Horde.Build.Services.LogFileService.AuthorizeForSession(LogFile, Context.GetHttpContext().User))
 			{
 				throw new StructuredRpcException(StatusCode.PermissionDenied, "Access denied");
 			}
diff --git a/Engine/Source/Programs/Horde/Horde.Build/Startup.cs b/Engine/Source/Programs/Horde/Horde.Build/Startup.cs
--- a/Engine/Source/Programs/Horde/Horde.Build/Startup.cs
+++ b/Engine/Source/Programs/Horde/Horde.Build/Startup.cs
@@ -13,8 +13,8 @@ using System.Text.Json;
 using System.Text.Json.Serialization;
 using System.Threading.Tasks;
 using Serilog;
-using HordeServer.Models;
-using HordeServer.Services;
+using Horde.Build.Models;
+using Horde.Build.Services;
 using Microsoft.AspNetCore.Authentication;
 using Microsoft.AspNetCore.Authentication.Cookies;
 using Microsoft.AspNetCore.Authentication.JwtBearer;
@@ -35,13 +35,13 @@ using MongoDB.Bson;
 using MongoDB.Bson.Serialization;
 using MongoDB.Bson.Serialization.Conventions;
 using Microsoft.AspNetCore.HttpOverrides;
-using HordeServer.Utilities;
+using Horde.Build.Utilities;
 using Grpc.Core.Interceptors;
 using Grpc.Core;
 using EpicGames.Core;
 using Microsoft.Extensions.Diagnostics.HealthChecks;
 using Microsoft.AspNetCore.Diagnostics.HealthChecks;
-using HordeServer.Collections;
+using Horde.Build.Collections;
 using Microsoft.AspNetCore.Http;
 using System.Diagnostics.CodeAnalysis;
 using System.Globalization;
@@ -49,41 +49,40 @@ using System.Threading;
 using HordeCommon;
 using Microsoft.AspNetCore.DataProtection;
 using StackExchange.Redis;
-using HordeServer.Storage;
-using HordeServer.Logs;
-using HordeServer.Logs.Builder;
-using HordeServer.Logs.Storage;
-using HordeServer.Logs.Readers;
-using HordeServer.Logs.Storage.Impl;
-using HordeServer.Collections.Impl;
-using HordeServer.Services.Impl;
+using Horde.Build.Storage;
+using Horde.Build.Logs;
+using Horde.Build.Logs.Builder;
+using Horde.Build.Logs.Storage;
+using Horde.Build.Logs.Readers;
+using Horde.Build.Logs.Storage.Impl;
+using Horde.Build.Collections.Impl;
+using Horde.Build.Services.Impl;
 using Microsoft.Extensions.DependencyInjection.Extensions;
-using HordeServer.Authentication;
-using HordeServer.Tasks.Impl;
-using HordeServer.Tasks;
+using Horde.Build.Authentication;
+using Horde.Build.Tasks.Impl;
+using Horde.Build.Tasks;
 using StatsdClient;
 using Status = Grpc.Core.Status;
-using HordeServer.Notifications.Impl;
-using HordeServer.Notifications;
+using Horde.Build.Notifications.Impl;
+using Horde.Build.Notifications;
 using Microsoft.AspNetCore.Hosting.Server.Features;
 using System.Runtime.InteropServices;
 using Amazon.Extensions.NETCore.Setup;
 using Amazon.S3;
 using Amazon.Runtime;
-using HordeServer.Storage.Backends;
-using HordeServer.Commits.Impl;
-using HordeServer.Commits;
+using Horde.Build.Storage.Backends;
+using Horde.Build.Commits.Impl;
+using Horde.Build.Commits;
 using OpenTracing.Contrib.Grpc.Interceptors;
 using OpenTracing.Util;
 using EpicGames.Horde.Compute;
-using HordeServer.Compute.Impl;
-using HordeServer.Compute;
+using Horde.Build.Compute.Impl;
+using Horde.Build.Compute;
 using System.Net.Http.Headers;
 using Amazon.SecurityToken.Model;
 using Horde.Build.Fleet.Autoscale;
-using Horde.Build.Utilities;
 using Serilog.Events;
-using HordeServer.Jobs;
+using Horde.Build.Jobs;
 using Microsoft.AspNetCore.Mvc;
 using Microsoft.AspNetCore.Mvc.ModelBinding;
 using Horde.Build.Storage.Services;
@@ -93,7 +92,7 @@ using EpicGames.AspNet;
 using EpicGames.Horde.Storage.Impl;
 using MongoDB.Bson.Serialization.Serializers;
 
-namespace HordeServer
+namespace Horde.Build
 {
 	using MessageTemplate = EpicGames.Core.MessageTemplate;
 	using JobId = ObjectId<IJob>;
diff --git a/Engine/Source/Programs/Horde/Horde.Build/Storage/Backends/AwsStorageBackend.cs b/Engine/Source/Programs/Horde/Horde.Build/Storage/Backends/AwsStorageBackend.cs
--- a/Engine/Source/Programs/Horde/Horde.Build/Storage/Backends/AwsStorageBackend.cs
+++ b/Engine/Source/Programs/Horde/Horde.Build/Storage/Backends/AwsStorageBackend.cs
@@ -4,7 +4,7 @@ using Amazon.Extensions.NETCore.Setup;
 using Amazon.Runtime;
 using Amazon.S3;
 using Amazon.S3.Model;
-using HordeServer.Utilities;
+using Horde.Build.Utilities;
 using Microsoft.Extensions.Configuration;
 using Microsoft.Extensions.Logging;
 using Microsoft.Extensions.Options;
@@ -17,7 +17,7 @@ using System.Text;
 using System.Threading;
 using System.Threading.Tasks;
 
-namespace HordeServer.Storage.Backends
+namespace Horde.Build.Storage.Backends
 {
 	/// <summary>
 	/// Exception wrapper for S3 requests
diff --git a/Engine/Source/Programs/Horde/Horde.Build/Storage/Backends/FileSystemStorageBackend.cs b/Engine/Source/Programs/Horde/Horde.Build/Storage/Backends/FileSystemStorageBackend.cs
--- a/Engine/Source/Programs/Horde/Horde.Build/Storage/Backends/FileSystemStorageBackend.cs
+++ b/Engine/Source/Programs/Horde/Horde.Build/Storage/Backends/FileSystemStorageBackend.cs
@@ -1,7 +1,7 @@
 // Copyright Epic Games, Inc. All Rights Reserved.
 
 using EpicGames.Core;
-using HordeServer.Utilities;
+using Horde.Build.Utilities;
 using Microsoft.Extensions.Logging;
 using Microsoft.Extensions.Options;
 using System;
@@ -11,7 +11,7 @@ using System.Linq;
 using System.Threading;
 using System.Threading.Tasks;
 
-namespace HordeServer.Storage.Backends
+namespace Horde.Build.Storage.Backends
 {
 	/// <summary>
 	/// Options for the filesystem backend
diff --git a/Engine/Source/Programs/Horde/Horde.Build/Storage/Backends/RelayStorageBackend.cs b/Engine/Source/Programs/Horde/Horde.Build/Storage/Backends/RelayStorageBackend.cs
--- a/Engine/Source/Programs/Horde/Horde.Build/Storage/Backends/RelayStorageBackend.cs
+++ b/Engine/Source/Programs/Horde/Horde.Build/Storage/Backends/RelayStorageBackend.cs
@@ -10,7 +10,7 @@ using System.Net.Http;
 using System.Net.Http.Headers;
 using System.Threading.Tasks;
 
-namespace HordeServer.Storage.Backends
+namespace Horde.Build.Storage.Backends
 {
 	/// <summary>
 	/// Implementation of ILogFileStorage which forwards requests to another server
diff --git a/Engine/Source/Programs/Horde/Horde.Build/Storage/Backends/TransientStorageBackend.cs b/Engine/Source/Programs/Horde/Horde.Build/Storage/Backends/TransientStorageBackend.cs
--- a/Engine/Source/Programs/Horde/Horde.Build/Storage/Backends/TransientStorageBackend.cs
+++ b/Engine/Source/Programs/Horde/Horde.Build/Storage/Backends/TransientStorageBackend.cs
@@ -7,7 +7,7 @@ using System.IO;
 using System.Linq;
 using System.Threading.Tasks;
 
-namespace HordeServer.Storage.Backends
+namespace Horde.Build.Storage.Backends
 {
 	/// <summary>
 	/// In-memory implementation of ILogFileStorage
diff --git a/Engine/Source/Programs/Horde/Horde.Build/Storage/IStorageBackend.cs b/Engine/Source/Programs/Horde/Horde.Build/Storage/IStorageBackend.cs
--- a/Engine/Source/Programs/Horde/Horde.Build/Storage/IStorageBackend.cs
+++ b/Engine/Source/Programs/Horde/Horde.Build/Storage/IStorageBackend.cs
@@ -7,7 +7,7 @@ using System.IO;
 using System.Linq;
 using System.Threading.Tasks;
 
-namespace HordeServer.Storage
+namespace Horde.Build.Storage
 {
 	/// <summary>
 	/// Interface for a traditional location-addressed storage provider
diff --git a/Engine/Source/Programs/Horde/Horde.Build/Storage/Services/BasicStorageClient.cs b/Engine/Source/Programs/Horde/Horde.Build/Storage/Services/BasicStorageClient.cs
--- a/Engine/Source/Programs/Horde/Horde.Build/Storage/Services/BasicStorageClient.cs
+++ b/Engine/Source/Programs/Horde/Horde.Build/Storage/Services/BasicStorageClient.cs
@@ -3,9 +3,9 @@
 using EpicGames.Core;
 using EpicGames.Horde.Storage;
 using EpicGames.Serialization;
-using HordeServer.Services;
-using HordeServer.Storage;
-using HordeServer.Utilities;
+using Horde.Build.Services;
+using Horde.Build.Storage;
+using Horde.Build.Utilities;
 using MongoDB.Bson.Serialization.Attributes;
 using MongoDB.Driver;
 using System;
diff --git a/Engine/Source/Programs/Horde/Horde.Build/Ugs/UgsMetadata.cs b/Engine/Source/Programs/Horde/Horde.Build/Ugs/UgsMetadata.cs
--- a/Engine/Source/Programs/Horde/Horde.Build/Ugs/UgsMetadata.cs
+++ b/Engine/Source/Programs/Horde/Horde.Build/Ugs/UgsMetadata.cs
@@ -1,14 +1,14 @@
 // Copyright Epic Games, Inc. All Rights Reserved.
 
 using Amazon.S3.Model;
-using HordeServer.Api;
+using Horde.Build.Api;
 using MongoDB.Bson;
 using System;
 using System.Collections.Generic;
 using System.Linq;
 using System.Threading.Tasks;
 
-namespace HordeServer.Models
+namespace Horde.Build.Models
 {
 	/// <summary>
 	/// Metadata for a particular user
diff --git a/Engine/Source/Programs/Horde/Horde.Build/Users/Impl/IServiceAccountCollection.cs b/Engine/Source/Programs/Horde/Horde.Build/Users/Impl/IServiceAccountCollection.cs
--- a/Engine/Source/Programs/Horde/Horde.Build/Users/Impl/IServiceAccountCollection.cs
+++ b/Engine/Source/Programs/Horde/Horde.Build/Users/Impl/IServiceAccountCollection.cs
@@ -1,11 +1,11 @@
 // Copyright Epic Games, Inc. All Rights Reserved.
 
-using HordeServer.Models;
+using Horde.Build.Models;
 using MongoDB.Bson;
 using System.Collections.Generic;
 using System.Threading.Tasks;
 
-namespace HordeServer.Collections
+namespace Horde.Build.Collections
 {
 	/// <summary>
 	/// Interface for a collection of service account documents
diff --git a/Engine/Source/Programs/Horde/Horde.Build/Users/Impl/ServiceAccountCollection.cs b/Engine/Source/Programs/Horde/Horde.Build/Users/Impl/ServiceAccountCollection.cs
--- a/Engine/Source/Programs/Horde/Horde.Build/Users/Impl/ServiceAccountCollection.cs
+++ b/Engine/Source/Programs/Horde/Horde.Build/Users/Impl/ServiceAccountCollection.cs
@@ -1,8 +1,8 @@
 // Copyright Epic Games, Inc. All Rights Reserved.
 
 using System;
-using HordeServer.Models;
-using HordeServer.Services;
+using Horde.Build.Models;
+using Horde.Build.Services;
 using MongoDB.Bson;
 using MongoDB.Bson.Serialization.Attributes;
 using MongoDB.Driver;
@@ -10,7 +10,7 @@ using System.Collections.Generic;
 using System.Linq;
 using System.Threading.Tasks;
 
-namespace HordeServer.Collections.Impl
+namespace Horde.Build.Collections.Impl
 {
 	/// <summary>
 	/// Collection of service account documents
diff --git a/Engine/Source/Programs/Horde/Horde.Build/Utilities/AgentUtilities.cs b/Engine/Source/Programs/Horde/Horde.Build/Utilities/AgentUtilities.cs
--- a/Engine/Source/Programs/Horde/Horde.Build/Utilities/AgentUtilities.cs
+++ b/Engine/Source/Programs/Horde/Horde.Build/Utilities/AgentUtilities.cs
@@ -12,7 +12,7 @@ using System.Diagnostics;
 
 #nullable disable
 
-namespace HordeServer.Utilities
+namespace Horde.Build.Utilities
 {
 	static class AgentUtilities
 	{
diff --git a/Engine/Source/Programs/Horde/Horde.Build/Utilities/BlockKit/ActionsBlock.cs b/Engine/Source/Programs/Horde/Horde.Build/Utilities/BlockKit/ActionsBlock.cs
--- a/Engine/Source/Programs/Horde/Horde.Build/Utilities/BlockKit/ActionsBlock.cs
+++ b/Engine/Source/Programs/Horde/Horde.Build/Utilities/BlockKit/ActionsBlock.cs
@@ -5,7 +5,7 @@ using System.Collections.Generic;
 using System.Globalization;
 using System.Text.Json;
 
-namespace HordeServer.Utilities.Slack.BlockKit
+namespace Horde.Build.Utilities.Slack.BlockKit
 {
 	/// <summary>
 	/// Wrapper for a BlockKit Actions block that contains interactive elements.
diff --git a/Engine/Source/Programs/Horde/Horde.Build/Utilities/BlockKit/BlockBase.cs b/Engine/Source/Programs/Horde/Horde.Build/Utilities/BlockKit/BlockBase.cs
--- a/Engine/Source/Programs/Horde/Horde.Build/Utilities/BlockKit/BlockBase.cs
+++ b/Engine/Source/Programs/Horde/Horde.Build/Utilities/BlockKit/BlockBase.cs
@@ -4,7 +4,7 @@ using System;
 using System.Text.Json;
 using System.Text.Json.Serialization;
 
-namespace HordeServer.Utilities.Slack.BlockKit
+namespace Horde.Build.Utilities.Slack.BlockKit
 {
 	/// <summary>
 	/// Converter for overriding default Json serialization/deserialization for BlockKit blocks.
diff --git a/Engine/Source/Programs/Horde/Horde.Build/Utilities/BlockKit/SlackResponses.cs b/Engine/Source/Programs/Horde/Horde.Build/Utilities/BlockKit/SlackResponses.cs
--- a/Engine/Source/Programs/Horde/Horde.Build/Utilities/BlockKit/SlackResponses.cs
+++ b/Engine/Source/Programs/Horde/Horde.Build/Utilities/BlockKit/SlackResponses.cs
@@ -1,9 +1,9 @@
 // Copyright Epic Games, Inc. All Rights Reserved.
 
 using System.Drawing;
-using HordeServer.Utilities.Slack.BlockKit;
+using Horde.Build.Utilities.Slack.BlockKit;
 
-namespace HordeServer.Utilities.BlockKit
+namespace Horde.Build.Utilities.BlockKit
 {
 	/// <summary>
 	/// Helper methods for creating <see cref="BlockKitMessage"/> responses to send
diff --git a/Engine/Source/Programs/Horde/Horde.Build/Utilities/BuildGraphExtensions.cs b/Engine/Source/Programs/Horde/Horde.Build/Utilities/BuildGraphExtensions.cs
--- a/Engine/Source/Programs/Horde/Horde.Build/Utilities/BuildGraphExtensions.cs
+++ b/Engine/Source/Programs/Horde/Horde.Build/Utilities/BuildGraphExtensions.cs
@@ -2,8 +2,8 @@
 
 using EpicGames.BuildGraph;
 using EpicGames.Perforce;
-using HordeServer.Collections;
-using HordeServer.Models;
+using Horde.Build.Collections;
+using Horde.Build.Models;
 using System;
 using System.Collections.Generic;
 using System.Globalization;
diff --git a/Engine/Source/Programs/Horde/Horde.Build/Utilities/Clock.cs b/Engine/Source/Programs/Horde/Horde.Build/Utilities/Clock.cs
--- a/Engine/Source/Programs/Horde/Horde.Build/Utilities/Clock.cs
+++ b/Engine/Source/Programs/Horde/Horde.Build/Utilities/Clock.cs
@@ -2,7 +2,7 @@
 
 using EpicGames.Core;
 using EpicGames.Redis.Utility;
-using HordeServer.Services;
+using Horde.Build.Services;
 using Microsoft.Extensions.Logging;
 using StackExchange.Redis;
 using System;
diff --git a/Engine/Source/Programs/Horde/Horde.Build/Utilities/HordeClaimTypes.cs b/Engine/Source/Programs/Horde/Horde.Build/Utilities/HordeClaimTypes.cs
--- a/Engine/Source/Programs/Horde/Horde.Build/Utilities/HordeClaimTypes.cs
+++ b/Engine/Source/Programs/Horde/Horde.Build/Utilities/HordeClaimTypes.cs
@@ -1,7 +1,7 @@
 // Copyright Epic Games, Inc. All Rights Reserved.
 
-using HordeServer.Collections;
-using HordeServer.Models;
+using Horde.Build.Collections;
+using Horde.Build.Models;
 using MongoDB.Bson;
 using System;
 using System.Collections.Generic;
@@ -9,7 +9,7 @@ using System.Linq;
 using System.Security.Claims;
 using System.Threading.Tasks;
 
-namespace HordeServer.Utilities
+namespace Horde.Build.Utilities
 {
 	using UserId = ObjectId<IUser>;
 
diff --git a/Engine/Source/Programs/Horde/Horde.Build/Utilities/HordeControllerBase.cs b/Engine/Source/Programs/Horde/Horde.Build/Utilities/HordeControllerBase.cs
--- a/Engine/Source/Programs/Horde/Horde.Build/Utilities/HordeControllerBase.cs
+++ b/Engine/Source/Programs/Horde/Horde.Build/Utilities/HordeControllerBase.cs
@@ -1,7 +1,7 @@
 // Copyright Epic Games, Inc. All Rights Reserved.
 
 using EpicGames.Core;
-using HordeServer.Models;
+using Horde.Build.Models;
 using Microsoft.AspNetCore.Http;
 using Microsoft.AspNetCore.Mvc;
 using Microsoft.Extensions.Logging;
@@ -13,7 +13,7 @@ using System.Linq.Expressions;
 using System.Threading;
 using System.Threading.Tasks;
 
-namespace HordeServer.Utilities
+namespace Horde.Build.Utilities
 {
 	using JobId = ObjectId<IJob>;
 	using LeaseId = ObjectId<ILease>;
diff --git a/Engine/Source/Programs/Horde/Horde.Build/Utilities/MongoTracingCollection.cs b/Engine/Source/Programs/Horde/Horde.Build/Utilities/MongoTracingCollection.cs
--- a/Engine/Source/Programs/Horde/Horde.Build/Utilities/MongoTracingCollection.cs
+++ b/Engine/Source/Programs/Horde/Horde.Build/Utilities/MongoTracingCollection.cs
@@ -9,7 +9,7 @@ using MongoDB.Driver;
 using OpenTracing;
 using OpenTracing.Util;
 
-namespace HordeServer.Utiltiies
+namespace Horde.Build.Utiltiies
 {
 	/// <summary>
 	/// Wrap a IMongoCollection with trace scopes
diff --git a/Engine/Source/Programs/Horde/Horde.Build/Utilities/ObjectSet.cs b/Engine/Source/Programs/Horde/Horde.Build/Utilities/ObjectSet.cs
--- a/Engine/Source/Programs/Horde/Horde.Build/Utilities/ObjectSet.cs
+++ b/Engine/Source/Programs/Horde/Horde.Build/Utilities/ObjectSet.cs
@@ -4,14 +4,14 @@ using EpicGames.Core;
 using EpicGames.Horde.Storage;
 using EpicGames.Serialization;
 using EpicGames.Serialization.Converters;
-using HordeServer.Storage;
+using Horde.Build.Storage;
 using System;
 using System.Collections.Generic;
 using System.Diagnostics.CodeAnalysis;
 using System.Linq;
 using System.Threading.Tasks;
 
-namespace HordeServer.Utilities
+namespace Horde.Build.Utilities
 {
 	/// <summary>
 	/// Information about a blob stored in a blob pack file
diff --git a/Engine/Source/Programs/Horde/Horde.Build/Utilities/ReadOnlyTrie.cs b/Engine/Source/Programs/Horde/Horde.Build/Utilities/ReadOnlyTrie.cs
--- a/Engine/Source/Programs/Horde/Horde.Build/Utilities/ReadOnlyTrie.cs
+++ b/Engine/Source/Programs/Horde/Horde.Build/Utilities/ReadOnlyTrie.cs
@@ -1,7 +1,7 @@
 // Copyright Epic Games, Inc. All Rights Reserved.
 
 using EpicGames.Core;
-using HordeServer.Controllers;
+using Horde.Build.Controllers;
 using System;
 using System.Collections;
 using System.Collections.Generic;
@@ -9,7 +9,7 @@ using System.Linq;
 using System.Runtime.InteropServices;
 using System.Threading.Tasks;
 
-namespace HordeServer.Logs
+namespace Horde.Build.Logs
 {
 	/// <summary>
 	/// A read-only trie of long values
diff --git a/Engine/Source/Programs/Horde/Horde.Build/Utilities/ReadOnlyTrieBuilder.cs b/Engine/Source/Programs/Horde/Horde.Build/Utilities/ReadOnlyTrieBuilder.cs
--- a/Engine/Source/Programs/Horde/Horde.Build/Utilities/ReadOnlyTrieBuilder.cs
+++ b/Engine/Source/Programs/Horde/Horde.Build/Utilities/ReadOnlyTrieBuilder.cs
@@ -2,7 +2,7 @@
 
 using Amazon.Runtime;
 using EpicGames.Core;
-using HordeServer.Logs;
+using Horde.Build.Logs;
 using Microsoft.AspNetCore.Mvc;
 using Microsoft.AspNetCore.Routing.Constraints;
 using System;
@@ -12,7 +12,7 @@ using System.Runtime.CompilerServices;
 using System.Runtime.InteropServices;
 using System.Threading.Tasks;
 
-namespace HordeServer.Logs
+namespace Horde.Build.Logs
 {
 	/// <summary>
 	/// Structure used for generating prefix trees
diff --git a/Engine/Source/Programs/Horde/Horde.Build/Utilities/SingletonDocument.cs b/Engine/Source/Programs/Horde/Horde.Build/Utilities/SingletonDocument.cs
--- a/Engine/Source/Programs/Horde/Horde.Build/Utilities/SingletonDocument.cs
+++ b/Engine/Source/Programs/Horde/Horde.Build/Utilities/SingletonDocument.cs
@@ -1,7 +1,7 @@
 // Copyright Epic Games, Inc. All Rights Reserved.
 
-using HordeServer.Models;
-using HordeServer.Services;
+using Horde.Build.Models;
+using Horde.Build.Services;
 using MongoDB.Bson;
 using System;
 using System.Collections.Generic;
@@ -9,7 +9,7 @@ using System.Linq;
 using System.Reflection;
 using System.Threading.Tasks;
 
-namespace HordeServer.Utilities
+namespace Horde.Build.Utilities
 {
 	/// <summary>
 	/// Attribute specifying the unique id for a singleton document
diff --git a/Engine/Source/Programs/Horde/Horde.Build/Utilities/TracingExtensions.cs b/Engine/Source/Programs/Horde/Horde.Build/Utilities/TracingExtensions.cs
--- a/Engine/Source/Programs/Horde/Horde.Build/Utilities/TracingExtensions.cs
+++ b/Engine/Source/Programs/Horde/Horde.Build/Utilities/TracingExtensions.cs
@@ -3,8 +3,8 @@
 using System;
 using System.Collections.Generic;
 using EpicGames.Core;
-using HordeServer.Models;
-using HordeServer.Utilities;
+using Horde.Build.Models;
+using Horde.Build.Utilities;
 using OpenTracing;
 
 namespace Horde.Build.Utilities
diff --git a/Engine/Source/Programs/Shared/EpicGames.Core/CommandLineExtensions.cs b/Engine/Source/Programs/Shared/EpicGames.Core/CommandLineExtensions.cs
--- a/Engine/Source/Programs/Shared/EpicGames.Core/CommandLineExtensions.cs
+++ b/Engine/Source/Programs/Shared/EpicGames.Core/CommandLineExtensions.cs
@@ -4,7 +4,7 @@ using System;
 using System.Collections.Generic;
 using System.Text;
 
-namespace HordeAgent.Utility
+namespace EpicGames.Core
 {
 	/// <summary>
 	/// Extensions to StringBuilder for building command lines
diff --git a/Engine/Source/Programs/Shared/EpicGames.Core/LogEvent.cs b/Engine/Source/Programs/Shared/EpicGames.Core/LogEvent.cs
--- a/Engine/Source/Programs/Shared/EpicGames.Core/LogEvent.cs
+++ b/Engine/Source/Programs/Shared/EpicGames.Core/LogEvent.cs
@@ -1,6 +1,5 @@
 // Copyright Epic Games, Inc. All Rights Reserved.
 
-using HordeAgent.Utility;
 using Microsoft.Extensions.Logging;
 using System;
 using System.Buffers;
diff --git a/Engine/Source/Programs/Shared/EpicGames.Horde.Tests/BundleTests.cs b/Engine/Source/Programs/Shared/EpicGames.Horde.Tests/BundleTests.cs
--- a/Engine/Source/Programs/Shared/EpicGames.Horde.Tests/BundleTests.cs
+++ b/Engine/Source/Programs/Shared/EpicGames.Horde.Tests/BundleTests.cs
@@ -16,7 +16,7 @@ using System.Security.Cryptography;
 using System.Text;
 using System.Threading.Tasks;
 
-namespace HordeServerTests
+namespace EpicGames.Horde.Tests
 {
 	[TestClass]
 	public class BundleTests : IDisposable
diff --git a/Engine/Source/Programs/Shared/EpicGames.Perforce.Managed/ManagedWorkspace.cs b/Engine/Source/Programs/Shared/EpicGames.Perforce.Managed/ManagedWorkspace.cs
--- a/Engine/Source/Programs/Shared/EpicGames.Perforce.Managed/ManagedWorkspace.cs
+++ b/Engine/Source/Programs/Shared/EpicGames.Perforce.Managed/ManagedWorkspace.cs
@@ -1,6 +1,5 @@
 // Copyright Epic Games, Inc. All Rights Reserved.
 
-using HordeAgent.Utility;
 using Microsoft.Extensions.Logging;
 using System;
 using System.Collections.Concurrent;
