Horde Storage - Fixed issues with the shutdown of background services.
Also increased number of replications we do by default before stopping the replication as its not really running a full thruput right now.

diff --git a/Engine/Source/Programs/Horde/HordeStorage/Horde.Storage/Implementation/Blob/ConsistencyCheckService.cs b/Engine/Source/Programs/Horde/HordeStorage/Horde.Storage/Implementation/Blob/ConsistencyCheckService.cs
--- a/Engine/Source/Programs/Horde/HordeStorage/Horde.Storage/Implementation/Blob/ConsistencyCheckService.cs
+++ b/Engine/Source/Programs/Horde/HordeStorage/Horde.Storage/Implementation/Blob/ConsistencyCheckService.cs
@@ -19,7 +19,7 @@ using Serilog;
 namespace Horde.Storage.Implementation
 {
     // ReSharper disable once ClassNeverInstantiated.Global
-    public class ConsistencyCheckService : PollingService<ConsistencyCheckService.ConsistencyState>, IDisposable
+    public class ConsistencyCheckService : PollingService<ConsistencyCheckService.ConsistencyState>
     {
         private readonly IOptionsMonitor<ConsistencyCheckSettings> _settings;
         private readonly IServiceProvider _provider;
@@ -126,11 +126,6 @@ namespace Horde.Storage.Implementation
         {
             return Task.CompletedTask;
         }
-
-        public void Dispose()
-        {
-
-        }
     }
 
     public class ConsistencyCheckSettings
diff --git a/Engine/Source/Programs/Horde/HordeStorage/Horde.Storage/Implementation/Replication/RefsReplicator.cs b/Engine/Source/Programs/Horde/HordeStorage/Horde.Storage/Implementation/Replication/RefsReplicator.cs
--- a/Engine/Source/Programs/Horde/HordeStorage/Horde.Storage/Implementation/Replication/RefsReplicator.cs
+++ b/Engine/Source/Programs/Horde/HordeStorage/Horde.Storage/Implementation/Replication/RefsReplicator.cs
@@ -41,6 +41,7 @@ namespace Horde.Storage.Implementation
         private readonly HttpClient _httpClient;
         private RefsState _refsState;
         private bool _replicationRunning;
+        private bool _disposed = false;
 
         public RefsReplicator(ReplicatorSettings replicatorSettings, IOptionsMonitor<ReplicationSettings> replicationSettings, IBlobService blobService, IHttpClientFactory httpClientFactory, IReplicationLog replicationLog, IServiceCredentials serviceCredentials)
         {
@@ -83,9 +84,15 @@ namespace Horde.Storage.Implementation
 
         public void Dispose()
         {
+            if (_disposed)
+                return;
+
             SaveState(_stateFile, _refsState);
 
+            _replicationFinishedEvent.Wait();
             _replicationTokenSource.Dispose();
+
+            _disposed = true;
         }
 
         private static RefsState ReadState(FileInfo stateFile)
@@ -373,6 +380,9 @@ namespace Horde.Storage.Implementation
             CancellationTokenSource linkedTokenSource = CancellationTokenSource.CreateLinkedTokenSource(cancellationTokenSource.Token, replicationToken);
             ISpanContext parentSpanContext = Tracer.Instance.ActiveScope.Span.Context;
 
+            if (replicationToken.IsCancellationRequested)
+                return countOfReplicationsDone;
+
             await GetRefEvents(ns, lastBucket, lastEvent, replicationToken).ParallelForEachAsync(async (ReplicationLogEvent @event) =>
             {
                 // if we have done all the replication events we should do in a single run we abort
@@ -591,6 +601,9 @@ namespace Horde.Storage.Implementation
 
         public async Task StopReplicating()
         {
+            if (_disposed)
+                return;
+
             _replicationTokenSource.Cancel(true);
             await _replicationFinishedEvent.WaitAsync();
         }
diff --git a/Engine/Source/Programs/Horde/HordeStorage/Horde.Storage/Implementation/Replication/ReplicationService.cs b/Engine/Source/Programs/Horde/HordeStorage/Horde.Storage/Implementation/Replication/ReplicationService.cs
--- a/Engine/Source/Programs/Horde/HordeStorage/Horde.Storage/Implementation/Replication/ReplicationService.cs
+++ b/Engine/Source/Programs/Horde/HordeStorage/Horde.Storage/Implementation/Replication/ReplicationService.cs
@@ -16,7 +16,7 @@ using Serilog;
 namespace Horde.Storage.Implementation
 {
     // ReSharper disable once ClassNeverInstantiated.Global
-    public class ReplicationService : PollingService<ReplicationService.ReplicationState>, IDisposable
+    public class ReplicationService : PollingService<ReplicationService.ReplicationState>
     {
         private readonly IOptionsMonitor<ReplicationSettings> _settings;
         private readonly ILeaderElection _leaderElection;
@@ -64,10 +64,8 @@ namespace Horde.Storage.Implementation
                 return;
 
             // if we are no longer the leader cancel any pending replications
-            foreach (IReplicator replicator in State.Replicators)
-            {
-                //TODO: Cancel replications
-            }
+            Task[] stopReplicatingTasks = State.Replicators.Select(replicator => replicator.StopReplicating()).ToArray();
+            Task.WaitAll(stopReplicatingTasks);
         }
 
         private static IReplicator CreateReplicator(ReplicatorSettings replicatorSettings, IServiceProvider provider)
@@ -143,17 +141,14 @@ namespace Horde.Storage.Implementation
         protected override async Task OnStopping(ReplicationState state)
         {
             await Task.WhenAll(state.Replicators.Select(replicator => replicator.StopReplicating()).ToArray());
-        }
 
-        public void Dispose()
-        {
             // we should have been stopped first so this should not be needed, but to make sure state is stored we dispose of it again
             foreach (IReplicator replicator in State.Replicators)
             {
                 replicator.Dispose();
             }
         }
-
+        
         public IEnumerable<IReplicator> GetReplicators(NamespaceId ns)
         {
             return State.Replicators
diff --git a/Engine/Source/Programs/Horde/HordeStorage/Horde.Storage/Implementation/Replication/ReplicationSettings.cs b/Engine/Source/Programs/Horde/HordeStorage/Horde.Storage/Implementation/Replication/ReplicationSettings.cs
--- a/Engine/Source/Programs/Horde/HordeStorage/Horde.Storage/Implementation/Replication/ReplicationSettings.cs
+++ b/Engine/Source/Programs/Horde/HordeStorage/Horde.Storage/Implementation/Replication/ReplicationSettings.cs
@@ -75,7 +75,7 @@ namespace Horde.Storage
         /// <summary>
         /// The replication will stop its run after this many events have been processed, yielding its time until the next time a replication is triggered.
         /// </summary>
-        public int MaxReplicationsPerRun { get; set; } = 1000;
+        public int MaxReplicationsPerRun { get; set; } = 5000;
 
         /// <summary>
         /// Max number of offsets to skip when finding broken events in the transaction logs, higher value means we scan further increasing the likelihood of finding a valid entry, but also means risk that we skip records you would want.
diff --git a/Engine/Source/Programs/Horde/HordeStorage/Jupiter.Common/PollingService.cs b/Engine/Source/Programs/Horde/HordeStorage/Jupiter.Common/PollingService.cs
--- a/Engine/Source/Programs/Horde/HordeStorage/Jupiter.Common/PollingService.cs
+++ b/Engine/Source/Programs/Horde/HordeStorage/Jupiter.Common/PollingService.cs
@@ -9,7 +9,7 @@ using Serilog;
 
 namespace Jupiter
 {
-    public abstract class PollingService<T> : IHostedService
+    public abstract class PollingService<T> : IHostedService, IAsyncDisposable
     {
         private struct ThreadState
         {
@@ -29,6 +29,7 @@ namespace Jupiter
         private readonly ManualResetEvent _hasFinishedRunning = new ManualResetEvent(true);
         private volatile bool _alreadyRunning = false;
         private Timer? _timer;
+        private bool _disposed = false;
 
         protected PollingService(string serviceName, TimeSpan pollFrequency, T state)
         {
@@ -147,13 +148,23 @@ namespace Jupiter
         {
             _logger.Information("{Service} poll service stopping.", _serviceName);
 
+            if (_timer != null)
+                await _timer.DisposeAsync();
+            _timer = null;
+
             await OnStopping(_state);
 
             _stopPolling.Cancel();
             _hasFinishedRunning.WaitOne();
-            if (_timer != null)
-                await _timer.DisposeAsync();
-            _timer = null;
+        }
+
+        public async ValueTask DisposeAsync()
+        {
+            if (_disposed)
+                return;
+
+            _disposed = true;
+            await StopAsync(CancellationToken.None);
         }
     }
 }
