Horde Storage - Disabling caching for our full table scan, a scylla specific extension to CQL to limit cache thrasing.

diff --git a/Engine/Source/Programs/Horde/HordeStorage/Horde.Storage/Implementation/Objects/ScyllaReferencesStore.cs b/Engine/Source/Programs/Horde/HordeStorage/Horde.Storage/Implementation/Objects/ScyllaReferencesStore.cs
--- a/Engine/Source/Programs/Horde/HordeStorage/Horde.Storage/Implementation/Objects/ScyllaReferencesStore.cs
+++ b/Engine/Source/Programs/Horde/HordeStorage/Horde.Storage/Implementation/Objects/ScyllaReferencesStore.cs
@@ -50,7 +50,8 @@ namespace Horde.Storage.Implementation
             );"
             ));
 
-            _getObjectsStatement = _session.Prepare("SELECT bucket, name, last_access_time FROM objects WHERE namespace = ? ALLOW FILTERING");
+            // BYPASS CACHE is a scylla specific extension to disable populating the cache, should be disabled by other cassandra dbs
+            _getObjectsStatement = _session.Prepare("SELECT bucket, name, last_access_time FROM objects WHERE namespace = ? ALLOW FILTERING BYPASS CACHE");
             _getNamespacesStatement = _session.Prepare("SELECT DISTINCT namespace FROM buckets");
         }
 
