Horde: Rename log channel names in config file to match namespace renames.

diff --git a/Engine/Source/Programs/Horde/Horde.Build/appsettings.json b/Engine/Source/Programs/Horde/Horde.Build/appsettings.json
--- a/Engine/Source/Programs/Horde/Horde.Build/appsettings.json
+++ b/Engine/Source/Programs/Horde/Horde.Build/appsettings.json
@@ -18,10 +18,10 @@
 				"Microsoft": "Warning",
 				"Microsoft.AspNetCore": "Warning",
 				"Microsoft.Hosting": "Information",
-				"HordeServer.Authentication": "Warning",
+				"Horde.Build.Authentication": "Warning",
 				"System.Net.Http.HttpClient": "Warning",
 				"Grpc": "Warning",
-				"HordeServer.Services.Impl.IssueService": "Debug"
+				"Horde.Build.Services.Impl.IssueService": "Debug"
 			}
 		}
 	},
