Horde Storage - Fixed exception in format resolver when */* was used as Accept type, now using the default format.

diff --git a/Engine/Source/Programs/Horde/HordeStorage/Horde.Storage/Controllers/FormatResolver.cs b/Engine/Source/Programs/Horde/HordeStorage/Horde.Storage/Controllers/FormatResolver.cs
--- a/Engine/Source/Programs/Horde/HordeStorage/Horde.Storage/Controllers/FormatResolver.cs
+++ b/Engine/Source/Programs/Horde/HordeStorage/Horde.Storage/Controllers/FormatResolver.cs
@@ -43,6 +43,12 @@ namespace Horde.Storage.Controllers
                 return defaultContentType;
             }
 
+            // */* means to accept anything, so we use the default content type
+            if (acceptHeader == "*/*")
+            {
+                return defaultContentType;
+            }
+
             foreach (string header in acceptHeader)
             {
                 string s = header.ToLowerInvariant();
