Horde Storage - Kubernetes leader election service no longer inherits IDisposable as it doesnt need it.

diff --git a/Engine/Source/Programs/Horde/HordeStorage/Horde.Storage/Implementation/LeaderElection/KubernetesLeaderElection.cs b/Engine/Source/Programs/Horde/HordeStorage/Horde.Storage/Implementation/LeaderElection/KubernetesLeaderElection.cs
--- a/Engine/Source/Programs/Horde/HordeStorage/Horde.Storage/Implementation/LeaderElection/KubernetesLeaderElection.cs
+++ b/Engine/Source/Programs/Horde/HordeStorage/Horde.Storage/Implementation/LeaderElection/KubernetesLeaderElection.cs
@@ -33,7 +33,7 @@ namespace Horde.Storage.Implementation.LeaderElection
 
     }
 
-    public class KubernetesLeaderElection : PollingService<KubernetesLeaderElectionState>, ILeaderElection, IDisposable
+    public class KubernetesLeaderElection : PollingService<KubernetesLeaderElectionState>, ILeaderElection
     {
         private readonly ILogger _logger = Log.ForContext<KubernetesLeaderElection>();
 
@@ -81,10 +81,12 @@ namespace Horde.Storage.Implementation.LeaderElection
 
         public event EventHandler<ILeaderElection.OnLeaderChangedEventArgs>? OnLeaderChanged;
 
-        public void Dispose()
+        protected override Task OnStopping(KubernetesLeaderElectionState state)
         {
             _leaderElector?.Dispose();
             _client?.Dispose();
+
+            return Task.CompletedTask;
         }
 
         public override async Task<bool> OnPoll(KubernetesLeaderElectionState state, CancellationToken cancellationToken)
