Horde Storage - Attempting to fix issue with replication not running anymore

diff --git a/Engine/Source/Programs/Horde/HordeStorage/Horde.Storage/Implementation/Replication/ReplicationService.cs b/Engine/Source/Programs/Horde/HordeStorage/Horde.Storage/Implementation/Replication/ReplicationService.cs
--- a/Engine/Source/Programs/Horde/HordeStorage/Horde.Storage/Implementation/Replication/ReplicationService.cs
+++ b/Engine/Source/Programs/Horde/HordeStorage/Horde.Storage/Implementation/Replication/ReplicationService.cs
@@ -64,8 +64,10 @@ namespace Horde.Storage.Implementation
                 return;
 
             // if we are no longer the leader cancel any pending replications
-            Task[] stopReplicatingTasks = State.Replicators.Select(replicator => replicator.StopReplicating()).ToArray();
-            Task.WaitAll(stopReplicatingTasks);
+            // TODO: Reset replication token if we are the new leader
+            /*Task[] stopReplicatingTasks = State.Replicators.Select(replicator => replicator.StopReplicating()).ToArray();
+            Task.WaitAll(stopReplicatingTasks);*/
+
         }
 
         private static IReplicator CreateReplicator(ReplicatorSettings replicatorSettings, IServiceProvider provider)
