Horde: Reformat ACL code according to Microsoft style guide.

diff --git a/Engine/Source/Programs/Horde/Horde.Build/Acls/Acl.cs b/Engine/Source/Programs/Horde/Horde.Build/Acls/Acl.cs
--- a/Engine/Source/Programs/Horde/Horde.Build/Acls/Acl.cs
+++ b/Engine/Source/Programs/Horde/Horde.Build/Acls/Acl.cs
@@ -1,8 +1,6 @@
-// Copyright Epic Games, Inc. All Rights Reserved.
+﻿// Copyright Epic Games, Inc. All Rights Reserved.
 
 using Horde.Build.Api;
-using Horde.Build.Services;
-using Horde.Build.Utilities;
 using MongoDB.Bson;
 using MongoDB.Bson.Serialization;
 using MongoDB.Bson.Serialization.Attributes;
@@ -13,7 +11,7 @@ using System.Linq;
 using System.Linq.Expressions;
 using System.Security.Claims;
 
-namespace Horde.Build.Models
+namespace Horde.Build.Acls
 {
 	/// <summary>
 	/// Set of actions that can be performed by a user. NOTE: This enum is sensitive to ordering. Do not change values.
@@ -22,10 +20,12 @@ namespace Horde.Build.Models
 	{
 		//// PROJECTS ////
 
-		/// <summary>
-		/// Allows the creation of new projects
-		/// </summary>
-		CreateProject,
+        #region Projects
+
+        /// <summary>
+        /// Allows the creation of new projects
+        /// </summary>
+        CreateProject,
 
 		/// <summary>
 		/// Allows deletion of projects.
@@ -42,15 +42,16 @@ namespace Horde.Build.Models
 		/// </summary>
 		ViewProject,
 
-
-
+        #endregion
 
 		//// STREAMS ////
 
-		/// <summary>
-		/// Allows the creation of new streams within a project
-		/// </summary>
-		CreateStream,
+        #region Streams
+
+        /// <summary>
+        /// Allows the creation of new streams within a project
+        /// </summary>
+        CreateStream,
 
 		/// <summary>
 		/// Allows updating a stream (agent types, templates, schedules)
@@ -82,21 +83,16 @@ namespace Horde.Build.Models
 		/// </summary>
 		ViewCommitQueues,
 
-
-
-
-
-
-
-
-
+        #endregion
 
 		//// JOBS ////
 
-		/// <summary>
-		/// Ability to start new jobs
-		/// </summary>
-		CreateJob,
+        #region Jobs
+
+        /// <summary>
+        /// Ability to start new jobs
+        /// </summary>
+        CreateJob,
 
 		/// <summary>
 		/// Rename a job, modify its priority, etc...
@@ -123,31 +119,32 @@ namespace Horde.Build.Models
 		/// </summary>
 		ViewJob,
 
+        #endregion
 
+		//// EVENTS ////
 
+        #region Events
 
-
-
-		//// EVENTS ////
-		
-		/// <summary>
-		/// Ability to create events
-		/// </summary>
-		CreateEvent,
+        /// <summary>
+        /// Ability to create events
+        /// </summary>
+        CreateEvent,
 
 		/// <summary>
 		/// Ability to view events
 		/// </summary>
 		ViewEvent,
 
-
+        #endregion
 
 		//// AGENTS ////
 
-		/// <summary>
-		/// Ability to create an agent. This may be done explicitly, or granted to agents to allow them to self-register.
-		/// </summary>
-		CreateAgent,
+        #region Agents
+
+        /// <summary>
+        /// Ability to create an agent. This may be done explicitly, or granted to agents to allow them to self-register.
+        /// </summary>
+        CreateAgent,
 
 		/// <summary>
 		/// Update an agent's name, pools, etc...
@@ -169,15 +166,16 @@ namespace Horde.Build.Models
 		/// </summary>
 		ListAgents,
 
-
-
+        #endregion
 
 		//// POOLS ////
 
-		/// <summary>
-		/// Create a global pool of agents
-		/// </summary>
-		CreatePool,
+        #region Pools
+
+        /// <summary>
+        /// Create a global pool of agents
+        /// </summary>
+        CreatePool,
 
 		/// <summary>
 		/// Modify an agent pool
@@ -199,30 +197,32 @@ namespace Horde.Build.Models
 		/// </summary>
 		ListPools,
 
-
-
+        #endregion
 
 		//// SESSIONS ////
 
-		/// <summary>
-		/// Granted to agents to call CreateSession, which returns a bearer token identifying themselves valid to call UpdateSesssion via gRPC.
-		/// </summary>
-		CreateSession,
+        #region Sessions
+
+        /// <summary>
+        /// Granted to agents to call CreateSession, which returns a bearer token identifying themselves valid to call UpdateSesssion via gRPC.
+        /// </summary>
+        CreateSession,
 
 		/// <summary>
 		/// Allows viewing information about an agent session
 		/// </summary>
 		ViewSession,
 
+        #endregion
 
+        //// CREDENTIALS ////
 
+        #region Credentials
 
-		//// CREDENTIALS ////
-
-		/// <summary>
-		/// Create a new credential
-		/// </summary>
-		CreateCredential,
+        /// <summary>
+        /// Create a new credential
+        /// </summary>
+        CreateCredential,
 
 		/// <summary>
 		/// Delete a credential
@@ -244,29 +244,27 @@ namespace Horde.Build.Models
 		/// </summary>
 		ViewCredential,
 
+        #endregion
 
+        //// LEASES ////
 
+        #region Leases
 
+        /// <summary>
+        /// View all the leases that an agent has worked on
+        /// </summary>
+        ViewLeases,
 
+        #endregion
 
+        //// COMMITS ////
 
-		//// LEASES ////
-
-		/// <summary>
-		/// View all the leases that an agent has worked on
-		/// </summary>
-		ViewLeases,
-
+        #region Commits
 
-
-
-
-		//// COMMITS ////
-
-		/// <summary>
-		/// Add a new commit
-		/// </summary>
-		AddCommit,
+        /// <summary>
+        /// Add a new commit
+        /// </summary>
+        AddCommit,
 
 		/// <summary>
 		/// List all the commits that have been added
@@ -278,38 +276,38 @@ namespace Horde.Build.Models
 		/// </summary>
 		ViewCommits,
 
+        #endregion
 
+        //// ISSUES ////
 
+        #region Issues
 
+        /// <summary>
+        /// View a build health issue
+        /// </summary>
+        ViewIssue,
 
+        #endregion
 
-		/// <summary>
-		/// View a build health issue
-		/// </summary>
-		ViewIssue,
-
-
-
-
-
-		//// TEMPLATES ////
+        //// TEMPLATES ////
 
-		/// <summary>
-		/// View template associated with a stream
-		/// </summary>
-		ViewTemplate,
+        #region Templates
 
+        /// <summary>
+        /// View template associated with a stream
+        /// </summary>
+        ViewTemplate,
 
+        #endregion
 
+        //// LOGS ////
 
+        #region Logs
 
-
-		//// LOGS ////
-
-		/// <summary>
-		/// Ability to create a log. Implicitly granted to agents.
-		/// </summary>
-		CreateLog,
+        /// <summary>
+        /// Ability to create a log. Implicitly granted to agents.
+        /// </summary>
+        CreateLog,
 
 		/// <summary>
 		/// Ability to update log metadata
@@ -326,30 +324,32 @@ namespace Horde.Build.Models
 		/// </summary>
 		WriteLogData,
 
+        #endregion
 
+        //// ARTIFACTS ////
 
+        #region Artifacts
 
-		//// ARTIFACTS ////
-
-		/// <summary>
-		/// Ability to create an artifact. Typically just for debugging; agents have this access for a particular session.
-		/// </summary>
-		UploadArtifact,
+        /// <summary>
+        /// Ability to create an artifact. Typically just for debugging; agents have this access for a particular session.
+        /// </summary>
+        UploadArtifact,
 
 		/// <summary>
 		/// Ability to download an artifact
 		/// </summary>
 		DownloadArtifact,
 
+        #endregion
 
+        //// SOFTWARE ////
 
+        #region Software
 
-		//// SOFTWARE ////
-
-		/// <summary>
-		/// Ability to upload new versions of the agent software
-		/// </summary>
-		UploadSoftware,
+        /// <summary>
+        /// Ability to upload new versions of the agent software
+        /// </summary>
+        UploadSoftware,
 
 		/// <summary>
 		/// Ability to download the agent software
@@ -361,15 +361,16 @@ namespace Horde.Build.Models
 		/// </summary>
 		DeleteSoftware,
 
+        #endregion
 
+        //// ADMIN ////
 
+        #region Admin
 
-		//// ADMIN ////
-
-		/// <summary>
-		/// Ability to read any data from the server. Always inherited.
-		/// </summary>
-		AdminRead,
+        /// <summary>
+        /// Ability to read any data from the server. Always inherited.
+        /// </summary>
+        AdminRead,
 
 		/// <summary>
 		/// Ability to write any data to the server.
@@ -386,15 +387,16 @@ namespace Horde.Build.Models
 		/// </summary>
 		ViewCosts,
 
+        #endregion
 
+        //// PERMISSIONS ////
 
+        #region Permissions
 
-		//// PERMISSIONS ////
-
-		/// <summary>
-		/// Ability to view permissions on an object
-		/// </summary>
-		ViewPermissions,
+        /// <summary>
+        /// Ability to view permissions on an object
+        /// </summary>
+        ViewPermissions,
 
 		/// <summary>
 		/// Ability to change permissions on an object
@@ -406,64 +408,59 @@ namespace Horde.Build.Models
 		/// </summary>
 		IssueBearerToken,
 
+        #endregion
 
+        //// NOTIFICATIONS ////
 
+        #region Notifications
 
+        /// <summary>
+        /// Ability to subscribe to notifications
+        /// </summary>
+        CreateSubscription,
 
-		//// NOTIFICATIONS ////
+        #endregion
 
-		/// <summary>
-		/// Ability to subscribe to notifications
-		/// </summary>
-		CreateSubscription,
+        //// DEVICES ////
 
+        #region Devices
 
-
-
-
-		//// DEVICES ////
-		
-		/// <summary>
-		/// Ability to read devices
-		/// </summary>
-		DeviceRead,
+        /// <summary>
+        /// Ability to read devices
+        /// </summary>
+        DeviceRead,
 
 		/// <summary>
 		/// Ability to write devices
 		/// </summary>
 		DeviceWrite,
 
+        #endregion
 
+        //// COMPUTE ////
 
+        #region Compute
 
-
-
-
-		//// COMPUTE ////
-
-		/// <summary>
-		/// User can add tasks to the compute cluster
-		/// </summary>
-		AddComputeTasks,
+        /// <summary>
+        /// User can add tasks to the compute cluster
+        /// </summary>
+        AddComputeTasks,
 
 		/// <summary>
 		/// User can poll for compute results
 		/// </summary>
 		ViewComputeTasks,
 
+        #endregion
 
+        //// STORAGE ////
 
+        #region Storage
 
-
-
-
-
-		//// STORAGE ////
-
-		/// <summary>
-		/// Ability to read blobs from the storage service
-		/// </summary>
-		ReadBlobs,
+        /// <summary>
+        /// Ability to read blobs from the storage service
+        /// </summary>
+        ReadBlobs,
 
 		/// <summary>
 		/// Ability to write blobs to the storage service
@@ -484,12 +481,14 @@ namespace Horde.Build.Models
 		/// Ability to delete refs
 		/// </summary>
 		DeleteRefs,
-	}
 
-	/// <summary>
-	/// Stores information about a claim
-	/// </summary>
-	public class AclClaim
+        #endregion
+    }
+
+    /// <summary>
+    /// Stores information about a claim
+    /// </summary>
+    public class AclClaim
 	{
 		/// <summary>
 		/// The claim type, typically a URI
@@ -504,29 +503,29 @@ namespace Horde.Build.Models
 		/// <summary>
 		/// Constructor
 		/// </summary>
-		/// <param name="Claim">The claim object</param>
-		public AclClaim(Claim Claim)
-			: this(Claim.Type, Claim.Value)
+		/// <param name="claim">The claim object</param>
+		public AclClaim(Claim claim)
+			: this(claim.Type, claim.Value)
 		{
 		}
 
 		/// <summary>
 		/// Constructor
 		/// </summary>
-		/// <param name="Type">The claim type</param>
-		/// <param name="Value">The claim value</param>
-		public AclClaim(string Type, string Value)
+		/// <param name="type">The claim type</param>
+		/// <param name="value">The claim value</param>
+		public AclClaim(string type, string value)
 		{
-			this.Type = Type;
-			this.Value = Value;
+			Type = type;
+			Value = value;
 		}
 
 		/// <summary>
 		/// Constructs a claim from a request object
 		/// </summary>
-		/// <param name="Request">The request object</param>
-		public AclClaim(CreateAclClaimRequest Request)
-			: this(Request.Type, Request.Value)
+		/// <param name="request">The request object</param>
+		public AclClaim(CreateAclClaimRequest request)
+			: this(request.Type, request.Value)
 		{
 		}
 	}
@@ -560,56 +559,56 @@ namespace Horde.Build.Models
 		/// <summary>
 		/// Constructor
 		/// </summary>
-		/// <param name="Claim">The claim this entry applies to</param>
-		/// <param name="Actions">List of allowed operations</param>
-		internal AclEntry(AclClaim Claim, IEnumerable<AclAction> Actions)
+		/// <param name="claim">The claim this entry applies to</param>
+		/// <param name="actions">List of allowed operations</param>
+		internal AclEntry(AclClaim claim, IEnumerable<AclAction> actions)
 		{
-			this.Claim = Claim;
-			this.Actions = new HashSet<AclAction>(Actions);
+			Claim = claim;
+			Actions = new HashSet<AclAction>(actions);
 		}
 
 		/// <summary>
 		/// Constructs an ACL entry from a request
 		/// </summary>
-		/// <param name="Request">Request instance</param>
+		/// <param name="request">Request instance</param>
 		/// <returns>New ACL entry</returns>
-		public static AclEntry FromRequest(CreateAclEntryRequest Request)
+		public static AclEntry FromRequest(CreateAclEntryRequest request)
 		{
-			return new AclEntry(new AclClaim(Request.Claim.Type, Request.Claim.Value), ParseActionNames(Request.Actions).ToArray());
+			return new AclEntry(new AclClaim(request.Claim.Type, request.Claim.Value), ParseActionNames(request.Actions).ToArray());
 		}
 
 		/// <summary>
 		/// Parses a list of names as an AclAction bitmask
 		/// </summary>
-		/// <param name="ActionNames">Array of names</param>
+		/// <param name="actionNames">Array of names</param>
 		/// <returns>Action bitmask</returns>
-		public static List<AclAction> ParseActionNames(string[]? ActionNames)
+		public static List<AclAction> ParseActionNames(string[]? actionNames)
 		{
-			List<AclAction> Actions = new List<AclAction>();
-			if (ActionNames != null)
+			List<AclAction> actions = new();
+			if (actionNames != null)
 			{
-				foreach (string Name in ActionNames)
+				foreach (string name in actionNames)
 				{
-					Actions.Add(Enum.Parse<AclAction>(Name, true));
+					actions.Add(Enum.Parse<AclAction>(name, true));
 				}
 			}
-			return Actions;
+			return actions;
 		}
 
 		/// <summary>
 		/// Build a list of action names from an array of flags
 		/// </summary>
-		/// <param name="Actions"></param>
+		/// <param name="actions"></param>
 		/// <returns></returns>
-		public static List<string> GetActionNames(IEnumerable<AclAction> Actions)
+		public static List<string> GetActionNames(IEnumerable<AclAction> actions)
 		{
-			List<string> ActionNames = new List<string>();
-			foreach(AclAction Action in Actions)
+            List<string> actionNames = new();
+            foreach(AclAction action in actions)
 			{
-				string Name = Enum.GetName(typeof(AclAction), Action)!;
-				ActionNames.Add(Name);
+				string name = Enum.GetName(typeof(AclAction), action)!;
+				actionNames.Add(name);
 			}
-			return ActionNames;
+			return actionNames;
 		}
 	}
 
@@ -633,7 +632,6 @@ namespace Horde.Build.Models
 		/// </summary>
 		public List<AclAction>? Exceptions { get; set; }
 
-
 		/// <summary>
 		/// Default constructor
 		/// </summary>
@@ -646,35 +644,35 @@ namespace Horde.Build.Models
 		/// <summary>
 		/// Constructor
 		/// </summary>
-		/// <param name="Entries">List of entries for this ACL</param>
-		/// <param name="bInheritPermissions">Whether to inherit permissions from the parent object by default</param>
-		public Acl(List<AclEntry> Entries, bool bInheritPermissions)
+		/// <param name="entries">List of entries for this ACL</param>
+		/// <param name="inherit">Whether to inherit permissions from the parent object by default</param>
+		public Acl(List<AclEntry> entries, bool inherit)
 		{
-			this.Entries = Entries;
-			this.Inherit = bInheritPermissions;
+			Entries = entries;
+			Inherit = inherit;
 		}
 
 		/// <summary>
 		/// Constructor
 		/// </summary>
-		/// <param name="Entries">List of entries for this ACL</param>
-		public Acl(params AclEntry[] Entries)
-			: this(Entries.ToList(), true)
+		/// <param name="entries">List of entries for this ACL</param>
+		public Acl(params AclEntry[] entries)
+			: this(entries.ToList(), true)
 		{
 		}
 
 		/// <summary>
 		/// Tests whether a user is authorized to perform the given actions
 		/// </summary>
-		/// <param name="Action">Action that is being performed. This should be a single flag.</param>
-		/// <param name="User">The principal to authorize</param>
+		/// <param name="action">Action that is being performed. This should be a single flag.</param>
+		/// <param name="user">The principal to authorize</param>
 		/// <returns>True/false if the action is allowed or denied, null if there is no specific setting for this user</returns>
-		public bool? Authorize(AclAction Action, ClaimsPrincipal User)
+		public bool? Authorize(AclAction action, ClaimsPrincipal user)
 		{
 			// Check if there's a specific entry for this action
-			foreach (AclEntry Entry in Entries)
+			foreach (AclEntry entry in Entries)
 			{
-				if(Entry.Actions.Contains(Action) && User.HasClaim(Entry.Claim.Type, Entry.Claim.Value))
+				if(entry.Actions.Contains(action) && user.HasClaim(entry.Claim.Type, entry.Claim.Value))
 				{
 					return true;
 				}
@@ -683,14 +681,14 @@ namespace Horde.Build.Models
 			// Otherwise check if we're prevented from inheriting permissions
 			if(Inherit)
 			{
-				if(Exceptions != null && Exceptions.Contains(Action))
+				if(Exceptions != null && Exceptions.Contains(action))
 				{
 					return false;
 				}
 			}
 			else
 			{
-				if (Exceptions == null || !Exceptions.Contains(Action))
+				if (Exceptions == null || !Exceptions.Contains(action))
 				{
 					return false;
 				}
@@ -703,53 +701,53 @@ namespace Horde.Build.Models
 		/// <summary>
 		/// Merge new settings into an ACL
 		/// </summary>
-		/// <param name="BaseAcl">The current acl</param>
-		/// <param name="Update">The update to apply</param>
+		/// <param name="baseAcl">The current acl</param>
+		/// <param name="update">The update to apply</param>
 		/// <returns>The new ACL value. Null if the ACL has all default settings.</returns>
-		public static Acl? Merge(Acl? BaseAcl, UpdateAclRequest? Update)
+		public static Acl? Merge(Acl? baseAcl, UpdateAclRequest? update)
 		{
-			Acl? NewAcl = null;
-			if(Update != null)
+			Acl? newAcl = null;
+			if(update != null)
 			{
-				NewAcl = new Acl();
+				newAcl = new Acl();
 
-				if(Update.Entries != null)
+				if(update.Entries != null)
 				{
-					NewAcl.Entries = Update.Entries.ConvertAll(x => AclEntry.FromRequest(x));
+					newAcl.Entries = update.Entries.ConvertAll(x => AclEntry.FromRequest(x));
 				}
-				else if(BaseAcl != null)
+				else if(baseAcl != null)
 				{
-					NewAcl.Entries = BaseAcl.Entries;
+					newAcl.Entries = baseAcl.Entries;
 				}
 
-				if(Update.Inherit != null)
+				if(update.Inherit != null)
 				{
-					NewAcl.Inherit = Update.Inherit.Value;
+					newAcl.Inherit = update.Inherit.Value;
 				}
-				else if(BaseAcl != null)
+				else if(baseAcl != null)
 				{
-					NewAcl.Inherit = BaseAcl.Inherit;
+					newAcl.Inherit = baseAcl.Inherit;
 				}
 			}
-			return NewAcl;
+			return newAcl;
 		}
 
 		/// <summary>
 		/// Creates an update definition for the given ACL. Clears the ACL property if it's null.
 		/// </summary>
 		/// <typeparam name="T">Type of document containing the ACL</typeparam>
-		/// <param name="Field">Selector for the ACL property</param>
-		/// <param name="NewAcl">The new ACL value</param>
+		/// <param name="field">Selector for the ACL property</param>
+		/// <param name="newAcl">The new ACL value</param>
 		/// <returns>Update definition for the document</returns>
-		public static UpdateDefinition<T> CreateUpdate<T>(Expression<Func<T, object>> Field, Acl NewAcl)
+		public static UpdateDefinition<T> CreateUpdate<T>(Expression<Func<T, object>> field, Acl newAcl)
 		{
-			if (NewAcl.Entries.Count == 0 && NewAcl.Inherit)
+			if (newAcl.Entries.Count == 0 && newAcl.Inherit)
 			{
-				return Builders<T>.Update.Unset(Field);
+				return Builders<T>.Update.Unset(field);
 			}
 			else
 			{
-				return Builders<T>.Update.Set(Field, NewAcl);
+				return Builders<T>.Update.Set(field, newAcl);
 			}
 		}
 	}
@@ -760,55 +758,52 @@ namespace Horde.Build.Models
 	public sealed class AclActionSetSerializer : IBsonSerializer<HashSet<AclAction>>
 	{
 		/// <inheritdoc/>
-		public Type ValueType
-		{
-			get { return typeof(HashSet<AclAction>); }
-		}
+		public Type ValueType => typeof(HashSet<AclAction>);
 
 		/// <inheritdoc/>
-		void IBsonSerializer.Serialize(BsonSerializationContext Context, BsonSerializationArgs Args, object Value)
+		void IBsonSerializer.Serialize(BsonSerializationContext context, BsonSerializationArgs args, object value)
 		{
-			Serialize(Context, Args, (HashSet<AclAction>)Value);
+			Serialize(context, args, (HashSet<AclAction>)value);
 		}
 
 		/// <inheritdoc/>
-		object IBsonSerializer.Deserialize(BsonDeserializationContext Context, BsonDeserializationArgs Args)
+		object IBsonSerializer.Deserialize(BsonDeserializationContext context, BsonDeserializationArgs args)
 		{
-			return ((IBsonSerializer<HashSet<AclAction>>)this).Deserialize(Context, Args);
+			return ((IBsonSerializer<HashSet<AclAction>>)this).Deserialize(context, args);
 		}
 
 		/// <inheritdoc/>
-		public HashSet<AclAction> Deserialize(BsonDeserializationContext Context, BsonDeserializationArgs Args)
+		public HashSet<AclAction> Deserialize(BsonDeserializationContext context, BsonDeserializationArgs args)
 		{
-			HashSet<AclAction> Values = new HashSet<AclAction>();
+            HashSet<AclAction> values = new();
 
-			Context.Reader.ReadStartArray();
+			context.Reader.ReadStartArray();
 			for(; ;)
 			{
-				BsonType Type = Context.Reader.ReadBsonType();
-				if(Type == BsonType.EndOfDocument)
+				BsonType type = context.Reader.ReadBsonType();
+				if(type == BsonType.EndOfDocument)
 				{
 					break;
 				}
 				else
 				{
-					Values.Add((AclAction)Enum.Parse(typeof(AclAction), Context.Reader.ReadString()));
+					values.Add((AclAction)Enum.Parse(typeof(AclAction), context.Reader.ReadString()));
 				}
 			}
-			Context.Reader.ReadEndArray();
+			context.Reader.ReadEndArray();
 
-			return Values;
+			return values;
 		}
 
 		/// <inheritdoc/>
-		public void Serialize(BsonSerializationContext Context, BsonSerializationArgs Args, HashSet<AclAction> Values)
+		public void Serialize(BsonSerializationContext context, BsonSerializationArgs args, HashSet<AclAction> values)
 		{
-			Context.Writer.WriteStartArray();
-			foreach (AclAction Value in Values)
+			context.Writer.WriteStartArray();
+			foreach (AclAction value in values)
 			{
-				Context.Writer.WriteString(Value.ToString());
+				context.Writer.WriteString(value.ToString());
 			}
-			Context.Writer.WriteEndArray();
+			context.Writer.WriteEndArray();
 		}
 	}
 }
diff --git a/Engine/Source/Programs/Horde/Horde.Build/Acls/AclService.cs b/Engine/Source/Programs/Horde/Horde.Build/Acls/AclService.cs
--- a/Engine/Source/Programs/Horde/Horde.Build/Acls/AclService.cs
+++ b/Engine/Source/Programs/Horde/Horde.Build/Acls/AclService.cs
@@ -1,24 +1,19 @@
-// Copyright Epic Games, Inc. All Rights Reserved.
+﻿// Copyright Epic Games, Inc. All Rights Reserved.
 
-using Horde.Build.Api;
-using Horde.Build.Collections;
 using Horde.Build.Models;
+using Horde.Build.Services;
 using Horde.Build.Utilities;
-using Microsoft.AspNetCore.Authentication.OAuth.Claims;
 using Microsoft.Extensions.Options;
 using Microsoft.IdentityModel.Tokens;
-using MongoDB.Bson;
 using MongoDB.Driver;
 using System;
 using System.Collections.Generic;
 using System.IdentityModel.Tokens.Jwt;
 using System.Linq;
-using System.Linq.Expressions;
 using System.Security.Claims;
-using System.Security.Cryptography;
 using System.Threading.Tasks;
 
-namespace Horde.Build.Services
+namespace Horde.Build.Acls
 {
 	using UserId = ObjectId<IUser>;
 	using SessionId = ObjectId<ISession>;
@@ -69,50 +64,37 @@ namespace Horde.Build.Services
 		/// </summary>
 		public static AclClaim AgentRoleClaim { get; } = new AclClaim(HordeClaimTypes.Role, "agent");
 
-		/// <summary>
-		/// The default permissions
-		/// </summary>
-		Acl DefaultAcl = new Acl();
-
-		/// <summary>
-		/// The database service
-		/// </summary>
-		DatabaseService DatabaseService;
-
-		/// <summary>
-		/// The settings object
-		/// </summary>
-		ServerSettings Settings;
+		private readonly Acl _defaultAcl = new();
+		private readonly DatabaseService _databaseService;
 
 		/// <summary>
 		/// Constructor
 		/// </summary>
-		/// <param name="DatabaseService">The database service instance</param>
-		/// <param name="Settings">The settings object</param>
-		public AclService(DatabaseService DatabaseService, IOptionsMonitor<ServerSettings> Settings)
+		/// <param name="databaseService">The database service instance</param>
+		/// <param name="settings">The settings object</param>
+		public AclService(DatabaseService databaseService, IOptionsMonitor<ServerSettings> settings)
 		{
-			this.DatabaseService = DatabaseService;
-			this.Settings = Settings.CurrentValue;
+			_databaseService = databaseService;
 
-			List<AclAction> AdminActions = new List<AclAction>();
-			foreach (AclAction? Action in Enum.GetValues(typeof(AclAction)))
+            List<AclAction> adminActions = new();
+			foreach (AclAction? action in Enum.GetValues(typeof(AclAction)))
 			{
-				AdminActions.Add(Action!.Value);
+				adminActions.Add(action!.Value);
 			}
 
-			DefaultAcl = new Acl();
-			DefaultAcl.Entries.Add(new AclEntry(new AclClaim(ClaimTypes.Role, "internal:AgentRegistration"), new[] { AclAction.CreateAgent, AclAction.CreateSession }));
-			DefaultAcl.Entries.Add(new AclEntry(AgentRegistrationClaim, new[] { AclAction.CreateAgent, AclAction.CreateSession, AclAction.UpdateAgent, AclAction.DownloadSoftware, AclAction.CreatePool, AclAction.UpdatePool, AclAction.ViewPool, AclAction.DeletePool, AclAction.ListPools, AclAction.ViewStream, AclAction.ViewProject, AclAction.ViewJob, AclAction.ViewCosts }));
-			DefaultAcl.Entries.Add(new AclEntry(AgentRoleClaim, new[] { AclAction.ViewProject, AclAction.ViewStream, AclAction.CreateEvent, AclAction.DownloadSoftware }));
-			DefaultAcl.Entries.Add(new AclEntry(DownloadSoftwareClaim, new[] { AclAction.DownloadSoftware }));
-			DefaultAcl.Entries.Add(new AclEntry(UploadSoftwareClaim, new[] { AclAction.UploadSoftware }));
-			DefaultAcl.Entries.Add(new AclEntry(ConfigureProjectsClaim, new[] { AclAction.CreateProject, AclAction.UpdateProject, AclAction.ViewProject, AclAction.CreateStream, AclAction.UpdateStream, AclAction.ViewStream, AclAction.ChangePermissions }));
-			DefaultAcl.Entries.Add(new AclEntry(StartChainedJobClaim, new[] { AclAction.CreateJob, AclAction.ExecuteJob, AclAction.UpdateJob, AclAction.ViewJob, AclAction.ViewTemplate, AclAction.ViewStream }));
+			_defaultAcl = new();
+			_defaultAcl.Entries.Add(new AclEntry(new AclClaim(ClaimTypes.Role, "internal:AgentRegistration"), new[] { AclAction.CreateAgent, AclAction.CreateSession }));
+			_defaultAcl.Entries.Add(new AclEntry(AgentRegistrationClaim, new[] { AclAction.CreateAgent, AclAction.CreateSession, AclAction.UpdateAgent, AclAction.DownloadSoftware, AclAction.CreatePool, AclAction.UpdatePool, AclAction.ViewPool, AclAction.DeletePool, AclAction.ListPools, AclAction.ViewStream, AclAction.ViewProject, AclAction.ViewJob, AclAction.ViewCosts }));
+			_defaultAcl.Entries.Add(new AclEntry(AgentRoleClaim, new[] { AclAction.ViewProject, AclAction.ViewStream, AclAction.CreateEvent, AclAction.DownloadSoftware }));
+			_defaultAcl.Entries.Add(new AclEntry(DownloadSoftwareClaim, new[] { AclAction.DownloadSoftware }));
+			_defaultAcl.Entries.Add(new AclEntry(UploadSoftwareClaim, new[] { AclAction.UploadSoftware }));
+			_defaultAcl.Entries.Add(new AclEntry(ConfigureProjectsClaim, new[] { AclAction.CreateProject, AclAction.UpdateProject, AclAction.ViewProject, AclAction.CreateStream, AclAction.UpdateStream, AclAction.ViewStream, AclAction.ChangePermissions }));
+			_defaultAcl.Entries.Add(new AclEntry(StartChainedJobClaim, new[] { AclAction.CreateJob, AclAction.ExecuteJob, AclAction.UpdateJob, AclAction.ViewJob, AclAction.ViewTemplate, AclAction.ViewStream }));
 
-			ServerSettings SettingsValue = Settings.CurrentValue;
-			if (SettingsValue.AdminClaimType != null && SettingsValue.AdminClaimValue != null)
+			ServerSettings settingsValue = settings.CurrentValue;
+			if (settingsValue.AdminClaimType != null && settingsValue.AdminClaimValue != null)
 			{
-				DefaultAcl.Entries.Add(new AclEntry(new AclClaim(SettingsValue.AdminClaimType, SettingsValue.AdminClaimValue), AdminActions.ToArray()));
+				_defaultAcl.Entries.Add(new AclEntry(new AclClaim(settingsValue.AdminClaimType, settingsValue.AdminClaimValue), adminActions.ToArray()));
 			}
 		}
 
@@ -122,161 +104,161 @@ namespace Horde.Build.Services
 		/// <returns>Scopes instance</returns>
 		public async Task<Acl> GetRootAcl()
 		{
-			Globals Globals = await DatabaseService.GetGlobalsAsync();
-			return Globals.RootAcl ?? new Acl();
+			Globals globals = await _databaseService.GetGlobalsAsync();
+			return globals.RootAcl ?? new Acl();
 		}
 
 		/// <summary>
 		/// Authorizes a user against a given scope
 		/// </summary>
-		/// <param name="Action">The action being performed</param>
-		/// <param name="User">The principal to validate</param>
-		/// <param name="Cache">The ACL scope cache</param>
+		/// <param name="action">The action being performed</param>
+		/// <param name="user">The principal to validate</param>
+		/// <param name="cache">The ACL scope cache</param>
 		/// <returns>Async task</returns>
-		public async Task<bool> AuthorizeAsync(AclAction Action, ClaimsPrincipal User, GlobalPermissionsCache? Cache = null)
+		public async Task<bool> AuthorizeAsync(AclAction action, ClaimsPrincipal user, GlobalPermissionsCache? cache = null)
 		{
-			Acl? RootAcl;
-			if(Cache == null)
+			Acl? rootAcl;
+			if(cache == null)
 			{
-				RootAcl = await GetRootAcl();
+				rootAcl = await GetRootAcl();
 			}
-			else if(Cache.RootAcl == null)
+			else if(cache.RootAcl == null)
 			{
-				RootAcl = Cache.RootAcl = await GetRootAcl();
+				rootAcl = cache.RootAcl = await GetRootAcl();
 			}
 			else
 			{
-				RootAcl = Cache.RootAcl;
+				rootAcl = cache.RootAcl;
 			}
-			return RootAcl.Authorize(Action, User) ?? DefaultAcl.Authorize(Action, User) ?? false;
+			return rootAcl.Authorize(action, user) ?? _defaultAcl.Authorize(action, user) ?? false;
 		}
 
 		/// <summary>
 		/// Issues a bearer token with the given roles
 		/// </summary>
-		/// <param name="Claims">List of claims to include</param>
-		/// <param name="Expiry">Time that the token expires</param>
+		/// <param name="claims">List of claims to include</param>
+		/// <param name="expiry">Time that the token expires</param>
 		/// <returns>JWT security token with a claim for creating new agents</returns>
-		public string IssueBearerToken(IEnumerable<AclClaim> Claims, TimeSpan? Expiry)
+		public string IssueBearerToken(IEnumerable<AclClaim> claims, TimeSpan? expiry)
 		{
-			return IssueBearerToken(Claims.Select(x => new Claim(x.Type, x.Value)), Expiry);
+			return IssueBearerToken(claims.Select(x => new Claim(x.Type, x.Value)), expiry);
 		}
 
 		/// <summary>
 		/// Issues a bearer token with the given claims
 		/// </summary>
-		/// <param name="Claims">List of claims to include</param>
-		/// <param name="Expiry">Time that the token expires</param>
+		/// <param name="claims">List of claims to include</param>
+		/// <param name="expiry">Time that the token expires</param>
 		/// <returns>JWT security token with a claim for creating new agents</returns>
-		public string IssueBearerToken(IEnumerable<Claim> Claims, TimeSpan? Expiry)
+		public string IssueBearerToken(IEnumerable<Claim> claims, TimeSpan? expiry)
 		{
-			SigningCredentials SigningCredentials = new SigningCredentials(DatabaseService.JwtSigningKey, SecurityAlgorithms.HmacSha256);
+			SigningCredentials signingCredentials = new(_databaseService.JwtSigningKey, SecurityAlgorithms.HmacSha256);
 
-			JwtSecurityToken Token = new JwtSecurityToken(DatabaseService.JwtIssuer, null, Claims, null, DateTime.UtcNow + Expiry, SigningCredentials);
-			return new JwtSecurityTokenHandler().WriteToken(Token);
+			JwtSecurityToken token = new(_databaseService.JwtIssuer, null, claims, null, DateTime.UtcNow + expiry, signingCredentials);
+			return new JwtSecurityTokenHandler().WriteToken(token);
 		}
 
 		/// <summary>
 		/// Get the roles for the given user
 		/// </summary>
-		/// <param name="User">The user to query roles for</param>
+		/// <param name="user">The user to query roles for</param>
 		/// <returns>Collection of roles</returns>
-		public static HashSet<string> GetRoles(ClaimsPrincipal User)
+		public static HashSet<string> GetRoles(ClaimsPrincipal user)
 		{
-			return new HashSet<string>(User.Claims.Where(x => x.Type == ClaimTypes.Role).Select(x => x.Value));
+			return new HashSet<string>(user.Claims.Where(x => x.Type == ClaimTypes.Role).Select(x => x.Value));
 		}
 
 		/// <summary>
 		/// Gets the user name from the given principal
 		/// </summary>
-		/// <param name="User">The principal to check</param>
+		/// <param name="user">The principal to check</param>
 		/// <returns></returns>
-		public static string? GetUserName(ClaimsPrincipal User)
+		public static string? GetUserName(ClaimsPrincipal user)
 		{
-			return (User.Claims.FirstOrDefault(x => x.Type == HordeClaimTypes.User) ?? User.Claims.FirstOrDefault(x => x.Type == ClaimTypes.Name))?.Value ?? "Anonymous";
+			return (user.Claims.FirstOrDefault(x => x.Type == HordeClaimTypes.User) ?? user.Claims.FirstOrDefault(x => x.Type == ClaimTypes.Name))?.Value ?? "Anonymous";
 		}
 
 		/// <summary>
 		/// Gets the agent id associated with a particular user
 		/// </summary>
-		/// <param name="User"></param>
+		/// <param name="user"></param>
 		/// <returns></returns>
-		public static AgentId? GetAgentId(ClaimsPrincipal User)
+		public static AgentId? GetAgentId(ClaimsPrincipal user)
 		{
-			Claim? Claim = User.Claims.FirstOrDefault(x => x.Type == HordeClaimTypes.AgentId);
-			if (Claim == null)
+			Claim? claim = user.Claims.FirstOrDefault(x => x.Type == HordeClaimTypes.AgentId);
+			if (claim == null)
 			{
 				return null;
 			}
 			else
 			{
-				return new AgentId(Claim.Value);
+				return new AgentId(claim.Value);
 			}
 		}
 
 		/// <summary>
 		/// Gets the role for a specific agent
 		/// </summary>
-		/// <param name="AgentId">The session id</param>
+		/// <param name="agentId">The session id</param>
 		/// <returns>New claim instance</returns>
-		public static AclClaim GetAgentClaim(AgentId AgentId)
+		public static AclClaim GetAgentClaim(AgentId agentId)
 		{
-			return new AclClaim(HordeClaimTypes.AgentId, AgentId.ToString());
+			return new AclClaim(HordeClaimTypes.AgentId, agentId.ToString());
 		}
 
 		/// <summary>
 		/// Gets the role for a specific agent session
 		/// </summary>
-		/// <param name="SessionId">The session id</param>
+		/// <param name="sessionId">The session id</param>
 		/// <returns>New claim instance</returns>
-		public static AclClaim GetSessionClaim(SessionId SessionId)
+		public static AclClaim GetSessionClaim(SessionId sessionId)
 		{
-			return new AclClaim(HordeClaimTypes.AgentSessionId, SessionId.ToString());
+			return new AclClaim(HordeClaimTypes.AgentSessionId, sessionId.ToString());
 		}
 
 		/// <summary>
 		/// Determines whether the given user can masquerade as a given user
 		/// </summary>
-		/// <param name="User"></param>
-		/// <param name="UserId"></param>
+		/// <param name="user"></param>
+		/// <param name="userId"></param>
 		/// <returns></returns>
-		public Task<bool> AuthorizeAsUserAsync(ClaimsPrincipal User, UserId UserId)
+		public Task<bool> AuthorizeAsUserAsync(ClaimsPrincipal user, UserId userId)
 		{
-			UserId? CurrentUserId = User.GetUserId();
-			if (CurrentUserId != null && CurrentUserId.Value == UserId)
+			UserId? currentUserId = user.GetUserId();
+			if (currentUserId != null && currentUserId.Value == userId)
 			{
 				return Task.FromResult(true);
 			}
 			else
 			{
-				return AuthorizeAsync(AclAction.Impersonate, User);
+				return AuthorizeAsync(AclAction.Impersonate, user);
 			}
 		}
 	}
 
-	static class ClaimExtensions
+    internal static class ClaimExtensions
 	{
-		public static bool HasSessionClaim(this ClaimsPrincipal User, SessionId SessionId)
+		public static bool HasSessionClaim(this ClaimsPrincipal user, SessionId sessionId)
 		{
-			return User.HasClaim(HordeClaimTypes.AgentSessionId, SessionId.ToString());
+			return user.HasClaim(HordeClaimTypes.AgentSessionId, sessionId.ToString());
 		}
 
-		public static SessionId? GetSessionClaim(this ClaimsPrincipal User)
+		public static SessionId? GetSessionClaim(this ClaimsPrincipal user)
 		{
-			Claim? Claim = User.FindFirst(HordeClaimTypes.AgentSessionId);
-			if (Claim == null || !SessionId.TryParse(Claim.Value, out SessionId SessionIdValue))
+			Claim? claim = user.FindFirst(HordeClaimTypes.AgentSessionId);
+			if (claim == null || !SessionId.TryParse(claim.Value, out SessionId sessionIdValue))
 			{
 				return null;
 			}
 			else
 			{
-				return SessionIdValue;
+				return sessionIdValue;
 			}
 		}
 		
-		public static string GetSessionClaimsAsString(this ClaimsPrincipal User)
+		public static string GetSessionClaimsAsString(this ClaimsPrincipal user)
 		{
-			return String.Join(",", User.Claims
+			return String.Join(",", user.Claims
 				.Where(c => c.Type == HordeClaimTypes.AgentSessionId)
 				.Select(c => c.Value));
 		}
diff --git a/Engine/Source/Programs/Horde/Horde.Build/Agents/IAgent.cs b/Engine/Source/Programs/Horde/Horde.Build/Agents/IAgent.cs
--- a/Engine/Source/Programs/Horde/Horde.Build/Agents/IAgent.cs
+++ b/Engine/Source/Programs/Horde/Horde.Build/Agents/IAgent.cs
@@ -1,7 +1,8 @@
-// Copyright Epic Games, Inc. All Rights Reserved.
+﻿// Copyright Epic Games, Inc. All Rights Reserved.
 
 using Google.Protobuf.WellKnownTypes;
 using EpicGames.Core;
+using Horde.Build.Acls;
 using Horde.Build.Api;
 using HordeCommon;
 using HordeCommon.Rpc.Tasks;
diff --git a/Engine/Source/Programs/Horde/Horde.Build/Agents/IAgentCollection.cs b/Engine/Source/Programs/Horde/Horde.Build/Agents/IAgentCollection.cs
--- a/Engine/Source/Programs/Horde/Horde.Build/Agents/IAgentCollection.cs
+++ b/Engine/Source/Programs/Horde/Horde.Build/Agents/IAgentCollection.cs
@@ -1,6 +1,7 @@
-// Copyright Epic Games, Inc. All Rights Reserved.
+﻿// Copyright Epic Games, Inc. All Rights Reserved.
 
 using Google.Protobuf.WellKnownTypes;
+using Horde.Build.Acls;
 using Horde.Build.Api;
 using HordeCommon;
 using Horde.Build.Models;
diff --git a/Engine/Source/Programs/Horde/Horde.Build/Agents/Impl/AgentCollection.cs b/Engine/Source/Programs/Horde/Horde.Build/Agents/Impl/AgentCollection.cs
--- a/Engine/Source/Programs/Horde/Horde.Build/Agents/Impl/AgentCollection.cs
+++ b/Engine/Source/Programs/Horde/Horde.Build/Agents/Impl/AgentCollection.cs
@@ -1,9 +1,10 @@
-// Copyright Epic Games, Inc. All Rights Reserved.
+﻿// Copyright Epic Games, Inc. All Rights Reserved.
 
 using EpicGames.Redis;
 using Google.Protobuf.WellKnownTypes;
 using HordeCommon;
 using HordeCommon.Rpc.Tasks;
+using Horde.Build.Acls;
 using Horde.Build.Models;
 using Horde.Build.Services;
 using Horde.Build.Utilities;
diff --git a/Engine/Source/Programs/Horde/Horde.Build/Agents/Impl/AgentService.cs b/Engine/Source/Programs/Horde/Horde.Build/Agents/Impl/AgentService.cs
--- a/Engine/Source/Programs/Horde/Horde.Build/Agents/Impl/AgentService.cs
+++ b/Engine/Source/Programs/Horde/Horde.Build/Agents/Impl/AgentService.cs
@@ -1,10 +1,11 @@
-// Copyright Epic Games, Inc. All Rights Reserved.
+﻿// Copyright Epic Games, Inc. All Rights Reserved.
 
 using EpicGames.Core;
 using EpicGames.Redis;
 using EpicGames.Serialization;
 using Google.Protobuf.WellKnownTypes;
 using HordeCommon;
+using Horde.Build.Acls;
 using Horde.Build.Collections;
 using Horde.Build.Models;
 using Horde.Build.Tasks;
diff --git a/Engine/Source/Programs/Horde/Horde.Build/Api/Acls.cs b/Engine/Source/Programs/Horde/Horde.Build/Api/Acls.cs
--- a/Engine/Source/Programs/Horde/Horde.Build/Api/Acls.cs
+++ b/Engine/Source/Programs/Horde/Horde.Build/Api/Acls.cs
@@ -1,5 +1,6 @@
-// Copyright Epic Games, Inc. All Rights Reserved.
+﻿// Copyright Epic Games, Inc. All Rights Reserved.
 
+using Horde.Build.Acls;
 using Horde.Build.Models;
 using System;
 using System.Collections.Generic;
diff --git a/Engine/Source/Programs/Horde/Horde.Build/Compute/Impl/ComputeController.cs b/Engine/Source/Programs/Horde/Horde.Build/Compute/Impl/ComputeController.cs
--- a/Engine/Source/Programs/Horde/Horde.Build/Compute/Impl/ComputeController.cs
+++ b/Engine/Source/Programs/Horde/Horde.Build/Compute/Impl/ComputeController.cs
@@ -1,8 +1,9 @@
-// Copyright Epic Games, Inc. All Rights Reserved.
+﻿// Copyright Epic Games, Inc. All Rights Reserved.
 
 using EpicGames.Core;
 using EpicGames.Horde.Compute;
 using EpicGames.Horde.Compute.Impl;
+using Horde.Build.Acls;
 using Horde.Build.Models;
 using Horde.Build.Services;
 using Horde.Build.Utilities;
diff --git a/Engine/Source/Programs/Horde/Horde.Build/Controllers/AccountController.cs b/Engine/Source/Programs/Horde/Horde.Build/Controllers/AccountController.cs
--- a/Engine/Source/Programs/Horde/Horde.Build/Controllers/AccountController.cs
+++ b/Engine/Source/Programs/Horde/Horde.Build/Controllers/AccountController.cs
@@ -1,5 +1,6 @@
-// Copyright Epic Games, Inc. All Rights Reserved.
+﻿// Copyright Epic Games, Inc. All Rights Reserved.
 
+using Horde.Build.Acls;
 using Horde.Build.Authentication;
 using Horde.Build.Models;
 using Horde.Build.Services;
diff --git a/Engine/Source/Programs/Horde/Horde.Build/Controllers/AdminController.cs b/Engine/Source/Programs/Horde/Horde.Build/Controllers/AdminController.cs
--- a/Engine/Source/Programs/Horde/Horde.Build/Controllers/AdminController.cs
+++ b/Engine/Source/Programs/Horde/Horde.Build/Controllers/AdminController.cs
@@ -1,6 +1,7 @@
-// Copyright Epic Games, Inc. All Rights Reserved.
+﻿// Copyright Epic Games, Inc. All Rights Reserved.
 
 using EpicGames.Core;
+using Horde.Build.Acls;
 using Horde.Build.Api;
 using Horde.Build.Collections;
 using HordeCommon;
diff --git a/Engine/Source/Programs/Horde/Horde.Build/Controllers/AgentSoftwareController.cs b/Engine/Source/Programs/Horde/Horde.Build/Controllers/AgentSoftwareController.cs
--- a/Engine/Source/Programs/Horde/Horde.Build/Controllers/AgentSoftwareController.cs
+++ b/Engine/Source/Programs/Horde/Horde.Build/Controllers/AgentSoftwareController.cs
@@ -1,5 +1,6 @@
-// Copyright Epic Games, Inc. All Rights Reserved.
+﻿// Copyright Epic Games, Inc. All Rights Reserved.
 
+using Horde.Build.Acls;
 using Horde.Build.Api;
 using Horde.Build.Services;
 using Microsoft.AspNetCore.Authorization;
diff --git a/Engine/Source/Programs/Horde/Horde.Build/Controllers/AgentsController.cs b/Engine/Source/Programs/Horde/Horde.Build/Controllers/AgentsController.cs
--- a/Engine/Source/Programs/Horde/Horde.Build/Controllers/AgentsController.cs
+++ b/Engine/Source/Programs/Horde/Horde.Build/Controllers/AgentsController.cs
@@ -1,5 +1,6 @@
-// Copyright Epic Games, Inc. All Rights Reserved.
+﻿// Copyright Epic Games, Inc. All Rights Reserved.
 
+using Horde.Build.Acls;
 using Horde.Build.Api;
 using Horde.Build.Collections;
 using Horde.Build.Models;
diff --git a/Engine/Source/Programs/Horde/Horde.Build/Controllers/ArtifactsController.cs b/Engine/Source/Programs/Horde/Horde.Build/Controllers/ArtifactsController.cs
--- a/Engine/Source/Programs/Horde/Horde.Build/Controllers/ArtifactsController.cs
+++ b/Engine/Source/Programs/Horde/Horde.Build/Controllers/ArtifactsController.cs
@@ -1,6 +1,7 @@
-// Copyright Epic Games, Inc. All Rights Reserved.
+﻿// Copyright Epic Games, Inc. All Rights Reserved.
 
 using Amazon.S3.Transfer;
+using Horde.Build.Acls;
 using Horde.Build.Api;
 using Horde.Build.Jobs;
 using Horde.Build.Models;
diff --git a/Engine/Source/Programs/Horde/Horde.Build/Controllers/ConfigController.cs b/Engine/Source/Programs/Horde/Horde.Build/Controllers/ConfigController.cs
--- a/Engine/Source/Programs/Horde/Horde.Build/Controllers/ConfigController.cs
+++ b/Engine/Source/Programs/Horde/Horde.Build/Controllers/ConfigController.cs
@@ -1,6 +1,7 @@
-// Copyright Epic Games, Inc. All Rights Reserved.
+﻿// Copyright Epic Games, Inc. All Rights Reserved.
 
 using EpicGames.Core;
+using Horde.Build.Acls;
 using Horde.Build.Api;
 using Horde.Build.Models;
 using Horde.Build.Services;
diff --git a/Engine/Source/Programs/Horde/Horde.Build/Controllers/CredentialsController.cs b/Engine/Source/Programs/Horde/Horde.Build/Controllers/CredentialsController.cs
--- a/Engine/Source/Programs/Horde/Horde.Build/Controllers/CredentialsController.cs
+++ b/Engine/Source/Programs/Horde/Horde.Build/Controllers/CredentialsController.cs
@@ -1,5 +1,6 @@
-// Copyright Epic Games, Inc. All Rights Reserved.
+﻿// Copyright Epic Games, Inc. All Rights Reserved.
 
+using Horde.Build.Acls;
 using Horde.Build.Api;
 using Horde.Build.Models;
 using Horde.Build.Services;
diff --git a/Engine/Source/Programs/Horde/Horde.Build/Controllers/DevicesController.cs b/Engine/Source/Programs/Horde/Horde.Build/Controllers/DevicesController.cs
--- a/Engine/Source/Programs/Horde/Horde.Build/Controllers/DevicesController.cs
+++ b/Engine/Source/Programs/Horde/Horde.Build/Controllers/DevicesController.cs
@@ -1,5 +1,6 @@
-// Copyright Epic Games, Inc. All Rights Reserved.
+﻿// Copyright Epic Games, Inc. All Rights Reserved.
 
+using Horde.Build.Acls;
 using Horde.Build.Api;
 using Horde.Build.Collections;
 using Horde.Build.Models;
@@ -1048,4 +1049,4 @@ namespace Horde.Build.Controllers
 			return Ok();
 		}
 	}
-}
\ No newline at end of file
+}
diff --git a/Engine/Source/Programs/Horde/Horde.Build/Controllers/IssuesController.cs b/Engine/Source/Programs/Horde/Horde.Build/Controllers/IssuesController.cs
--- a/Engine/Source/Programs/Horde/Horde.Build/Controllers/IssuesController.cs
+++ b/Engine/Source/Programs/Horde/Horde.Build/Controllers/IssuesController.cs
@@ -1,6 +1,7 @@
-// Copyright Epic Games, Inc. All Rights Reserved.
+﻿// Copyright Epic Games, Inc. All Rights Reserved.
 
 using EpicGames.Core;
+using Horde.Build.Acls;
 using Horde.Build.Api;
 using Horde.Build.Collections;
 using HordeCommon;
diff --git a/Engine/Source/Programs/Horde/Horde.Build/Controllers/JiraController.cs b/Engine/Source/Programs/Horde/Horde.Build/Controllers/JiraController.cs
--- a/Engine/Source/Programs/Horde/Horde.Build/Controllers/JiraController.cs
+++ b/Engine/Source/Programs/Horde/Horde.Build/Controllers/JiraController.cs
@@ -1,8 +1,9 @@
-// Copyright Epic Games, Inc. All Rights Reserved.
+﻿// Copyright Epic Games, Inc. All Rights Reserved.
 
 using Microsoft.AspNetCore.Authorization;
 using Microsoft.AspNetCore.Mvc;
 using System.Threading.Tasks;
+using Horde.Build.Acls;
 using Horde.Build.Services;
 using Microsoft.Extensions.Logging;
 using Horde.Build.Api;
diff --git a/Engine/Source/Programs/Horde/Horde.Build/Controllers/JobsController.cs b/Engine/Source/Programs/Horde/Horde.Build/Controllers/JobsController.cs
--- a/Engine/Source/Programs/Horde/Horde.Build/Controllers/JobsController.cs
+++ b/Engine/Source/Programs/Horde/Horde.Build/Controllers/JobsController.cs
@@ -1,6 +1,7 @@
-// Copyright Epic Games, Inc. All Rights Reserved.
+﻿// Copyright Epic Games, Inc. All Rights Reserved.
 
 using EpicGames.Core;
+using Horde.Build.Acls;
 using Horde.Build.Utilities;
 using Horde.Build.Api;
 using Horde.Build.Collections;
diff --git a/Engine/Source/Programs/Horde/Horde.Build/Controllers/LeasesController.cs b/Engine/Source/Programs/Horde/Horde.Build/Controllers/LeasesController.cs
--- a/Engine/Source/Programs/Horde/Horde.Build/Controllers/LeasesController.cs
+++ b/Engine/Source/Programs/Horde/Horde.Build/Controllers/LeasesController.cs
@@ -1,5 +1,6 @@
-// Copyright Epic Games, Inc. All Rights Reserved.
+﻿// Copyright Epic Games, Inc. All Rights Reserved.
 
+using Horde.Build.Acls;
 using Horde.Build.Api;
 using Horde.Build.Models;
 using Horde.Build.Services;
diff --git a/Engine/Source/Programs/Horde/Horde.Build/Controllers/LogsController.cs b/Engine/Source/Programs/Horde/Horde.Build/Controllers/LogsController.cs
--- a/Engine/Source/Programs/Horde/Horde.Build/Controllers/LogsController.cs
+++ b/Engine/Source/Programs/Horde/Horde.Build/Controllers/LogsController.cs
@@ -1,4 +1,4 @@
-// Copyright Epic Games, Inc. All Rights Reserved.
+﻿// Copyright Epic Games, Inc. All Rights Reserved.
 
 using System;
 using System.Buffers;
@@ -11,6 +11,7 @@ using System.Security.Claims;
 using System.Text;
 using System.Text.Json;
 using System.Threading.Tasks;
+using Horde.Build.Acls;
 using Horde.Build.Api;
 using Horde.Build.Collections;
 using Horde.Build.Logs;
@@ -382,4 +383,4 @@ namespace Horde.Build.Controllers
 			return false;
 		}
 	}
-}
\ No newline at end of file
+}
diff --git a/Engine/Source/Programs/Horde/Horde.Build/Controllers/NoticesController.cs b/Engine/Source/Programs/Horde/Horde.Build/Controllers/NoticesController.cs
--- a/Engine/Source/Programs/Horde/Horde.Build/Controllers/NoticesController.cs
+++ b/Engine/Source/Programs/Horde/Horde.Build/Controllers/NoticesController.cs
@@ -1,5 +1,6 @@
-// Copyright Epic Games, Inc. All Rights Reserved.
+﻿// Copyright Epic Games, Inc. All Rights Reserved.
 
+using Horde.Build.Acls;
 using Horde.Build.Api;
 using Horde.Build.Collections;
 using Horde.Build.Models;
diff --git a/Engine/Source/Programs/Horde/Horde.Build/Controllers/PerforceController.cs b/Engine/Source/Programs/Horde/Horde.Build/Controllers/PerforceController.cs
--- a/Engine/Source/Programs/Horde/Horde.Build/Controllers/PerforceController.cs
+++ b/Engine/Source/Programs/Horde/Horde.Build/Controllers/PerforceController.cs
@@ -1,10 +1,11 @@
-// Copyright Epic Games, Inc. All Rights Reserved.
+﻿// Copyright Epic Games, Inc. All Rights Reserved.
 
 using Microsoft.AspNetCore.Mvc;
 using System.Net;
 using Microsoft.Extensions.Logging;
 using System.Threading.Tasks;
 using System.Collections.Generic;
+using Horde.Build.Acls;
 using Horde.Build.Services;
 using Microsoft.AspNetCore.Authorization;
 using Horde.Build.Models;
diff --git a/Engine/Source/Programs/Horde/Horde.Build/Controllers/PoolsController.cs b/Engine/Source/Programs/Horde/Horde.Build/Controllers/PoolsController.cs
--- a/Engine/Source/Programs/Horde/Horde.Build/Controllers/PoolsController.cs
+++ b/Engine/Source/Programs/Horde/Horde.Build/Controllers/PoolsController.cs
@@ -1,5 +1,6 @@
-// Copyright Epic Games, Inc. All Rights Reserved.
+﻿// Copyright Epic Games, Inc. All Rights Reserved.
 
+using Horde.Build.Acls;
 using Horde.Build.Api;
 using Horde.Build.Models;
 using Horde.Build.Services;
diff --git a/Engine/Source/Programs/Horde/Horde.Build/Controllers/ProjectsController.cs b/Engine/Source/Programs/Horde/Horde.Build/Controllers/ProjectsController.cs
--- a/Engine/Source/Programs/Horde/Horde.Build/Controllers/ProjectsController.cs
+++ b/Engine/Source/Programs/Horde/Horde.Build/Controllers/ProjectsController.cs
@@ -1,5 +1,6 @@
-// Copyright Epic Games, Inc. All Rights Reserved.
+﻿// Copyright Epic Games, Inc. All Rights Reserved.
 
+using Horde.Build.Acls;
 using Horde.Build.Api;
 using Horde.Build.Models;
 using Horde.Build.Services;
diff --git a/Engine/Source/Programs/Horde/Horde.Build/Controllers/StreamsController.cs b/Engine/Source/Programs/Horde/Horde.Build/Controllers/StreamsController.cs
--- a/Engine/Source/Programs/Horde/Horde.Build/Controllers/StreamsController.cs
+++ b/Engine/Source/Programs/Horde/Horde.Build/Controllers/StreamsController.cs
@@ -1,5 +1,6 @@
-// Copyright Epic Games, Inc. All Rights Reserved.
+﻿// Copyright Epic Games, Inc. All Rights Reserved.
 
+using Horde.Build.Acls;
 using Horde.Build.Api;
 using Horde.Build.Collections;
 using Horde.Build.Models;
diff --git a/Engine/Source/Programs/Horde/Horde.Build/Controllers/SubscriptionsController.cs b/Engine/Source/Programs/Horde/Horde.Build/Controllers/SubscriptionsController.cs
--- a/Engine/Source/Programs/Horde/Horde.Build/Controllers/SubscriptionsController.cs
+++ b/Engine/Source/Programs/Horde/Horde.Build/Controllers/SubscriptionsController.cs
@@ -1,5 +1,6 @@
-// Copyright Epic Games, Inc. All Rights Reserved.
+﻿// Copyright Epic Games, Inc. All Rights Reserved.
 
+using Horde.Build.Acls;
 using Horde.Build.Api;
 using Horde.Build.Collections;
 using Horde.Build.Models;
@@ -189,4 +190,4 @@ namespace Horde.Build.Controllers
 			return false;
 		}
 	}
-}
\ No newline at end of file
+}
diff --git a/Engine/Source/Programs/Horde/Horde.Build/Controllers/TestDataController.cs b/Engine/Source/Programs/Horde/Horde.Build/Controllers/TestDataController.cs
--- a/Engine/Source/Programs/Horde/Horde.Build/Controllers/TestDataController.cs
+++ b/Engine/Source/Programs/Horde/Horde.Build/Controllers/TestDataController.cs
@@ -1,5 +1,6 @@
-// Copyright Epic Games, Inc. All Rights Reserved.
+﻿// Copyright Epic Games, Inc. All Rights Reserved.
 
+using Horde.Build.Acls;
 using Horde.Build.Api;
 using Horde.Build.Collections;
 using HordeCommon;
diff --git a/Engine/Source/Programs/Horde/Horde.Build/Devices/IDevice.cs b/Engine/Source/Programs/Horde/Horde.Build/Devices/IDevice.cs
--- a/Engine/Source/Programs/Horde/Horde.Build/Devices/IDevice.cs
+++ b/Engine/Source/Programs/Horde/Horde.Build/Devices/IDevice.cs
@@ -1,10 +1,11 @@
-// Copyright Epic Games, Inc. All Rights Reserved.
+﻿// Copyright Epic Games, Inc. All Rights Reserved.
 
 using System;
 using MongoDB.Bson;
 using MongoDB.Bson.Serialization.Attributes;
 
 using System.Collections.Generic;
+using Horde.Build.Acls;
 using Horde.Build.Utilities;
 
 namespace Horde.Build.Models
diff --git a/Engine/Source/Programs/Horde/Horde.Build/Devices/Impl/DeviceCollection.cs b/Engine/Source/Programs/Horde/Horde.Build/Devices/Impl/DeviceCollection.cs
--- a/Engine/Source/Programs/Horde/Horde.Build/Devices/Impl/DeviceCollection.cs
+++ b/Engine/Source/Programs/Horde/Horde.Build/Devices/Impl/DeviceCollection.cs
@@ -1,9 +1,10 @@
-// Copyright Epic Games, Inc. All Rights Reserved.
+﻿// Copyright Epic Games, Inc. All Rights Reserved.
 
 using Horde.Build.Services;
 using MongoDB.Bson.Serialization.Attributes;
 using System.Threading.Tasks;
 
+using Horde.Build.Acls;
 using Horde.Build.Models;
 using MongoDB.Driver;
 using System;
diff --git a/Engine/Source/Programs/Horde/Horde.Build/Devices/Impl/DeviceService.cs b/Engine/Source/Programs/Horde/Horde.Build/Devices/Impl/DeviceService.cs
--- a/Engine/Source/Programs/Horde/Horde.Build/Devices/Impl/DeviceService.cs
+++ b/Engine/Source/Programs/Horde/Horde.Build/Devices/Impl/DeviceService.cs
@@ -1,7 +1,8 @@
-// Copyright Epic Games, Inc. All Rights Reserved.
+﻿// Copyright Epic Games, Inc. All Rights Reserved.
 
 using System.Threading.Tasks;
 
+using Horde.Build.Acls;
 using Horde.Build.Models;
 using Horde.Build.Collections;
 using System;
@@ -653,4 +654,4 @@ namespace Horde.Build.Services
 		}
 
 	}
-}
\ No newline at end of file
+}
diff --git a/Engine/Source/Programs/Horde/Horde.Build/Jobs/IJob.cs b/Engine/Source/Programs/Horde/Horde.Build/Jobs/IJob.cs
--- a/Engine/Source/Programs/Horde/Horde.Build/Jobs/IJob.cs
+++ b/Engine/Source/Programs/Horde/Horde.Build/Jobs/IJob.cs
@@ -1,6 +1,7 @@
-// Copyright Epic Games, Inc. All Rights Reserved.
+﻿// Copyright Epic Games, Inc. All Rights Reserved.
 
 using EpicGames.Core;
+using Horde.Build.Acls;
 using Horde.Build.Api;
 using HordeCommon;
 using Horde.Build.Utilities;
diff --git a/Engine/Source/Programs/Horde/Horde.Build/Jobs/Impl/JobCollection.cs b/Engine/Source/Programs/Horde/Horde.Build/Jobs/Impl/JobCollection.cs
--- a/Engine/Source/Programs/Horde/Horde.Build/Jobs/Impl/JobCollection.cs
+++ b/Engine/Source/Programs/Horde/Horde.Build/Jobs/Impl/JobCollection.cs
@@ -1,6 +1,7 @@
-// Copyright Epic Games, Inc. All Rights Reserved.
+﻿// Copyright Epic Games, Inc. All Rights Reserved.
 
 using EpicGames.Core;
+using Horde.Build.Acls;
 using Horde.Build.Api;
 using HordeCommon;
 using Horde.Build.Models;
diff --git a/Engine/Source/Programs/Horde/Horde.Build/Jobs/Impl/JobService.cs b/Engine/Source/Programs/Horde/Horde.Build/Jobs/Impl/JobService.cs
--- a/Engine/Source/Programs/Horde/Horde.Build/Jobs/Impl/JobService.cs
+++ b/Engine/Source/Programs/Horde/Horde.Build/Jobs/Impl/JobService.cs
@@ -1,6 +1,7 @@
-// Copyright Epic Games, Inc. All Rights Reserved.
+﻿// Copyright Epic Games, Inc. All Rights Reserved.
 
 using EpicGames.Core;
+using Horde.Build.Acls;
 using Horde.Build.Api;
 using Horde.Build.Collections;
 using HordeCommon;
diff --git a/Engine/Source/Programs/Horde/Horde.Build/Models/Globals.cs b/Engine/Source/Programs/Horde/Horde.Build/Models/Globals.cs
--- a/Engine/Source/Programs/Horde/Horde.Build/Models/Globals.cs
+++ b/Engine/Source/Programs/Horde/Horde.Build/Models/Globals.cs
@@ -1,6 +1,7 @@
-// Copyright Epic Games, Inc. All Rights Reserved.
+﻿// Copyright Epic Games, Inc. All Rights Reserved.
 
 using Google.Protobuf.WellKnownTypes;
+using Horde.Build.Acls;
 using Horde.Build.Services;
 using Horde.Build.Utilities;
 using Microsoft.IdentityModel.Tokens;
diff --git a/Engine/Source/Programs/Horde/Horde.Build/Projects/IProject.cs b/Engine/Source/Programs/Horde/Horde.Build/Projects/IProject.cs
--- a/Engine/Source/Programs/Horde/Horde.Build/Projects/IProject.cs
+++ b/Engine/Source/Programs/Horde/Horde.Build/Projects/IProject.cs
@@ -1,6 +1,7 @@
-// Copyright Epic Games, Inc. All Rights Reserved.
+﻿// Copyright Epic Games, Inc. All Rights Reserved.
 
 using EpicGames.Core;
+using Horde.Build.Acls;
 using Horde.Build.Api;
 using Horde.Build.Utilities;
 using MongoDB.Bson;
diff --git a/Engine/Source/Programs/Horde/Horde.Build/Projects/IStream.cs b/Engine/Source/Programs/Horde/Horde.Build/Projects/IStream.cs
--- a/Engine/Source/Programs/Horde/Horde.Build/Projects/IStream.cs
+++ b/Engine/Source/Programs/Horde/Horde.Build/Projects/IStream.cs
@@ -1,7 +1,8 @@
-// Copyright Epic Games, Inc. All Rights Reserved.
+﻿// Copyright Epic Games, Inc. All Rights Reserved.
 
 using EpicGames.Core;
 using Google.Protobuf.WellKnownTypes;
+using Horde.Build.Acls;
 using Horde.Build.Api;
 using HordeCommon;
 using Horde.Build.Utilities;
diff --git a/Engine/Source/Programs/Horde/Horde.Build/Projects/Impl/ProjectCollection.cs b/Engine/Source/Programs/Horde/Horde.Build/Projects/Impl/ProjectCollection.cs
--- a/Engine/Source/Programs/Horde/Horde.Build/Projects/Impl/ProjectCollection.cs
+++ b/Engine/Source/Programs/Horde/Horde.Build/Projects/Impl/ProjectCollection.cs
@@ -1,5 +1,6 @@
-// Copyright Epic Games, Inc. All Rights Reserved.
+﻿// Copyright Epic Games, Inc. All Rights Reserved.
 
+using Horde.Build.Acls;
 using Horde.Build.Api;
 using Horde.Build.Models;
 using Horde.Build.Services;
diff --git a/Engine/Source/Programs/Horde/Horde.Build/Projects/Impl/ProjectService.cs b/Engine/Source/Programs/Horde/Horde.Build/Projects/Impl/ProjectService.cs
--- a/Engine/Source/Programs/Horde/Horde.Build/Projects/Impl/ProjectService.cs
+++ b/Engine/Source/Programs/Horde/Horde.Build/Projects/Impl/ProjectService.cs
@@ -1,5 +1,6 @@
-// Copyright Epic Games, Inc. All Rights Reserved.
+﻿// Copyright Epic Games, Inc. All Rights Reserved.
 
+using Horde.Build.Acls;
 using Horde.Build.Api;
 using Horde.Build.Collections;
 using Horde.Build.Models;
diff --git a/Engine/Source/Programs/Horde/Horde.Build/Projects/Impl/StreamCollection.cs b/Engine/Source/Programs/Horde/Horde.Build/Projects/Impl/StreamCollection.cs
--- a/Engine/Source/Programs/Horde/Horde.Build/Projects/Impl/StreamCollection.cs
+++ b/Engine/Source/Programs/Horde/Horde.Build/Projects/Impl/StreamCollection.cs
@@ -1,6 +1,7 @@
-// Copyright Epic Games, Inc. All Rights Reserved.
+﻿// Copyright Epic Games, Inc. All Rights Reserved.
 
 using HordeCommon;
+using Horde.Build.Acls;
 using Horde.Build.Api;
 using Horde.Build.Models;
 using Horde.Build.Services;
diff --git a/Engine/Source/Programs/Horde/Horde.Build/Projects/Impl/StreamService.cs b/Engine/Source/Programs/Horde/Horde.Build/Projects/Impl/StreamService.cs
--- a/Engine/Source/Programs/Horde/Horde.Build/Projects/Impl/StreamService.cs
+++ b/Engine/Source/Programs/Horde/Horde.Build/Projects/Impl/StreamService.cs
@@ -1,5 +1,6 @@
-// Copyright Epic Games, Inc. All Rights Reserved.
+﻿// Copyright Epic Games, Inc. All Rights Reserved.
 
+using Horde.Build.Acls;
 using Horde.Build.Api;
 using HordeCommon;
 using Horde.Build.Models;
diff --git a/Engine/Source/Programs/Horde/Horde.Build/Secrets/Credential.cs b/Engine/Source/Programs/Horde/Horde.Build/Secrets/Credential.cs
--- a/Engine/Source/Programs/Horde/Horde.Build/Secrets/Credential.cs
+++ b/Engine/Source/Programs/Horde/Horde.Build/Secrets/Credential.cs
@@ -1,5 +1,6 @@
-// Copyright Epic Games, Inc. All Rights Reserved.
+﻿// Copyright Epic Games, Inc. All Rights Reserved.
 
+using Horde.Build.Acls;
 using Horde.Build.Api;
 using Horde.Build.Utilities;
 using MongoDB.Bson;
diff --git a/Engine/Source/Programs/Horde/Horde.Build/Secrets/CredentialService.cs b/Engine/Source/Programs/Horde/Horde.Build/Secrets/CredentialService.cs
--- a/Engine/Source/Programs/Horde/Horde.Build/Secrets/CredentialService.cs
+++ b/Engine/Source/Programs/Horde/Horde.Build/Secrets/CredentialService.cs
@@ -1,5 +1,6 @@
-// Copyright Epic Games, Inc. All Rights Reserved.
+﻿// Copyright Epic Games, Inc. All Rights Reserved.
 
+using Horde.Build.Acls;
 using Horde.Build.Api;
 using Horde.Build.Models;
 using MongoDB.Bson;
diff --git a/Engine/Source/Programs/Horde/Horde.Build/Server/ConfigService.cs b/Engine/Source/Programs/Horde/Horde.Build/Server/ConfigService.cs
--- a/Engine/Source/Programs/Horde/Horde.Build/Server/ConfigService.cs
+++ b/Engine/Source/Programs/Horde/Horde.Build/Server/ConfigService.cs
@@ -1,6 +1,7 @@
-// Copyright Epic Games, Inc. All Rights Reserved.
+﻿// Copyright Epic Games, Inc. All Rights Reserved.
 
 using EpicGames.Core;
+using Horde.Build.Acls;
 using Horde.Build.Api;
 using Horde.Build.Collections;
 using Horde.Build.Controllers;
diff --git a/Engine/Source/Programs/Horde/Horde.Build/Services/RpcService.cs b/Engine/Source/Programs/Horde/Horde.Build/Services/RpcService.cs
--- a/Engine/Source/Programs/Horde/Horde.Build/Services/RpcService.cs
+++ b/Engine/Source/Programs/Horde/Horde.Build/Services/RpcService.cs
@@ -1,10 +1,11 @@
-// Copyright Epic Games, Inc. All Rights Reserved.
+﻿// Copyright Epic Games, Inc. All Rights Reserved.
 
 using EpicGames.Core;
 using Google.Protobuf.WellKnownTypes;
 using Grpc.Core;
 using HordeCommon;
 using HordeCommon.Rpc;
+using Horde.Build.Acls;
 using Horde.Build.Collections;
 using Horde.Build.Models;
 using Horde.Build.Utilities;
diff --git a/Engine/Source/Programs/Horde/Horde.Build/Startup.cs b/Engine/Source/Programs/Horde/Horde.Build/Startup.cs
--- a/Engine/Source/Programs/Horde/Horde.Build/Startup.cs
+++ b/Engine/Source/Programs/Horde/Horde.Build/Startup.cs
@@ -1,4 +1,4 @@
-// Copyright Epic Games, Inc. All Rights Reserved.
+﻿// Copyright Epic Games, Inc. All Rights Reserved.
 
 using System;
 using System.Collections.Generic;
@@ -35,6 +35,7 @@ using MongoDB.Bson;
 using MongoDB.Bson.Serialization;
 using MongoDB.Bson.Serialization.Conventions;
 using Microsoft.AspNetCore.HttpOverrides;
+using Horde.Build.Acls;
 using Horde.Build.Utilities;
 using Grpc.Core.Interceptors;
 using Grpc.Core;
diff --git a/Engine/Source/Programs/Horde/Horde.Build/Utilities/HordeControllerBase.cs b/Engine/Source/Programs/Horde/Horde.Build/Utilities/HordeControllerBase.cs
--- a/Engine/Source/Programs/Horde/Horde.Build/Utilities/HordeControllerBase.cs
+++ b/Engine/Source/Programs/Horde/Horde.Build/Utilities/HordeControllerBase.cs
@@ -1,6 +1,7 @@
-// Copyright Epic Games, Inc. All Rights Reserved.
+﻿// Copyright Epic Games, Inc. All Rights Reserved.
 
 using EpicGames.Core;
+using Horde.Build.Acls;
 using Horde.Build.Models;
 using Microsoft.AspNetCore.Http;
 using Microsoft.AspNetCore.Mvc;
