Horde: Add a generic base class for versioned collections.

- Versioned collections support individual versioning of each document, through a discriminator that distinguishes between different possible implementations of the document.
- All documents are derived from a common base class which includes an id field, and last update timestamp.
- Atomic operations on documents are supported by querying the last update timestamp before applying updates.
- Documents are cached in Redis whenever read or updated, allowing for performant access without directly accessing the database.

diff --git a/Engine/Source/Programs/Horde/Horde.Build.Tests/TestSetup.cs b/Engine/Source/Programs/Horde/Horde.Build.Tests/TestSetup.cs
--- a/Engine/Source/Programs/Horde/Horde.Build.Tests/TestSetup.cs
+++ b/Engine/Source/Programs/Horde/Horde.Build.Tests/TestSetup.cs
@@ -9,6 +9,7 @@ using Datadog.Trace;
 using EpicGames.Core;
 using EpicGames.Horde.Storage;
 using HordeCommon;
+using Horde.Build.Acls;
 using Horde.Build.Storage.Services;
 using Horde.Build;
 using Horde.Build.Collections;
diff --git a/Engine/Source/Programs/Horde/Horde.Build.Tests/VersionedCollectionTests.cs b/Engine/Source/Programs/Horde/Horde.Build.Tests/VersionedCollectionTests.cs
new file mode 100644
--- /dev/null
+++ b/Engine/Source/Programs/Horde/Horde.Build.Tests/VersionedCollectionTests.cs
@@ -0,0 +1,179 @@
+// Copyright Epic Games, Inc. All Rights Reserved.
+
+using Horde.Build.Services;
+using Horde.Build.Tools.Impl;
+using Microsoft.VisualStudio.TestTools.UnitTesting;
+using MongoDB.Driver;
+using StackExchange.Redis;
+using System.Threading.Tasks;
+
+namespace Horde.Build.Tests
+{
+	[TestClass]
+	public class VersionedCollectionTests : DatabaseIntegrationTest
+	{
+		abstract class DocumentBase : VersionedDocument<string, DocumentV2>
+		{
+			protected DocumentBase(string id)
+				: base(id)
+			{
+			}
+		}
+
+		class DocumentV1 : DocumentBase
+		{
+			public int ValueV1 { get; set; }
+
+			public DocumentV1(string id, int valueV1)
+				: base(id)
+			{
+				ValueV1 = valueV1;
+			}
+
+			public override DocumentV2 UpgradeToLatest() => new DocumentV2(Id, ValueV1.ToString());
+		}
+
+		class DocumentV2 : DocumentBase
+		{
+			public string ValueV2 { get; set; }
+
+			public DocumentV2(string id, string valueV2)
+				: base(id)
+			{
+				ValueV2 = valueV2;
+			}
+
+			public override DocumentV2 UpgradeToLatest() => this;
+		}
+
+
+		DatabaseService _databaseService;
+		IDatabase _redis;
+		RedisKey baseKey = new RedisKey("versioned/");
+		IMongoCollection<VersionedDocument<string, DocumentV2>> _baseCollection;
+		VersionedCollection<string, DocumentV2> _collection;
+
+		public VersionedCollectionTests()
+		{
+			_databaseService = GetDatabaseService();
+			_redis = GetRedisDatabase();
+			_baseCollection = _databaseService.GetCollection<VersionedDocument<string, DocumentV2>>("versioned");
+			_collection = new VersionedCollection<string, DocumentV2>(_baseCollection, _redis, baseKey);
+		}
+
+		[TestMethod]
+		public async Task Add()
+		{
+			DocumentV2? doc = await _collection.GetAsync("hello");
+			Assert.IsNull(doc);
+
+			Assert.IsTrue(await _collection.AddAsync(new DocumentV2("hello", "1234")));
+
+			doc = await _collection.GetAsync("hello");
+			Assert.IsNotNull(doc);
+			Assert.AreEqual(doc!.ValueV2, "1234");
+		}
+
+		[TestMethod]
+		public async Task AutoUpgrade()
+		{
+			DocumentV2? doc = await _collection.GetAsync("hello");
+			Assert.IsNull(doc);
+
+			await _baseCollection.InsertOneAsync(new DocumentV1("hello", 987));
+
+			doc = await _collection.GetAsync("hello");
+			Assert.IsNotNull(doc);
+			Assert.AreEqual(doc!.ValueV2, "987");
+		}
+
+		[TestMethod]
+		public async Task Delete()
+		{
+			DocumentV2? doc = await _collection.GetAsync("hello");
+			Assert.IsNull(doc);
+
+			Assert.IsTrue(await _collection.AddAsync(new DocumentV2("hello", "1234")));
+
+			doc = await _collection.GetAsync("hello");
+			Assert.IsNotNull(doc);
+			Assert.AreEqual(doc!.ValueV2, "1234");
+
+			Assert.IsTrue(await _collection.DeleteAsync("hello"));
+
+			doc = await _collection.GetAsync("hello");
+			Assert.IsNull(doc);
+		}
+
+		[TestMethod]
+		public async Task FindOrAdd()
+		{
+			await _baseCollection.InsertOneAsync(new DocumentV1("hello", 555));
+
+			DocumentV2? doc = await _collection.FindOrAddAsync("hello", () => new DocumentV2("hello", "1234"));
+			Assert.IsNotNull(doc);
+			Assert.AreEqual(doc.ValueV2, "555");
+
+			doc = await _collection.FindOrAddAsync("world", () => new DocumentV2("world", "1234"));
+			Assert.IsNotNull(doc);
+			Assert.AreEqual(doc.ValueV2, "1234");
+
+			doc = await _collection.FindOrAddAsync("world", () => new DocumentV2("world", "5678"));
+			Assert.IsNotNull(doc);
+			Assert.AreEqual(doc.ValueV2, "1234");
+		}
+
+		[TestMethod]
+		public async Task Update()
+		{
+			await _baseCollection.InsertOneAsync(new DocumentV1("hello", 555));
+
+			DocumentV2? doc = await _collection.GetAsync("hello");
+			Assert.IsNotNull(doc);
+			Assert.AreEqual(doc!.ValueV2, "555");
+
+			DocumentV2? updatedDoc = await _collection.UpdateAsync(doc!, Builders<DocumentV2>.Update.Set(x => x.ValueV2, "999"));
+			Assert.IsNotNull(updatedDoc);
+			Assert.AreEqual(updatedDoc!.ValueV2, "999");
+
+			Assert.AreEqual(doc!.ValueV2, "555");
+
+			DocumentV2? doc2 = await _collection.GetAsync("hello");
+			Assert.IsNotNull(doc2);
+			Assert.AreEqual(doc2!.ValueV2, "999");
+
+			DocumentV2? updatedDoc2 = await _collection.UpdateAsync(doc!, Builders<DocumentV2>.Update.Set(x => x.ValueV2, "1000"));
+			Assert.IsNull(updatedDoc2);
+
+			doc2 = await _collection.GetAsync("hello");
+			Assert.IsNotNull(doc2);
+			Assert.AreEqual(doc2!.ValueV2, "999");
+
+			updatedDoc2 = await _collection.UpdateAsync(updatedDoc, Builders<DocumentV2>.Update.Set(x => x.ValueV2, "1000"));
+			Assert.IsNotNull(updatedDoc2);
+			Assert.AreEqual(updatedDoc2!.ValueV2, "1000");
+		}
+
+		[TestMethod]
+		public async Task Replace()
+		{
+			await _baseCollection.InsertOneAsync(new DocumentV1("hello", 555));
+
+			DocumentV2? doc = await _collection.GetAsync("hello");
+			Assert.IsNotNull(doc);
+			Assert.AreEqual(doc!.ValueV2, "555");
+
+			Assert.IsTrue(await _collection.ReplaceAsync(doc!, new DocumentV2("hello", "world")));
+
+			DocumentV2? newDoc = await _collection.GetAsync("hello");
+			Assert.IsNotNull(newDoc);
+			Assert.AreEqual(newDoc!.ValueV2, "world");
+
+			Assert.IsFalse(await _collection.ReplaceAsync(doc!, new DocumentV2("hello", "world2")));
+
+			newDoc = await _collection.GetAsync("hello");
+			Assert.IsNotNull(newDoc);
+			Assert.AreEqual(newDoc!.ValueV2, "world");
+		}
+	}
+}
diff --git a/Engine/Source/Programs/Horde/Horde.Build/Utilities/VersionedCollection.cs b/Engine/Source/Programs/Horde/Horde.Build/Utilities/VersionedCollection.cs
new file mode 100644
--- /dev/null
+++ b/Engine/Source/Programs/Horde/Horde.Build/Utilities/VersionedCollection.cs
@@ -0,0 +1,338 @@
+﻿// Copyright Epic Games, Inc. All Rights Reserved.
+
+using Horde.Build.Services;
+using Horde.Build.Utilities;
+using MongoDB.Bson;
+using MongoDB.Bson.Serialization;
+using MongoDB.Bson.Serialization.Attributes;
+using MongoDB.Driver;
+using StackExchange.Redis;
+using System;
+using System.Collections.Generic;
+using System.Threading.Tasks;
+
+namespace Horde.Build.Tools.Impl
+{
+	/// <summary>
+	/// Base class for a versioned MongoDB document
+	/// </summary>
+	/// <typeparam name="TId">Type of the unique identifier for each document</typeparam>
+	/// <typeparam name="TLatest">Type of the latest document</typeparam>
+	public abstract class VersionedDocument<TId, TLatest>
+		where TId : notnull
+		where TLatest : VersionedDocument<TId, TLatest>
+	{
+		/// <summary>
+		/// Unique id for this document
+		/// </summary>
+		public TId Id { get; set; }
+
+		/// <summary>
+		/// Last time that the document was updated. This field is checked and updated as part of updates to ensure atomicity.
+		/// </summary>
+		[BsonElement("_u")]
+		public DateTime LastUpdateTime { get; set; }
+
+		/// <summary>
+		/// Constructor
+		/// </summary>
+		/// <param name="id">Unique id for the document</param>
+		protected VersionedDocument(TId id)
+		{
+			Id = id;
+			LastUpdateTime = DateTime.MinValue;
+		}
+
+		/// <summary>
+		/// Perform any transformations necessary to upgrade this document to the latest version
+		/// </summary>
+		/// <returns>Upgraded copy of the document</returns>
+		public abstract TLatest UpgradeToLatest();
+	}
+
+	/// <summary>
+	/// Collection of types derived from <see cref="VersionedDocument{TId, TLatest}"/>
+	/// </summary>
+	public sealed class VersionedCollection<TId, TLatest>
+		where TId : notnull
+		where TLatest : VersionedDocument<TId, TLatest>
+	{
+		private static readonly RedisKey s_timeSuffix = new RedisKey("/ts");
+
+		private static FilterDefinitionBuilder<VersionedDocument<TId, TLatest>> FilterBuilder { get; } = Builders<VersionedDocument<TId, TLatest>>.Filter;
+		private static UpdateDefinitionBuilder<VersionedDocument<TId, TLatest>> UpdateBuilder { get; } = Builders<VersionedDocument<TId, TLatest>>.Update;
+
+		private readonly IMongoCollection<VersionedDocument<TId, TLatest>> _baseCollection;
+		private readonly IMongoCollection<TLatest> _latestCollection;
+		private readonly IDatabase _redis;
+		private readonly RedisKey _baseKey;
+
+		/// <summary>
+		/// Constructor
+		/// </summary>
+		/// <param name="baseCollection">The collection of documents to manage</param>
+		/// <param name="redis">Instance of the redis service</param>
+		/// <param name="baseKey">Prefix for key types</param>
+		public VersionedCollection(IMongoCollection<VersionedDocument<TId, TLatest>> baseCollection, IDatabase redis, RedisKey baseKey)
+		{
+			_baseCollection = baseCollection;
+			_latestCollection = baseCollection.OfType<TLatest>();
+			_redis = redis;
+			_baseKey = baseKey;
+		}
+
+		/// <summary>
+		/// Constructor
+		/// </summary>
+		/// <param name="databaseService">The database service</param>
+		/// <param name="collectionName">Name of the collection of documents to manage</param>
+		/// <param name="redis">Instance of the redis service</param>
+		/// <param name="baseKey">Prefix for key types</param>
+		public VersionedCollection(DatabaseService databaseService, string collectionName, IDatabase redis, RedisKey baseKey)
+			: this(databaseService.Database.GetCollection<VersionedDocument<TId, TLatest>>(collectionName), redis, baseKey)
+		{
+		}
+
+		private RedisKey GetDocKey(TId id) => _baseKey.Append(id.ToString());
+
+		private static FilterDefinition<VersionedDocument<TId, TLatest>> GetFilter(TId id)
+		{
+			return Builders<VersionedDocument<TId, TLatest>>.Filter.Eq(x => x.Id, id);
+		}
+
+		private static FilterDefinition<VersionedDocument<TId, TLatest>> GetFilter(VersionedDocument<TId, TLatest> doc)
+		{
+			FilterDefinitionBuilder<VersionedDocument<TId, TLatest>> Builder = Builders<VersionedDocument<TId, TLatest>>.Filter;
+			return Builder.Eq(x => x.Id, doc.Id) & Builder.Eq(x => x.LastUpdateTime, doc.LastUpdateTime);
+		}
+
+		private static FilterDefinition<TLatest> GetFilter(TLatest doc)
+		{
+			FilterDefinitionBuilder<TLatest> Builder = Builders<TLatest>.Filter;
+			return Builder.Eq(x => x.Id, doc.Id) & Builder.Eq(x => x.LastUpdateTime, doc.LastUpdateTime);
+		}
+
+		/// <summary>
+		/// Adds a new document to the collection
+		/// </summary>
+		/// <param name="doc">The document to add</param>
+		/// <returns>True if the document was added, false if it already exists</returns>
+		public async Task<bool> AddAsync(TLatest doc)
+		{
+			doc.LastUpdateTime = DateTime.UtcNow;
+
+			if (!await _latestCollection.InsertOneIgnoreDuplicatesAsync(doc))
+			{
+				return false;
+			}
+
+			AddCachedValue(GetDocKey(doc.Id), doc);
+			return true;
+		}
+
+		private void AddCachedValue(RedisKey docKey, TLatest doc)
+		{
+			KeyValuePair<RedisKey, RedisValue>[] Pairs = new KeyValuePair<RedisKey, RedisValue>[2];
+			Pairs[0] = new KeyValuePair<RedisKey, RedisValue>(docKey, doc.ToBson());
+			Pairs[1] = new KeyValuePair<RedisKey, RedisValue>(docKey.Append(s_timeSuffix), doc.LastUpdateTime.Ticks);
+
+			_ = _redis.StringSetAsync(Pairs, When.NotExists, CommandFlags.FireAndForget);
+		}
+
+		private async ValueTask<VersionedDocument<TId, TLatest>?> GetCachedValueAsync(RedisKey docKey)
+		{
+			RedisValue cacheValue = await _redis.StringGetAsync(docKey);
+			if (!cacheValue.IsNull)
+			{
+				try
+				{
+					return BsonSerializer.Deserialize<VersionedDocument<TId, TLatest>>((byte[])cacheValue);
+				}
+				catch
+				{
+					await DeleteCachedValueAsync(docKey);
+				}
+			}
+			return null;
+		}
+
+		private async ValueTask<bool> UpdateCachedValueAsync(RedisKey docKey, VersionedDocument<TId, TLatest> prevDoc, TLatest doc)
+		{
+			ITransaction transaction = _redis.CreateTransaction();
+			transaction.AddCondition(Condition.StringEqual(docKey.Append(s_timeSuffix), prevDoc.LastUpdateTime.Ticks));
+
+			KeyValuePair<RedisKey, RedisValue>[] Pairs = new KeyValuePair<RedisKey, RedisValue>[2];
+			Pairs[0] = new KeyValuePair<RedisKey, RedisValue>(docKey, doc.ToBson());
+			Pairs[1] = new KeyValuePair<RedisKey, RedisValue>(docKey.Append(s_timeSuffix), doc.LastUpdateTime.Ticks);
+
+			_ = transaction.StringSetAsync(Pairs);
+
+			if (!await transaction.ExecuteAsync())
+			{
+				await DeleteCachedValueAsync(docKey);
+				return false;
+			}
+
+			return true;
+		}
+
+		private async ValueTask DeleteCachedValueAsync(RedisKey docKey)
+		{
+			await _redis.KeyDeleteAsync(new[] { docKey, docKey.Append(s_timeSuffix) });
+		}
+
+		/// <summary>
+		/// Gets a document with the given id
+		/// </summary>
+		/// <param name="id">The document id to look for</param>
+		/// <returns>The matching document, or null if it does not exist</returns>
+		public async Task<TLatest?> GetAsync(TId id)
+		{
+			RedisKey docKey = GetDocKey(id);
+			for (; ; )
+			{
+				// Attempt to get the cached value for this key
+				VersionedDocument<TId, TLatest>? cachedDoc = await GetCachedValueAsync(docKey);
+				if (cachedDoc == null)
+				{
+					// Read the value from the database
+					VersionedDocument<TId, TLatest>? doc = await _baseCollection.Find(FilterBuilder.Eq(x => x.Id, id)).FirstOrDefaultAsync();
+					if (doc == null)
+					{
+						return null;
+					}
+
+					if (doc is TLatest latestDoc)
+					{
+						AddCachedValue(docKey, latestDoc);
+						return latestDoc;
+					}
+
+					TLatest upgradedDoc = doc.UpgradeToLatest();
+					if (await ReplaceAsync(doc, upgradedDoc))
+					{
+						return upgradedDoc;
+					}
+				}
+				else
+				{
+					// Parse the cached value, and make sure it's the latest version
+					if (cachedDoc is TLatest latestCachedDoc)
+					{
+						return latestCachedDoc;
+					}
+
+					TLatest upgradedDoc = cachedDoc.UpgradeToLatest();
+					if (await ReplaceAsync(cachedDoc, upgradedDoc))
+					{
+						return upgradedDoc;
+					}
+				}
+			}
+		}
+
+		/// <summary>
+		/// Gets an existing document or creates one with the given callback
+		/// </summary>
+		/// <param name="id">The document id to look for</param>
+		/// <param name="factory">Factory method used to create a new document if need be</param>
+		/// <returns>The existing document, or the document that was inserted</returns>
+		public async Task<TLatest> FindOrAddAsync(TId id, Func<TLatest> factory)
+		{
+			TLatest? newDoc = null;
+			for (; ; )
+			{
+				// Try to get an existing document
+				TLatest? latest = await GetAsync(id);
+				if (latest != null)
+				{
+					return latest;
+				}
+
+				// Create a new document and try to add it
+				newDoc ??= factory();
+				if (await AddAsync(newDoc))
+				{
+					return newDoc;
+				}
+			}
+		}
+
+		/// <summary>
+		/// Attempt to update the given document. Fails if the document has been modified from the version presented.
+		/// </summary>
+		/// <param name="doc">The current document version</param>
+		/// <param name="update">Update to be applied</param>
+		/// <returns>True if the document was updated</returns>
+		public async Task<TLatest?> UpdateAsync(TLatest doc, UpdateDefinition<TLatest> update)
+		{
+			update = update.Set(x => x.LastUpdateTime, new DateTime(Math.Max(doc.LastUpdateTime.Ticks + 1, DateTime.UtcNow.Ticks)));
+
+			TLatest? newDoc = await _latestCollection.FindOneAndUpdateAsync(GetFilter(doc), update, new FindOneAndUpdateOptions<TLatest> { ReturnDocument = ReturnDocument.After });
+			if (newDoc != null)
+			{
+				await UpdateCachedValueAsync(GetDocKey(doc.Id), doc, newDoc);
+			}
+			return newDoc;
+		}
+
+		/// <summary>
+		/// Replaces an existing document with a new one
+		/// </summary>
+		/// <param name="oldDoc">The old document</param>
+		/// <param name="newDoc">The new document</param>
+		/// <returns>True if the document was replaced</returns>
+		public async Task<bool> ReplaceAsync(VersionedDocument<TId, TLatest> oldDoc, TLatest newDoc)
+		{
+			if (!oldDoc.Id.Equals(newDoc.Id))
+			{
+				throw new InvalidOperationException("Id for new document must match old document");
+			}
+
+			newDoc.LastUpdateTime = new DateTime(Math.Max(oldDoc.LastUpdateTime.Ticks + 1, DateTime.UtcNow.Ticks));
+
+			ReplaceOneResult result = await _baseCollection.ReplaceOneAsync(GetFilter(oldDoc), newDoc);
+			if (result.ModifiedCount == 0)
+			{
+				return false;
+			}
+
+			await UpdateCachedValueAsync(GetDocKey(oldDoc.Id), oldDoc, newDoc);
+			return true;
+		}
+
+		/// <summary>
+		/// Delete a document with the given identifier
+		/// </summary>
+		/// <param name="doc">Id of the document to delete</param>
+		/// <returns>True if the document was deleted, or false if it could not be found</returns>
+		public async Task<bool> DeleteAsync(TLatest doc)
+		{
+			RedisKey docKey = GetDocKey(doc.Id);
+
+			DeleteResult result = await _baseCollection.DeleteOneAsync(GetFilter((VersionedDocument<TId, TLatest>)doc));
+			if (result.DeletedCount > 0)
+			{
+				await DeleteCachedValueAsync(docKey);
+				return true;
+			}
+
+			return false;
+		}
+
+		/// <summary>
+		/// Delete a document with the given identifier
+		/// </summary>
+		/// <param name="id">Id of the document to delete</param>
+		/// <returns>True if the document was deleted, or false if it could not be found</returns>
+		public async Task<bool> DeleteAsync(TId id)
+		{
+			RedisKey docKey = GetDocKey(id);
+
+			DeleteResult result = await _baseCollection.DeleteOneAsync(GetFilter(id));
+			await DeleteCachedValueAsync(docKey);
+
+			return result.DeletedCount > 0;
+		}
+	}
+}
