Horde Storage - More detailed logging when running the tests to better see what is being run

diff --git a/Engine/Source/Programs/Horde/BuildHorde.xml b/Engine/Source/Programs/Horde/BuildHorde.xml
--- a/Engine/Source/Programs/Horde/BuildHorde.xml
+++ b/Engine/Source/Programs/Horde/BuildHorde.xml
@@ -42,7 +42,7 @@
 			<Spawn Exe="/bin/bash" Arguments="-c 'sleep 60'" />
 
 			<!-- Run the unit tests, and override the dotnet sdk used to the global one for the machine rather then what is in the sync, to allow us to use newer versions then the rest of the engine has migrated to -->
-			<DotNet Arguments="test Engine/Source/Programs/Horde/HordeStorage/Horde.Storage.sln" DotNetPath="/usr/share/dotnet/dotnet"/>
+			<DotNet Arguments="test Engine/Source/Programs/Horde/HordeStorage/Horde.Storage.sln -l console;verbosity=normal" DotNetPath="/usr/share/dotnet/dotnet"/>
 		</Node>
 
 		<Node Name="Build HordeStorage" Requires="Test HordeStorage">
