Horde Storage - Fixed race where we could end up adding refs to the blob index for blobs we had yet to commit.

diff --git a/Engine/Source/Programs/Horde/HordeStorage/Horde.Storage/Implementation/ObjectService.cs b/Engine/Source/Programs/Horde/HordeStorage/Horde.Storage/Implementation/ObjectService.cs
--- a/Engine/Source/Programs/Horde/HordeStorage/Horde.Storage/Implementation/ObjectService.cs
+++ b/Engine/Source/Programs/Horde/HordeStorage/Horde.Storage/Implementation/ObjectService.cs
@@ -102,7 +102,6 @@ namespace Horde.Storage.Implementation
             Task objectStorePut = _referencesStore.Put(ns, bucket, key, blobHash, payload.GetView().ToArray(), isFinalized);
 
             Task<BlobIdentifier> blobStorePut = _blobService.PutObject(ns, payload.GetView().ToArray(), blobHash);
-            Task addRefToBodyTask = _blobIndex.AddRefToBlobs(ns, bucket, key, new [] {blobHash});
             ContentId[] missingReferences = Array.Empty<ContentId>();
             BlobIdentifier[] missingBlobs = Array.Empty<BlobIdentifier>();
 
@@ -129,8 +128,10 @@ namespace Horde.Storage.Implementation
                 }
             }
 
-            await Task.WhenAll(objectStorePut, blobStorePut, addRefToBodyTask);
-            
+            await Task.WhenAll(objectStorePut, blobStorePut);
+
+            Task addRefToBlobsTask = _blobIndex.AddRefToBlobs(ns, bucket, key, new [] {blobHash});
+
             if (missingReferences.Length == 0 && missingBlobs.Length == 0)
             {
                 try
@@ -149,6 +150,8 @@ namespace Horde.Storage.Implementation
                 await _replicationLog.InsertAddEvent(ns, bucket, key, blobHash);
             }
 
+            await addRefToBlobsTask;
+
             return (missingReferences, missingBlobs);
         }
 
