Horde Storage - Added blob index consistency check service, this verifies that the blob index is consistent with the blob store and attempts to repair it if possible. It will also cleanup old blob index records that are no longer valid.
Renamed ConsistencyCheckService to BlobStoreConsistencyCheckService to make it clear that is what it deals with.

diff --git a/Engine/Source/Programs/Horde/HordeStorage/Horde.Storage/Horde.Storage.Startup.cs b/Engine/Source/Programs/Horde/HordeStorage/Horde.Storage/Horde.Storage.Startup.cs
--- a/Engine/Source/Programs/Horde/HordeStorage/Horde.Storage/Horde.Storage.Startup.cs
+++ b/Engine/Source/Programs/Horde/HordeStorage/Horde.Storage/Horde.Storage.Startup.cs
@@ -288,8 +288,11 @@ namespace Horde.Storage
             services.AddSingleton<ReplicationSnapshotService>();
             services.AddHostedService<ReplicationSnapshotService>(p => p.GetService<ReplicationSnapshotService>()!);
 
-            services.AddSingleton<ConsistencyCheckService>();
-            services.AddHostedService<ConsistencyCheckService>(p => p.GetService<ConsistencyCheckService>()!);
+            services.AddSingleton<BlobStoreConsistencyCheckService>();
+            services.AddHostedService<BlobStoreConsistencyCheckService>(p => p.GetService<BlobStoreConsistencyCheckService>()!);
+
+            services.AddSingleton<BlobIndexConsistencyCheckService>();
+            services.AddHostedService<BlobIndexConsistencyCheckService>(p => p.GetService<BlobIndexConsistencyCheckService>()!);
 
             services.AddSingleton(typeof(IPeerStatusService), typeof(PeerStatusService));
             services.AddHostedService<PeerStatusService>(p => (PeerStatusService)p.GetService<IPeerStatusService>()!);
diff --git a/Engine/Source/Programs/Horde/HordeStorage/Horde.Storage/Implementation/Blob/BlobIndexConsistencyCheckService.cs b/Engine/Source/Programs/Horde/HordeStorage/Horde.Storage/Implementation/Blob/BlobIndexConsistencyCheckService.cs
new file mode 100644
--- /dev/null
+++ b/Engine/Source/Programs/Horde/HordeStorage/Horde.Storage/Implementation/Blob/BlobIndexConsistencyCheckService.cs
@@ -0,0 +1,121 @@
+// Copyright Epic Games, Inc. All Rights Reserved.
+
+using System;
+using System.Linq;
+using System.Threading;
+using System.Threading.Tasks;
+using Datadog.Trace;
+using Horde.Storage.Implementation.Blob;
+using Jupiter;
+using Jupiter.Implementation;
+using Microsoft.Extensions.Options;
+using Serilog;
+
+namespace Horde.Storage.Implementation
+{
+    // ReSharper disable once ClassNeverInstantiated.Global
+    public class BlobIndexConsistencyCheckService : PollingService<BlobIndexConsistencyCheckService.ConsistencyState>
+    {
+        private readonly IOptionsMonitor<ConsistencyCheckSettings> _settings;
+        private readonly IOptionsMonitor<JupiterSettings> _jupiterSettings;
+        private readonly ILeaderElection _leaderElection;
+        private readonly IBlobIndex _blobIndex;
+        private readonly IBlobService _blobService;
+        private readonly ILogger _logger = Log.ForContext<BlobIndexConsistencyCheckService>();
+
+        public class ConsistencyState
+        {
+        }
+
+        protected override bool ShouldStartPolling()
+        {
+            return _settings.CurrentValue.EnableBlobIndexChecks;
+        }
+
+        public BlobIndexConsistencyCheckService(IOptionsMonitor<ConsistencyCheckSettings> settings, IOptionsMonitor<JupiterSettings> jupiterSettings, ILeaderElection leaderElection, IBlobIndex blobIndex, IBlobService blobService) :
+            base(serviceName: nameof(BlobStoreConsistencyCheckService), TimeSpan.FromSeconds(settings.CurrentValue.ConsistencyCheckPollFrequencySeconds), new ConsistencyState())
+        {
+            _settings = settings;
+            _jupiterSettings = jupiterSettings;
+            _leaderElection = leaderElection;
+            _blobIndex = blobIndex;
+            _blobService = blobService;
+        }
+
+        public override async Task<bool> OnPoll(ConsistencyState state, CancellationToken cancellationToken)
+        {
+            if (!_settings.CurrentValue.EnableBlobIndexChecks)
+            {
+                _logger.Information("Skipped running blob index consistency check as it is disabled");
+                return false;
+            }
+
+            if (!_leaderElection.IsThisInstanceLeader())
+            {
+                _logger.Information("Skipped running blob index consistency check because this instance was not the leader");
+                return false;
+            }
+
+            using IScope scope = Tracer.Instance.StartActive("blob_index.consistency.poll");
+
+            await RunConsistencyCheck();
+
+            return true;
+        }
+
+        private async Task RunConsistencyCheck()
+        {
+            string currentRegion = _jupiterSettings.CurrentValue.CurrentSite;
+            await foreach (IBlobIndex.BlobInfo blobInfo in _blobIndex.GetAllBlobs())
+            {
+                if (!blobInfo.Regions.Any())
+                {
+                    _logger.Warning("Blob {Blob} in namespace {Namespace} is not tracked to exist in any region", blobInfo.BlobIdentifier, blobInfo.Namespace);
+
+                    if (await _blobService.ExistsInRootStore(blobInfo.Namespace, blobInfo.BlobIdentifier))
+                    {
+                        _logger.Warning("Blob {Blob} in namespace {Namespace} was found to exist in our blob store so re-adding it to the index", blobInfo.BlobIdentifier, blobInfo.Namespace);
+
+                        // we did have it in our blob store so we adjust the index
+                        await _blobIndex.AddBlobToIndex(blobInfo.Namespace, blobInfo.BlobIdentifier);
+                        continue;
+                    }
+                    else
+                    {
+                        // this blob doesn't exist anywhere so we just cleanup the blob index
+                        _logger.Warning("Blob {Blob} in namespace {Namespace} was removed from the blob index as it didnt exist anywhere", blobInfo.BlobIdentifier, blobInfo.Namespace);
+
+                        await _blobIndex.RemoveBlobFromIndex(blobInfo.Namespace, blobInfo.BlobIdentifier);
+                    }
+                }
+
+                if (blobInfo.Regions.Contains(currentRegion))
+                {
+                    if (!await _blobService.ExistsInRootStore(blobInfo.Namespace, blobInfo.BlobIdentifier))
+                    {
+                        _logger.Warning("Blob {Blob} in namespace {Namespace} did not exist in root store but is tracked as doing so in the blob index. Attempting to replicate it.", blobInfo.BlobIdentifier, blobInfo.Namespace);
+                        try
+                        {
+                            BlobContents _ = await _blobService.ReplicateObject(blobInfo.Namespace, blobInfo.BlobIdentifier, force: true);
+                        }
+                        catch (BlobReplicationException e)
+                        {
+                            _logger.Error(e, "Failed to replicate Blob {Blob} in namespace {Namespace}. Unable to repair the blob index", blobInfo.BlobIdentifier, blobInfo.Namespace);
+
+                            // we update the blob index to accurately reflect that we do not have the blob, this is not good though as it means a upload that we thought happened now lacks content
+                            await _blobIndex.RemoveBlobFromRegion(blobInfo.Namespace, blobInfo.BlobIdentifier);
+                        }
+                    }
+
+                }
+            }
+
+            await Task.CompletedTask;
+        }
+
+        protected override Task OnStopping(ConsistencyState state)
+        {
+            return Task.CompletedTask;
+        }
+    }
+}
diff --git a/Engine/Source/Programs/Horde/HordeStorage/Horde.Storage/Implementation/Blob/BlobStoreConsistencyCheckService.cs b/Engine/Source/Programs/Horde/HordeStorage/Horde.Storage/Implementation/Blob/BlobStoreConsistencyCheckService.cs
new file mode 100644
--- /dev/null
+++ b/Engine/Source/Programs/Horde/HordeStorage/Horde.Storage/Implementation/Blob/BlobStoreConsistencyCheckService.cs
@@ -0,0 +1,129 @@
+// Copyright Epic Games, Inc. All Rights Reserved.
+
+using System;
+using System.Collections.Generic;
+using System.IO;
+using System.Threading;
+using System.Threading.Tasks;
+using Dasync.Collections;
+using Datadog.Trace;
+using EpicGames.Horde.Storage;
+using Horde.Storage.Implementation.Blob;
+using Jupiter;
+using Jupiter.Implementation;
+using Microsoft.Extensions.DependencyInjection;
+using Microsoft.Extensions.Options;
+using Serilog;
+
+namespace Horde.Storage.Implementation
+{
+    // ReSharper disable once ClassNeverInstantiated.Global
+    public class BlobStoreConsistencyCheckService : PollingService<BlobStoreConsistencyCheckService.ConsistencyState>
+    {
+        private readonly IOptionsMonitor<ConsistencyCheckSettings> _settings;
+        private readonly IServiceProvider _provider;
+        private readonly ILeaderElection _leaderElection;
+        private readonly IRefsStore _refsStore;
+        private readonly IReferencesStore _referencesStore;
+        private readonly IBlobIndex _blobIndex;
+        private readonly ILogger _logger = Log.ForContext<BlobStoreConsistencyCheckService>();
+
+        public class ConsistencyState
+        {
+        }
+
+        protected override bool ShouldStartPolling()
+        {
+            return _settings.CurrentValue.EnableBlobStoreChecks;
+        }
+
+        public BlobStoreConsistencyCheckService(IOptionsMonitor<ConsistencyCheckSettings> settings, IServiceProvider provider, ILeaderElection leaderElection, IRefsStore refsStore, IReferencesStore referencesStore, IBlobIndex blobIndex) :
+            base(serviceName: nameof(BlobStoreConsistencyCheckService), TimeSpan.FromSeconds(settings.CurrentValue.ConsistencyCheckPollFrequencySeconds), new ConsistencyState())
+        {
+            _settings = settings;
+            _provider = provider;
+            _leaderElection = leaderElection;
+            _refsStore = refsStore;
+            _referencesStore = referencesStore;
+            _blobIndex = blobIndex;
+        }
+
+        public override async Task<bool> OnPoll(ConsistencyState state, CancellationToken cancellationToken)
+        {
+            if (!_settings.CurrentValue.EnableBlobStoreChecks)
+            {
+                _logger.Information("Skipped running blob store consistency check as it is disabled");
+                return false;
+            }
+
+            if (!_leaderElection.IsThisInstanceLeader())
+            {
+                _logger.Information("Skipped running blob store consistency check because this instance was not the leader");
+                return false;
+            }
+
+            using IScope scope = Tracer.Instance.StartActive("blob_store.consistency_check.poll");
+
+            await RunConsistencyCheck();
+
+            return true;
+        }
+
+        private async Task RunConsistencyCheck()
+        {
+            //TODO: This isn't really specific to the S3 store, but it should likely only run on the backing store and not the cache layers which we do not have a way of fetching right now.
+            AmazonS3Store? s3Store;
+            try
+            {
+                s3Store = _provider.GetService<AmazonS3Store>();
+            }
+            catch (Exception)
+            {
+                s3Store = null;
+            }
+
+            if (s3Store == null)
+            {
+                _logger.Warning("No S3 Store found so will not run any consistency check");
+                return;
+            }
+
+            List<NamespaceId> namespaces = await _refsStore.GetNamespaces().ToListAsync();
+            namespaces.AddRange(await _referencesStore.GetNamespaces().ToListAsync());
+
+            long countOfBlobsChecked = 0;
+            // technically this does not need to be run per namespace but per s3 bucket
+            await foreach (NamespaceId ns in namespaces)
+            {
+                using IScope scope = Tracer.Instance.StartActive("consistency_check.run");
+                scope.Span.ResourceName = ns.ToString();
+
+                await foreach ((BlobIdentifier blob, DateTime lastModified) in s3Store.ListObjects(ns))
+                {
+                    if (countOfBlobsChecked % 100 == 0)
+                        _logger.Information("Consistency check running on S3 blob store, count of blobs processed so far: {CountOfBlobs}", countOfBlobsChecked);
+                    Interlocked.Increment(ref countOfBlobsChecked);
+                    
+                    BlobContents contents = await s3Store.GetObject(ns, blob);
+                    await using Stream s = contents.Stream;
+
+                    BlobIdentifier newHash = await BlobIdentifier.FromStream(s);
+                    if (!blob.Equals(newHash))
+                    {
+                        _logger.Error("Mismatching hash for S3 object {Blob} in {Namespace}, new hash has {NewHash}. Deleting incorrect blob.", blob, ns, newHash);
+
+                        await s3Store.DeleteObject(ns, blob);
+                        // update blob index tracking to indicate that we no longer have this blob in this region
+                        await _blobIndex.RemoveBlobFromRegion(ns, blob);
+                    }
+                }
+                
+            }
+        }
+
+        protected override Task OnStopping(ConsistencyState state)
+        {
+            return Task.CompletedTask;
+        }
+    }
+}
diff --git a/Engine/Source/Programs/Horde/HordeStorage/Horde.Storage/Implementation/Blob/CachedBlobIndex.cs b/Engine/Source/Programs/Horde/HordeStorage/Horde.Storage/Implementation/Blob/CachedBlobIndex.cs
--- a/Engine/Source/Programs/Horde/HordeStorage/Horde.Storage/Implementation/Blob/CachedBlobIndex.cs
+++ b/Engine/Source/Programs/Horde/HordeStorage/Horde.Storage/Implementation/Blob/CachedBlobIndex.cs
@@ -66,4 +66,9 @@ public class CachedBlobIndex : IBlobIndex
         // We do not actually track any blob information when running in cached mode
         await Task.CompletedTask;
     }
+
+    public IAsyncEnumerable<IBlobIndex.BlobInfo> GetAllBlobs()
+    {
+        throw new NotImplementedException();
+    }
 }
diff --git a/Engine/Source/Programs/Horde/HordeStorage/Horde.Storage/Implementation/Blob/ConsistencyCheckService.cs b/Engine/Source/Programs/Horde/HordeStorage/Horde.Storage/Implementation/Blob/ConsistencyCheckService.cs
deleted file mode 100644
diff --git a/Engine/Source/Programs/Horde/HordeStorage/Horde.Storage/Implementation/Blob/IBlobIndex.cs b/Engine/Source/Programs/Horde/HordeStorage/Horde.Storage/Implementation/Blob/IBlobIndex.cs
--- a/Engine/Source/Programs/Horde/HordeStorage/Horde.Storage/Implementation/Blob/IBlobIndex.cs
+++ b/Engine/Source/Programs/Horde/HordeStorage/Horde.Storage/Implementation/Blob/IBlobIndex.cs
@@ -26,5 +26,6 @@ namespace Horde.Storage.Implementation.Blob
 
         Task<bool> BlobExistsInRegion(NamespaceId ns, BlobIdentifier blobIdentifier);
         Task AddRefToBlobs(NamespaceId ns, BucketId bucket, IoHashKey key, BlobIdentifier[] blobs);
+        IAsyncEnumerable<BlobInfo> GetAllBlobs();
     }
 }
diff --git a/Engine/Source/Programs/Horde/HordeStorage/Horde.Storage/Implementation/Blob/IBlobService.cs b/Engine/Source/Programs/Horde/HordeStorage/Horde.Storage/Implementation/Blob/IBlobService.cs
--- a/Engine/Source/Programs/Horde/HordeStorage/Horde.Storage/Implementation/Blob/IBlobService.cs
+++ b/Engine/Source/Programs/Horde/HordeStorage/Horde.Storage/Implementation/Blob/IBlobService.cs
@@ -31,10 +31,18 @@ public interface IBlobService
 
     Task<BlobContents> GetObject(NamespaceId ns, BlobIdentifier blob);
 
-    Task<BlobContents> ReplicateObject(NamespaceId ns, BlobIdentifier blob);
+    Task<BlobContents> ReplicateObject(NamespaceId ns, BlobIdentifier blob, bool force = false);
 
     Task<bool> Exists(NamespaceId ns, BlobIdentifier blob);
 
+    /// <summary>
+    /// Checks that the blob exists in the root store, the store which is last in the list and thus is intended to have every blob in it
+    /// </summary>
+    /// <param name="ns">The namespace</param>
+    /// <param name="blob">The identifier of the blob</param>
+    /// <returns></returns>
+    Task<bool> ExistsInRootStore(NamespaceId ns, BlobIdentifier blob);
+
     // Delete a object
     Task DeleteObject(NamespaceId ns, BlobIdentifier blob);
 
@@ -248,9 +256,9 @@ public class BlobService : IBlobService
         return blobContents;
     }
 
-    public async Task<BlobContents> ReplicateObject(NamespaceId ns, BlobIdentifier blob)
+    public async Task<BlobContents> ReplicateObject(NamespaceId ns, BlobIdentifier blob, bool force = false)
     {
-        if (!ShouldFetchBlobOnDemand(ns))
+        if (!force && !ShouldFetchBlobOnDemand(ns))
             throw new NotSupportedException($"Replication is not allowed in namespace {ns}");
 
         using IScope scope = Tracer.Instance.StartActive("HierarchicalStore.Replicate");
@@ -259,6 +267,9 @@ public class BlobService : IBlobService
         if (blobInfo == null)
             throw new BlobNotFoundException(ns, blob);
 
+        if (!blobInfo.Regions.Any())
+            throw new BlobReplicationException(ns, blob, "Blob not found in any region");
+
         _logger.Information("On-demand replicating blob {Blob} in Namespace {Ns}", blob, ns);
         SortedList<int, string> possiblePeers = _peerStatusService.GetPeersByLatency(blobInfo.Regions.ToList());
 
@@ -338,7 +349,23 @@ public class BlobService : IBlobService
             return false;
         }
     }
-    
+
+    public async Task<bool> ExistsInRootStore(NamespaceId ns, BlobIdentifier blob)
+    {
+        IBlobStore store = _blobStores.Last();
+
+        using IScope scope = Tracer.Instance.StartActive("HierarchicalStore.ObjectExists");
+        scope.Span.ResourceName = blob.ToString();
+        scope.Span.SetTag("BlobStore", store.GetType().Name);
+        if (await store.Exists(ns, blob))
+        {
+            scope.Span.SetTag("ObjectFound", true.ToString());
+            return true;
+        }
+        scope.Span.SetTag("ObjectFound", false.ToString());
+        return false;
+    }
+
     public async Task DeleteObject(NamespaceId ns, BlobIdentifier blob)
     {
         bool blobNotFound = false;
diff --git a/Engine/Source/Programs/Horde/HordeStorage/Horde.Storage/Implementation/Blob/MemoryBlobIndex.cs b/Engine/Source/Programs/Horde/HordeStorage/Horde.Storage/Implementation/Blob/MemoryBlobIndex.cs
--- a/Engine/Source/Programs/Horde/HordeStorage/Horde.Storage/Implementation/Blob/MemoryBlobIndex.cs
+++ b/Engine/Source/Programs/Horde/HordeStorage/Horde.Storage/Implementation/Blob/MemoryBlobIndex.cs
@@ -100,6 +100,19 @@ public class MemoryBlobIndex : IBlobIndex
         return Task.CompletedTask;
     }
 
+    public async IAsyncEnumerable<IBlobIndex.BlobInfo> GetAllBlobs()
+    {
+        await Task.CompletedTask;
+
+        foreach (KeyValuePair<NamespaceId, ConcurrentDictionary<BlobIdentifier, MemoryBlobInfo>> pair in _index)
+        {
+            foreach ((BlobIdentifier? blobIdentifier, MemoryBlobInfo? blobInfo) in pair.Value)
+            {
+                yield return blobInfo;
+            }
+        }
+    }
+
     private MemoryBlobInfo NewBlobInfo(NamespaceId ns, BlobIdentifier blob, string region)
     {
         MemoryBlobInfo info = new MemoryBlobInfo
diff --git a/Engine/Source/Programs/Horde/HordeStorage/Horde.Storage/Implementation/Blob/MongoBlobIndex.cs b/Engine/Source/Programs/Horde/HordeStorage/Horde.Storage/Implementation/Blob/MongoBlobIndex.cs
--- a/Engine/Source/Programs/Horde/HordeStorage/Horde.Storage/Implementation/Blob/MongoBlobIndex.cs
+++ b/Engine/Source/Programs/Horde/HordeStorage/Horde.Storage/Implementation/Blob/MongoBlobIndex.cs
@@ -116,6 +116,19 @@ public class MongoBlobIndex : MongoStore, IBlobIndex
         await Task.WhenAll(refUpdateTasks);
     }
 
+    public async IAsyncEnumerable<IBlobIndex.BlobInfo> GetAllBlobs()
+    {
+        IMongoCollection<MongoBlobIndexModelV0> collection = GetCollection<MongoBlobIndexModelV0>();
+        IAsyncCursor<MongoBlobIndexModelV0>? cursor = await collection.FindAsync(FilterDefinition<MongoBlobIndexModelV0>.Empty);
+
+        while (await cursor.MoveNextAsync())
+        {
+            foreach (MongoBlobIndexModelV0 model in cursor.Current)
+            {
+                yield return model.ToBlobInfo();
+            }
+        }
+    }
 }
 
 [BsonDiscriminator("blob-index.v0")]
diff --git a/Engine/Source/Programs/Horde/HordeStorage/Horde.Storage/Implementation/Blob/ScyllaBlobIndex.cs b/Engine/Source/Programs/Horde/HordeStorage/Horde.Storage/Implementation/Blob/ScyllaBlobIndex.cs
--- a/Engine/Source/Programs/Horde/HordeStorage/Horde.Storage/Implementation/Blob/ScyllaBlobIndex.cs
+++ b/Engine/Source/Programs/Horde/HordeStorage/Horde.Storage/Implementation/Blob/ScyllaBlobIndex.cs
@@ -57,13 +57,7 @@ public class ScyllaBlobIndex : IBlobIndex
         if (blobIndex == null)
             return null;
 
-        return new IBlobIndex.BlobInfo
-        {
-            Regions = blobIndex.Regions ?? new HashSet<string>(),
-            BlobIdentifier = id,
-            Namespace = new NamespaceId(blobIndex.Namespace),
-            References = blobIndex.References != null ? blobIndex.References.Select(reference => reference.AsTuple()).ToList() : new List<(BucketId, IoHashKey)>()
-        };
+        return blobIndex.ToBlobInfo();
     }
 
     public async Task<bool> RemoveBlobFromIndex(NamespaceId ns, BlobIdentifier id)
@@ -92,7 +86,7 @@ public class ScyllaBlobIndex : IBlobIndex
         return blobInfo?.Regions.Contains(_jupiterSettings.CurrentValue.CurrentSite) ?? false;
     }
 
-    public Task AddRefToBlobs(NamespaceId ns, BucketId bucket, IoHashKey key, BlobIdentifier[] blobs)
+    public async Task AddRefToBlobs(NamespaceId ns, BucketId bucket, IoHashKey key, BlobIdentifier[] blobs)
     {
         using IScope _ = Tracer.Instance.StartActive("scylla.add_ref_blobs");
 
@@ -109,9 +103,18 @@ public class ScyllaBlobIndex : IBlobIndex
                 new ScyllaBlobIdentifier(id));
         }
 
-        return Task.WhenAll(refUpdateTasks);
+        await Task.WhenAll(refUpdateTasks);
     }
 
+    public async IAsyncEnumerable<IBlobIndex.BlobInfo> GetAllBlobs()
+    {
+        using IScope _ = Tracer.Instance.StartActive("scylla.get_all_blobs");
+
+        foreach (ScyllaBlobIndexTable blobIndex in await _mapper.FetchAsync<ScyllaBlobIndexTable>("BYPASS CACHE"))
+        {
+            yield return blobIndex.ToBlobInfo();
+        }
+    }
 }
 
 
@@ -146,4 +149,15 @@ class ScyllaBlobIndexTable
 
     [Cassandra.Mapping.Attributes.Column("references")]
     public List<ScyllaObjectReference>? References { get; set; }
+
+    public IBlobIndex.BlobInfo ToBlobInfo()
+    {
+        return new IBlobIndex.BlobInfo
+        {
+            Regions = Regions ?? new HashSet<string>(),
+            BlobIdentifier = BlobId.AsBlobIdentifier(),
+            Namespace = new NamespaceId(Namespace),
+            References = References != null ? References.Select(reference => reference.AsTuple()).ToList() : new List<(BucketId, IoHashKey)>()
+        };
+    }
 }
diff --git a/Engine/Source/Programs/Horde/HordeStorage/Horde.Storage/Implementation/MongoStore.cs b/Engine/Source/Programs/Horde/HordeStorage/Horde.Storage/Implementation/MongoStore.cs
--- a/Engine/Source/Programs/Horde/HordeStorage/Horde.Storage/Implementation/MongoStore.cs
+++ b/Engine/Source/Programs/Horde/HordeStorage/Horde.Storage/Implementation/MongoStore.cs
@@ -77,7 +77,7 @@ public class MongoStore
         return _client.GetDatabase(dbName).GetCollection<T>(collectionName);
     }
 
-    private string GetDatabaseName()
+    internal string GetDatabaseName()
     {
         return _overrideDatabaseName ?? "HordeStorage";
     }
diff --git a/Engine/Source/Programs/Horde/HordeStorage/Horde.Storage/Implementation/Objects/MongoReferencesStore.cs b/Engine/Source/Programs/Horde/HordeStorage/Horde.Storage/Implementation/Objects/MongoReferencesStore.cs
--- a/Engine/Source/Programs/Horde/HordeStorage/Horde.Storage/Implementation/Objects/MongoReferencesStore.cs
+++ b/Engine/Source/Programs/Horde/HordeStorage/Horde.Storage/Implementation/Objects/MongoReferencesStore.cs
@@ -3,7 +3,6 @@
 using System;
 using System.Collections.Generic;
 using System.Threading.Tasks;
-using Dasync.Collections;
 using EpicGames.Horde.Storage;
 using Jupiter.Implementation;
 using Microsoft.Extensions.Options;
@@ -105,7 +104,7 @@ namespace Horde.Storage.Implementation
 
             while (await cursor.MoveNextAsync())
             {
-                await foreach (MongoReferencesModelV0 model in cursor.Current)
+                foreach (MongoReferencesModelV0 model in cursor.Current)
                 {
                     yield return (new BucketId(model.Bucket), new IoHashKey(model.Key), model.LastAccessTime);
                 }
diff --git a/Engine/Source/Programs/Horde/HordeStorage/Horde.Storage/Implementation/Objects/ScyllaReferencesStore.cs b/Engine/Source/Programs/Horde/HordeStorage/Horde.Storage/Implementation/Objects/ScyllaReferencesStore.cs
--- a/Engine/Source/Programs/Horde/HordeStorage/Horde.Storage/Implementation/Objects/ScyllaReferencesStore.cs
+++ b/Engine/Source/Programs/Horde/HordeStorage/Horde.Storage/Implementation/Objects/ScyllaReferencesStore.cs
@@ -50,7 +50,7 @@ namespace Horde.Storage.Implementation
             );"
             ));
 
-            // BYPASS CACHE is a scylla specific extension to disable populating the cache, should be disabled by other cassandra dbs
+            // BYPASS CACHE is a scylla specific extension to disable populating the cache, should be ignored by other cassandra dbs
             _getObjectsStatement = _session.Prepare("SELECT bucket, name, last_access_time FROM objects WHERE namespace = ? ALLOW FILTERING BYPASS CACHE");
             _getNamespacesStatement = _session.Prepare("SELECT DISTINCT namespace FROM buckets");
         }
diff --git a/Engine/Source/Programs/Horde/HordeStorage/Horde.Storage/Settings.cs b/Engine/Source/Programs/Horde/HordeStorage/Horde.Storage/Settings.cs
--- a/Engine/Source/Programs/Horde/HordeStorage/Horde.Storage/Settings.cs
+++ b/Engine/Source/Programs/Horde/HordeStorage/Horde.Storage/Settings.cs
@@ -285,4 +285,11 @@ namespace Horde.Storage
 
         public bool IsInternal { get; set; } = false;
     }
+
+    public class ConsistencyCheckSettings
+    {
+        public bool EnableBlobStoreChecks { get; set; } = false;
+        public bool EnableBlobIndexChecks { get; set; } = false;
+        public double ConsistencyCheckPollFrequencySeconds { get; set; } = TimeSpan.FromHours(2).TotalSeconds;
+    }
 }
diff --git a/Engine/Source/Programs/Horde/HordeStorage/Horde.Storage/Tests/Functional/Blob/BlobIndexTests.cs b/Engine/Source/Programs/Horde/HordeStorage/Horde.Storage/Tests/Functional/Blob/BlobIndexTests.cs
--- a/Engine/Source/Programs/Horde/HordeStorage/Horde.Storage/Tests/Functional/Blob/BlobIndexTests.cs
+++ b/Engine/Source/Programs/Horde/HordeStorage/Horde.Storage/Tests/Functional/Blob/BlobIndexTests.cs
@@ -3,11 +3,13 @@
 using System;
 using System.Collections.Generic;
 using System.IO;
+using System.Linq;
 using System.Net.Http;
 using System.Net.Http.Headers;
 using System.Net.Mime;
 using System.Text;
 using System.Threading.Tasks;
+using async_enumerable_dotnet;
 using Cassandra;
 using EpicGames.Horde.Storage;
 using EpicGames.Serialization;
@@ -19,7 +21,9 @@ using Microsoft.AspNetCore.Hosting;
 using Microsoft.AspNetCore.TestHost;
 using Microsoft.Extensions.Configuration;
 using Microsoft.Extensions.DependencyInjection;
+using Microsoft.Extensions.Options;
 using Microsoft.VisualStudio.TestTools.UnitTesting;
+using MongoDB.Driver;
 using Serilog;
 using Logger = Serilog.Core.Logger;
 
@@ -150,6 +154,54 @@ namespace Horde.Storage.FunctionalTests.Storage
             Assert.IsNull(deletedBlobInfo);
         }
 
+
+        [TestMethod]
+        public async Task EnumerateAllBlobs()
+        {
+            IBlobIndex? index = _server!.Services.GetService<IBlobIndex>();
+            Assert.IsNotNull(index);
+            {
+                // verify the blob info list is empty at the start
+                IBlobIndex.BlobInfo[]? blobInfos =  await index.GetAllBlobs().ToArrayAsync();
+                Assert.AreEqual(0, blobInfos.Length);
+            }
+
+            // upload a blob
+            BlobIdentifier contentHash;
+            {
+                byte[] payload = Encoding.ASCII.GetBytes("I am a blob with contents");
+                contentHash = BlobIdentifier.FromBlob(payload);
+
+                {
+                    using ByteArrayContent requestContent = new ByteArrayContent(payload);
+                    requestContent.Headers.ContentType = new MediaTypeHeaderValue(MediaTypeNames.Application.Octet);
+                    HttpResponseMessage response = await _httpClient!.PutAsync(requestUri: $"api/v1/blobs/{TestNamespaceName}/{contentHash}", requestContent);
+                    response.EnsureSuccessStatusCode();
+                }
+            }
+
+            // upload a compressed blob
+            BlobIdentifier compressedPayloadIdentifier;
+            {
+                byte[] texturePayload = await File.ReadAllBytesAsync($"ContentId/Payloads/UncompressedTexture_CAS_dea81b6c3b565bb5089695377c98ce0f1c13b0c3.udd");
+                compressedPayloadIdentifier = BlobIdentifier.FromBlob(texturePayload);
+                BlobIdentifier uncompressedPayloadIdentifier = new BlobIdentifier("DEA81B6C3B565BB5089695377C98CE0F1C13B0C3");
+
+                ByteArrayContent content = new ByteArrayContent(texturePayload);
+                content.Headers.ContentType = new MediaTypeHeaderValue(CustomMediaTypeNames.UnrealCompressedBuffer);
+                HttpResponseMessage result = await _httpClient!.PutAsync($"api/v1/compressed-blobs/{TestNamespaceName}/{uncompressedPayloadIdentifier}", content);
+                result.EnsureSuccessStatusCode();
+            }
+
+            {
+                IBlobIndex.BlobInfo[]? blobInfos =  await index.GetAllBlobs().ToArrayAsync();
+                Assert.AreEqual(2, blobInfos.Length);
+
+                Assert.IsNotNull(blobInfos.FirstOrDefault(info => info.BlobIdentifier.Equals(compressedPayloadIdentifier)));
+                Assert.IsNotNull(blobInfos.FirstOrDefault(info => info.BlobIdentifier.Equals(contentHash)));
+            }
+        }
+
      
     }
 
@@ -202,9 +254,25 @@ namespace Horde.Storage.FunctionalTests.Storage
             return new[] { new KeyValuePair<string, string>("Horde_Storage:BlobIndexImplementation", HordeStorageSettings.BlobIndexImplementations.Mongo.ToString()) };
         }
 
-        protected override Task Seed(IServiceProvider serverServices)
+        protected override async Task Seed(IServiceProvider provider)
         {
-            return Task.CompletedTask;
+            IOptionsMonitor<MongoSettings> mongoSettings = provider.GetService<IOptionsMonitor<MongoSettings>>()!;
+
+            IBlobIndex blobIndex = provider.GetService<IBlobIndex>()!;
+            Assert.IsTrue(blobIndex is MongoBlobIndex);
+
+            MongoBlobIndex mongoBlobIndex = (MongoBlobIndex)blobIndex;
+            
+            MongoClient client = GetMongoClient(mongoSettings);
+            string dbName = mongoBlobIndex.GetDatabaseName();
+            if (client.GetDatabase(dbName) != null)
+            {
+                await client.DropDatabaseAsync(dbName);
+            }
+        }
+        private static MongoClient GetMongoClient(IOptionsMonitor<MongoSettings> mongoSettings)
+        {
+            return new MongoClient(mongoSettings.CurrentValue.ConnectionString);
         }
 
         protected override async Task Teardown(IServiceProvider provider)
diff --git a/Engine/Source/Programs/Horde/HordeStorage/Jupiter.Common/BaseStartup.cs b/Engine/Source/Programs/Horde/HordeStorage/Jupiter.Common/BaseStartup.cs
--- a/Engine/Source/Programs/Horde/HordeStorage/Jupiter.Common/BaseStartup.cs
+++ b/Engine/Source/Programs/Horde/HordeStorage/Jupiter.Common/BaseStartup.cs
@@ -328,7 +328,7 @@ namespace Jupiter
             OnConfigureApp(app, env);
         }
 
-        public virtual bool UseDeveloperExceptionPage { get; } = true;
+        public virtual bool UseDeveloperExceptionPage { get; } = false;
 
         protected virtual void OnConfigureAppEarly(IApplicationBuilder app, IWebHostEnvironment env)
         {
