Fix slow loop when loading Automated Test Session (backend index is by element not page)
Fetch smaller chunk of data to improve feedback
Use JSX.Element instead of React.FC when no State is used
Group test results by test case for better render performance(less icon to display)
Make test history loading dynamic

#rnx

diff --git a/Engine/Source/Programs/Horde/HordeDashboard/src/Plugins/AutomatedTestSession/models/AutomatedTestSession.ts b/Engine/Source/Programs/Horde/HordeDashboard/src/Plugins/AutomatedTestSession/models/AutomatedTestSession.ts
--- a/Engine/Source/Programs/Horde/HordeDashboard/src/Plugins/AutomatedTestSession/models/AutomatedTestSession.ts
+++ b/Engine/Source/Programs/Horde/HordeDashboard/src/Plugins/AutomatedTestSession/models/AutomatedTestSession.ts
@@ -13,6 +13,16 @@ export enum TestState {
     Unknown = "Unknown",
 }
 
+const TestStateRank: {[Key in string]: number}  = {
+    "Fail": 0,
+    "SuccessWithWarnings": 1,
+    "InProcess": 2,
+    "NotRun": 3,
+    "Success": 4,
+    "Skipped": 5,
+    "Unknown": 6,
+};
+
 export enum EventType {
     Info = "Info",
     Error = "Error",
@@ -118,13 +128,17 @@ export class MetaWrapper {
         return undefined;
     }
 
-    map(mapFunc: (key: string) => any) : any[] {
+    map(mapFunc: (key: string, value?: string) => any) : any[] {
         if (this.collection) {
-            return Object.keys(this.collection).map(mapFunc);
+            return Object.keys(this.collection).map((key) => mapFunc(key, this.collection![key]));
         }
         return [];
     }
 
+    maskByKeys(keyMask: string[]) {
+        return this.map((key, value) => keyMask.includes(key)?undefined:value).filter((info) => info)
+    }
+
     private generateHashCode(): string {
         if(this.collection) {
             return generateHashFromMetadata(this.collection);
@@ -266,6 +280,166 @@ export class TestStats {
     Skipped: number = 0;
     Unexecuted: number = 0;
     Incomplete: number = 0;
+
+    update(state: string, increment: boolean = true) {
+        switch(state) {
+        case TestState.Failed:
+            increment? this.Failed++ : this.Failed--;
+            break;
+        case TestState.NotRun:
+            increment? this.Unexecuted++ : this.Unexecuted--;
+            break;
+        case TestState.Success:
+            increment? this.Passed++ : this.Passed--;
+            break;
+        case TestState.InProcess:
+            increment? this.Incomplete++ : this.Incomplete--;
+            break;
+        case TestState.Skipped:
+            increment? this.Skipped++ : this.Skipped--;
+            break;
+        }
+    }
+
+    updateFromStats(stats: TestStats) {
+        this.Failed += stats.Failed;
+        this.Passed += stats.Passed;
+        this.Skipped += stats.Skipped;
+        this.Unexecuted += stats.Unexecuted;
+        this.Incomplete += stats.Incomplete;
+    }
+}
+
+export class TestCase {
+    Name: string;
+    DisplayName: string;
+    TestUID: string;
+    SessionName: string;
+    Suite: string;
+    State: string;
+    Metas: Map<string, TestResult>;
+    Stats: TestStats;
+
+    @observable
+    version: number = 0;
+    @action
+    updateVersion() {
+        this.version++;
+    }
+    whenVersionChange() {
+        const start = this.version;
+        return when(() => this.version !== start);
+    }
+
+    @observable
+    history: number = 0;
+    @action
+    updateHistory() {
+        this.history++;
+    }
+    reactOnHistoryChange(callback: (version: number) => void) {
+        return reaction(
+            () => this.history,
+            (version) => callback(version)
+        );
+    }
+
+    constructor(name: string, testuid: string, suite: string, sessionName: string, displayName?: string) {
+        this.Name = name;
+        this.DisplayName = displayName??this.toDislayName(name);
+        this.TestUID = testuid;
+        this.State = TestState.Unknown;
+        this.Suite = suite;
+        this.SessionName = sessionName;
+        this.Metas = new Map();
+        this.Stats = new TestStats();
+    }
+
+    removeMeta(meta: MetaWrapper) {
+        const metauid = meta.hash;
+        const test = this.Metas.get(metauid);
+        if (!test) {
+            return false;
+        }
+        this.Metas.delete(metauid);
+        this.Stats.update(test.State, false);
+        this.updateVersion();
+        return true;
+    }
+
+    isEmpty() {
+        return this.Metas.size === 0;
+    }
+
+    addMeta(item: TestResult) {
+        const metauid = item.Session.MetaHandler!.hash;
+        if(this.Metas.has(metauid)) {
+            return false;
+        }
+        this.Metas.set(metauid, item);
+        let state = item.State;
+        if (state === TestState.Success && item.WarningCount > 0) {
+            state = TestState.SuccessWithWarnings;
+            item.State = state;
+        }
+        this.setState(state);
+        this.Stats.update(state);
+        this.updateVersion();
+        return true;
+    }
+
+    forEachMeta(predicate: (test: TestResult) => void) {
+        this.Metas.forEach(predicate);
+    }
+
+    getResults(sorted: boolean = false) {
+        const results = Array.from(this.Metas.values());
+        if (sorted) {
+            results.sort(this.sortResults);
+        }
+        return results;
+    }
+
+    getFilteredTestCase(metafilter: Metadata) {
+        const metas = Array.from(this.Metas.values()).filter(
+            (test) => Object.keys(metafilter).every((key) => test.Session.Metadata[key] === metafilter[key])
+        );
+        if(metas.length === 0) {
+            return undefined;
+        }
+        const testcase = new TestCase(this.Name, this.TestUID, this.Suite, this.SessionName, this.DisplayName);
+        metas.forEach((test) => testcase.addMeta(test));
+        return testcase;
+    }
+
+    private toDislayName(name: string): string {
+        // It is necessary to sanitize the test names that contains / instead of . syntax
+        if(name.indexOf('/') > -1) {
+            return name.replaceAll(/\.*\//g, ".");
+        }
+        return name;
+    }
+
+    private setState(state: string) {
+        if(this.State !== state) {
+            const currentRank = TestStateRank[this.State];
+            const newRank = TestStateRank[state];
+            if(newRank < currentRank) {
+                this.State = state;
+            }
+        }
+    }
+
+    private sortResults(a: TestResult, b: TestResult) {
+        const aRank = TestStateRank[a.State];
+        const bRank = TestStateRank[b.State];
+        if(aRank > bRank) {
+            return 1;
+        } else if(aRank < bRank) {
+            return -1;
+        }
+        return 0;
+    }
 }
 
 export class Section {
@@ -273,7 +447,7 @@ export class Section {
     FullName: string;
     Type: FilterType;
     Stats: TestStats;
-    Results: TestResult[];
+    TestCases: TestCase[];
     TestNames: Set<string>;
 
     constructor(name: string, fullName: string, type: FilterType) {
@@ -281,31 +455,14 @@ export class Section {
         this.FullName = fullName;
         this.Type = type;
         this.Stats = new TestStats();
-        this.Results = [];
+        this.TestCases = [];
         this.TestNames = new Set();
     }
 
-    addTest(test: TestResult, testName?: string) {
-        this.Results.push(test);
-        this.TestNames.add(testName??test.Name);
-        // update stats
-        switch(test.State) {
-        case TestState.Failed:
-            this.Stats.Failed++;
-            break;
-        case TestState.NotRun:
-            this.Stats.Unexecuted++;
-            break;
-        case TestState.Success:
-            this.Stats.Passed++;
-            break;
-        case TestState.InProcess:
-            this.Stats.Incomplete++;
-            break;
-        case TestState.Skipped:
-            this.Stats.Skipped++;
-            break;
-        }
+    addTestCase(testcase: TestCase, testName?: string) {
+        this.TestCases.push(testcase);
+        this.TestNames.add(testName??testcase.DisplayName);
+        this.Stats.updateFromStats(testcase.Stats);
     }
 
     get isOneTest(): boolean {
@@ -338,16 +495,14 @@ export type ReportFilter = {
     view?: string;
 }
 
-type GetKeyFunction = (test: TestResult) => [string, string];
+type GetKeyFunction = (testcase: TestCase) => [string, string];
 
 export class SectionCollection {
 
-    private fixedTestNameCache: Map<string, string>;
     private filterCache: Map<string, {hash: string, sections: Section[]}>;
     private onTheFlyCache: {filter: string, hash: string, sections: Section[]} | undefined;
 
     constructor() {
-        this.fixedTestNameCache = new Map();
         this.filterCache = new Map();
     }
 
@@ -356,7 +511,7 @@ export class SectionCollection {
         const filterCacheKey = generateHashFromMetadata(filter);
         const cache = this.filterCache.get(filterCacheKey);
         const isSameHash = cache && cache.hash === sessions.hash;
-        const sections = cache && isSameHash? cache.sections : this.filterSections(filter, sessions.getTests());
+        const sections = cache && isSameHash? cache.sections : this.filterSections(filter, sessions.getTestCases());
         if (!this.filterCache.has(filterCacheKey) || !isSameHash) {
             this.filterCache.set(filterCacheKey, {hash: sessions.hash, sections:sections})
         }
@@ -367,7 +522,7 @@ export class SectionCollection {
         // to use only when iterating through cached data
         let sections = this.extractSections({...filter, namecontains: undefined}, sessions);
         const onTheFlyFilter = {...filter, namecontains: testNameFilter.toLowerCase()}
-        sections = this.filterSections(onTheFlyFilter, sections.flatMap((section) => section.Results));
+        sections = this.filterSections(onTheFlyFilter, sections.flatMap((section) => section.TestCases));
         this.onTheFlyCache = {filter: generateHashFromMetadata(onTheFlyFilter), hash: sessions.hash, sections};
         return sections;
     }
@@ -379,13 +534,8 @@ export class SectionCollection {
         }
     }
 
-    private filterSections(filter: SectionFilter, tests: TestResult[]) {
-        const filterFunc = (test: TestResult) => {
-            if(filter.project) {
-                if(test.Session.MetaHandler?.get('Project') !== filter.project) {
-                    return false;
-                }
-            }
+    private filterSections(filter: SectionFilter, tests: TestCase[]) {
+        const filterFunc = (test: TestCase) => {
             if(filter.suite) {
                 if(test.Suite !== filter.suite) {
                     return false;
@@ -393,13 +543,13 @@ export class SectionCollection {
             }
             if(filter.namecontains) {
                 const substrings = filter.namecontains.split(' ');
-                const testName = test.Name.toLowerCase();
+                const testName = test.DisplayName.toLowerCase();
                 if(!substrings.some((value) => testName.indexOf(value) >= 0)) {
                     return false;
                 }
             }
             if(filter.name) {
-                if(this.fixTestName(test.Name).indexOf(filter.name + '.') !== 0) {
+                if(test.DisplayName.indexOf(filter.name + '.') !== 0) {
                     return false;
                 }
             }
@@ -407,16 +557,16 @@ export class SectionCollection {
         }
         const filteredTests = tests.filter(filterFunc);
 
-        const commonSection = filter.name? this.getCommonSection(filteredTests.map((item) => this.fixTestName(item.Name))) : "";
+        const commonSection = filter.name? this.getCommonSection(filteredTests.map((item) => item.DisplayName)) : "";
         const sectionKeyFunc = (key: FilterType): GetKeyFunction | undefined => {
             switch(key) {
                 case FilterType.suite:
-                    return (test: TestResult) => {
-                        return [test.Suite, test.Suite];
+                    return (testcase: TestCase) => {
+                        return [testcase.Suite, testcase.Suite];
                     }
                 case FilterType.name:
-                    return (test: TestResult) => {
-                        const suite = this.fixTestName(test.Name).slice(commonSection.length).split('.', 1)[0];
+                    return (testcase: TestCase) => {
+                        const suite = testcase.DisplayName.slice(commonSection.length).split('.', 1)[0];
                         return [suite, commonSection+suite];
                     }
             }
@@ -429,11 +579,11 @@ export class SectionCollection {
         return sections;
     }
 
-    private buildSections(tests: TestResult[], type: FilterType, keyFunc: GetKeyFunction): Section[] {
+    private buildSections(testcases: TestCase[], type: FilterType, keyFunc: GetKeyFunction): Section[] {
         const sections : Section[] = [];
         const sectionsMap = new Map<string, Section>(); 
 
-        tests.forEach((test) => {
+        testcases.forEach((test) => {
             const [ key, fullKey ] = keyFunc(test);
             let targetSection = sectionsMap.get(key);
             if(!targetSection) {
@@ -441,25 +591,12 @@ export class SectionCollection {
                 sections.push(targetSection);
                 sectionsMap.set(key, targetSection);
             }
-            targetSection.addTest(test, this.fixTestName(test.Name));
+            targetSection.addTestCase(test, test.DisplayName);
         });
 
         return sections;
     }
 
-    private  fixTestName(name: string): string {
-        // It is necessary to sanitize the test names that contains / instead of . syntax
-        if(this.fixedTestNameCache.has(name)) {
-            return this.fixedTestNameCache.get(name)!;
-        }
-        let fixedName = name
-        if(name.indexOf('/') > -1) {
-            fixedName = name.replaceAll(/\.*\//g, ".");
-        }
-        this.fixedTestNameCache.set(name, fixedName);
-        return fixedName;
-    }
-
     private getCommonString(list: string[]): string {
         if(list.length === 0) {
             return "";
@@ -495,21 +632,23 @@ export class SectionCollection {
     // sort comparing the number of Failed, Incomplete, Unexecuted and Passed
     private sortSections(sections: Section[]) {
         sections.sort((a, b) => {
-            if (a.Stats.Failed > b.Stats.Failed) {
+            const aStats = a.isOneTest? a.TestCases[0].Stats : a.Stats;
+            const bStats = b.isOneTest? b.TestCases[0].Stats : b.Stats;
+            if (aStats.Failed > bStats.Failed) {
                 return -1;
-            } else if (a.Stats.Failed < b.Stats.Failed) {
+            } else if (aStats.Failed < bStats.Failed) {
                 return 1;
-            } else if (a.Stats.Incomplete > b.Stats.Incomplete) {
+            } else if (aStats.Incomplete > bStats.Incomplete) {
                 return -1;
-            } else if (a.Stats.Incomplete < b.Stats.Incomplete) {
+            } else if (aStats.Incomplete < bStats.Incomplete) {
                 return 1;
-            } else if (a.Stats.Unexecuted > b.Stats.Unexecuted) {
+            } else if (aStats.Unexecuted > bStats.Unexecuted) {
                 return -1;
-            } else if (a.Stats.Unexecuted < b.Stats.Unexecuted) {
+            } else if (aStats.Unexecuted < bStats.Unexecuted) {
                 return 1;
-            } else if (a.Stats.Passed > b.Stats.Passed) {
+            } else if (aStats.Passed > bStats.Passed) {
                 return -1;
-            } else if (a.Stats.Passed < b.Stats.Passed) {
+            } else if (aStats.Passed < bStats.Passed) {
                 return 1;
             } else if (a.isOneTest && !b.isOneTest) {
                 return 1;
@@ -522,35 +661,13 @@ export class SectionCollection {
 
     // sort in that order Failed, Incomplete, Unexecuted, Passed, Skipped
     private sortSectionResults(section: Section) {
-        section.Results.sort((a, b) => {
-            if(a.State === TestState.Failed) {
-                if(b.State !== TestState.Failed) {
-                    return -1;
-                }
-            } else if(a.State === TestState.InProcess) {
-                if(b.State === TestState.Failed) {
-                    return 1;
-                } else if(b.State !== TestState.InProcess) {
-                    return -1;
-                }
-            } else if(a.State === TestState.NotRun) {
-                if(b.State === TestState.Failed || b.State === TestState.InProcess) {
-                    return 1;
-                } else if(b.State !== TestState.NotRun) {
-                    return -1;
-                }
-            } else if(a.State === TestState.Success) {
-                if(b.State === TestState.Skipped) {
-                    return -1;
-                } else if(b.State !== TestState.Success || b.WarningCount > 0) {
-                    return 1;
-                } else if(a.WarningCount > 0) {
-                    return -1;
-                }
-            } else if(a.State === TestState.Skipped) {
-                if(b.State !== TestState.Skipped) {
-                    return 1;
-                }
+        section.TestCases.sort((a, b) => {
+            const aRank = TestStateRank[a.State];
+            const bRank = TestStateRank[b.State];
+            if(aRank > bRank) {
+                return 1;
+            } else if(aRank < bRank) {
+                return -1;
             }
 
             if (a.Name > b.Name) {
@@ -605,18 +722,23 @@ export class TestSessionCollection {
     }
 
     private _collection: Map<string, TestSessionWrapper>;
+    private _testcases: Map<string, TestCase>;
     private _datahandler: TestDataCollection;
     private _includePreflight: boolean = false;
 
     constructor(handler: TestDataCollection, includePreflight: boolean = false) {
         this._collection = new Map();
+        this._testcases = new Map();
         this._datahandler = handler;
         this._includePreflight = includePreflight;
     }
 
     set(uid: string, item: TestSessionWrapper) {
         this._collection.set(uid, item);
-        if(item.isAllPropertiesLoaded()) { this.updateVersion(); }
+        if(item.isAllPropertiesLoaded()) {
+            this.indexTestsFromSession(item);
+            this.updateVersion();
+        }
     }
 
     get(uid: string) {
@@ -628,12 +750,25 @@ export class TestSessionCollection {
     }
 
     delete(uid: string) {
+        const session = this._collection.get(uid);
         this._collection.delete(uid);
+        // clean up testuids map
+        if(session && session.TestSessionInfo) {
+            Object.keys(session.TestSessionInfo.Tests).forEach((key) => {
+                const testcase = this._testcases.get(key)!;
+                if(testcase.removeMeta(session.MetaHandler!)) {
+                    if (testcase.isEmpty()) {
+                        this._testcases.delete(key)
+                    }
+                }
+            });
+        }
         this.updateVersion();
     }
 
     clear() {
         this._collection.clear()
+        this._testcases.clear()
         this.updateVersion();
     }
 
@@ -675,7 +810,7 @@ export class TestSessionCollection {
                     // Max count resize
                     await this._datahandler.queueFetchItems(onFetch, SessionFields.initiate, maxCount);
             } else {
-                // Initiale fetch or reload
+                // Initial fetch or reload
                 await this._datahandler.setFromStream(
                     streamId, "Automated Test Session", onFetch, SessionFields.initiate, maxCount
                 );
@@ -717,8 +852,9 @@ export class TestSessionCollection {
         let hasUpdated = false;
         items.forEach((item) => {
             if (item) {
-                let data = item.getDataHandler() as TestSessionWrapper;
+                const data = item.getDataHandler() as TestSessionWrapper;
                 if(data && data.isAllPropertiesLoaded()) {
+                    this.indexTestsFromSession(data);
                     hasUpdated = true;
                 }
             }
@@ -726,76 +862,98 @@ export class TestSessionCollection {
         if(hasUpdated) { this.updateVersion() }
     }
 
-    async getTestResultByUID(uid: string): Promise<TestResult[]> {
-        const iterator = this._collection.values();
-        const tests: TestResult[] = [];
-        const toFetch: TestSessionWrapper[] = [];
-        let item = iterator.next();
-        while(!item.done) {
-            const session = item.value;
-            if (!session.isAllPropertiesLoaded()) {
-                await session.whenPropertyLoaded('TestSessionInfo');
+    private indexTestsFromSession(session: TestSessionWrapper) {
+        const tests = session.TestSessionInfo!.Tests;
+        Object.values(tests).some((item) => {
+            const testuid = item.TestUID;
+            let testcase = this._testcases.get(testuid);
+            if(!testcase) {
+                testcase = new TestCase(item.Name, testuid, item.Suite, session.Name);
+                this._testcases.set(testuid, testcase);
+            } else if (testcase.SessionName !== session.Name) {
+                return true; // Skip conficting session name, most likely session was renamed.
             }
-            const test = session.getTestResultByUID(uid);
-            if (test) {
-                if (test.Details) {
-                    tests.push(test);
-                } else {
-                    toFetch.push(session);
-                }
+            testcase.addMeta(item);
+            return false;
+        });
+    }
+
+    async getTestCaseByUID(uid: string): Promise<TestCase | undefined> {
+        let testcase: TestCase | undefined = undefined;
+        const toFetch: TestSessionWrapper[] = [];
+        const testcases = this._testcases;
+
+        function getTestCase() {
+            testcase = testcases.get(uid);
+            if (testcase) {
+                testcase.forEachMeta((test) => {
+                    if (test.Details === undefined) {
+                        toFetch.push(test.Session);
+                    }
+                });
             }
-            item = iterator.next();
         }
+        getTestCase();
 
-        if(tests.length === 0 && this.loadingCount !== 0) {
-            // Give a chance to the collection to be fully loaded if nothing was found
+        while(!testcase && toFetch.length === 0 && this.loadingCount !== 0) {
+            // Give a chance to the collection to load a bit more if nothing was found
             // It is possible the test is run infrequently and is deep in the result collection
-            await when(() => this.loadingCount === 0);
+            const previousVersion = this.version;
+            await when(() => this.version !== previousVersion || this.loadingCount === 0);
             // then retry
-            return await this.getTestResultByUID(uid);
+            getTestCase();
         }
 
         if(toFetch.length > 0) {
             // Fetch missing data from session test results
-            await this.fetchSessionResults(
+            this.fetchSessionResults(
                 toFetch.map((item) => {
                     return {SessionId: item.Testdata.id, ResultsId: item.TestSessionInfo.TestResultsTestDataUID}
                 }), 'data.' + uid
             );
-            toFetch.forEach((item) => {
-                const test = item.getTestResultByUID(uid);
-                if(test?.Details) {
-                    tests.push(test);
-                }
-            });
         }
 
-        return tests;
+        return testcase;
     }
 
-    async getTestResultByQuery(uid: string, metauid: string, change: number): Promise<TestResult | undefined> {
-        const latestTests = await this.getTestResultByUID(uid);
+    async getTestResultByQuery(uid: string, metauid: string, change?: number): Promise<TestResult | undefined> {
+        const testcase = await this.getTestCaseByUID(uid);
         let test: TestResult | undefined = undefined;
 
-        if(latestTests.length > 0) {
-            test = latestTests.find((test) => test.Session.Testdata.change === change && test.Session.MetaHandler?.hash === metauid);
-            if(!test) {
+        if(testcase) {
+            test = testcase.Metas.get(metauid);
+            while(!test && this.loadingCount !== 0) {
+                const previousVersion = testcase.version;
+                await when(() => testcase.version !== previousVersion || this.loadingCount === 0);
+                test = testcase.Metas.get(metauid);
+            }
+            if(test && change && test.Session.Testdata.change !== change) {
                 // Need to get the target change from the history
-                const sessionName = latestTests[0].Session.Name;
-                const changes = await this.getSessionHistory(sessionName);
-                const sessionsChange = changes.find((items) => items[0].Testdata.change === change);
-                const session = sessionsChange?.find((item) => item.MetaHandler?.hash === metauid);
-                test = session?.getTestResultByUID(uid);
+                test = this.getLoadedTestHistoryByMeta(uid)?.get(metauid)?.find(
+                    (item) => item.Session.Testdata.change === change
+                );
+                if(!test) {
+                    // trigger fetch history in background
+                    this.fetchSessionHistory(testcase.SessionName);
+                }
+                while(!test && this.loadingCount !== 0) {
+                    const previousVersion = this.version;
+                    await when(() => this.version !== previousVersion || this.loadingCount === 0);
+                    test = this.getLoadedTestHistoryByMeta(uid)?.get(metauid)?.find(
+                        (item) => item.Session.Testdata.change === change
+                    );
+                }
+
             }
         }
 
-        if(test && !test.Details) {
+        if(test && test.Details === undefined) {
             // Fetch missing data from the session test result
             await this.fetchSessionResults(
                 [{SessionId: test.Session.Testdata.id, ResultsId: test.Session.TestSessionInfo.TestResultsTestDataUID}]
                 , 'data.' + uid
             );
-            if(!test.Details) {
+            if(test.Details === undefined) {
                 // if nothing was fetched or failed to fetch, return nothing
                 return undefined;
             }
@@ -852,57 +1010,121 @@ export class TestSessionCollection {
         return tests;
     }
 
-    async getSessionHistory(sessionName: string, maxCount: number = 1000) {
+    getTestCases() {
+        return Array.from(this._testcases.values());
+    }
+
+    getLoadedTestHistoryByMeta(uid: string) {
+        const testcase = this._testcases.get(uid);
+        if(!testcase) {
+            return undefined;
+        }
+
+        const testresultsByMeta: Map<string, TestResult[]> = new Map();
+        this.forEachSessionHistoryByChange(testcase.SessionName,
+            (session) => {
+                if(session.isAllPropertiesLoaded()) {
+                    const test = session.getTestResultByUID(uid);
+                    if (test) {
+                        const metauid = session.MetaHandler!.hash;
+                        const meta = testresultsByMeta.get(metauid);
+                        if (meta) {
+                            const change = session.Testdata.change;
+                            if(!meta.some((item) => item.Session.Testdata.change === change)) {
+                                // Keep only one result per change
+                                meta.push(test);
+                            }
+                        } else {
+                            testresultsByMeta.set(metauid, [test]);
+                        }
+                    }
+                }
+            }
+        );
+
+        return testresultsByMeta;
+    }
+
+    async fetchSessionHistory(sessionName: string, maxCount: number = 1000) {
         if(this._datahandler.cursor) {
             // for single view mode we need to fetch the whole pool for context
-            const onFetch = (items: (TestDataWrapper|undefined)[]) => this.onInitialHistoryFetch(items)
-            await this._datahandler.getCursorHistory(onFetch, SessionFields.initiate, maxCount)
+            const onFetch = (items: (TestDataWrapper|undefined)[]) => this.onInitialHistoryFetch(items, sessionName)
+            await this._datahandler.fetchCursorHistory(onFetch, SessionFields.initiate, maxCount);
+            return;
+        }
+
+        await this.fetchIncompleteSessionsForHistory(sessionName);
+        while(this.loadingCount !== 0) {
+            // wait for the collection to be fully loaded, in the meanwhile load the missing session history
+            const previousVersion = this.version;
+            await when(() => this.version !== previousVersion || this.loadingCount === 0);
+            await this.fetchIncompleteSessionsForHistory(sessionName);
         }
+    }
 
+    private async fetchIncompleteSessionsForHistory(sessionName: string) {
         const missingSessions: string[] = [];
-        if(this.loadingCount !== 0) {
-            // wait for the current collection to have loaded
-            await when(() => this.loadingCount === 0);
+        this.forEachSessionHistoryByChange(sessionName,
+            (session) => {
+                if(!session.isAllPropertiesLoaded()) {
+                    missingSessions.push(session.Testdata.id);
+                }
+            }
+        );
+
+        if(missingSessions.length > 0) {
+            await this._datahandler.fetchUpdateItems(
+                missingSessions, (item) => this.onFetchTestCaseHistory(item), SessionFields.remaining
+            )
         }
+        return missingSessions.length;
+    }
 
-        const relatedSessionsByChange: TestSessionWrapper[][] = [];
+    private forEachSessionHistoryByChange(sessionName: string, callback: (session: TestSessionWrapper) => void) {
         const iterator = this._datahandler.iterItemsByChange();
         let change = iterator.next();
         while (!change.done) {
-            const sessions: TestSessionWrapper[] = [];
             change.value.forEach((item) => {
                 const session = item.getDataHandler() as TestSessionWrapper;
                 if(session.Name === sessionName) {
                     if (!this._includePreflight && session.PreFlightChange) {
                         return;
                     }
-                    sessions.push(session);
-                    if(!session.isAllPropertiesLoaded()) {
-                        missingSessions.push(session.Testdata.id);
-                    }
+                    callback(session);
                 }
             });
-            if (sessions.length > 0) {
-                relatedSessionsByChange.push(sessions);
-            }
             change = iterator.next();
         }
-
-        if(missingSessions.length > 0) {
-            await this._datahandler.fetchUpdateItems(
-                missingSessions, undefined, SessionFields.remaining
-            )
-        }
-        return relatedSessionsByChange;
     }
 
-    private onInitialHistoryFetch(items: (TestDataWrapper|undefined)[]) {
+    private onInitialHistoryFetch(items: (TestDataWrapper|undefined)[], sessionName: string) {
+        const toFetch: TestDataWrapper[] = [];
         items.forEach((item) => {
             if (item) {
                 let data = item.getDataHandler() as TestSessionWrapper;
                 if(!data) {
                     data = new TestSessionWrapper(item);
                     item.setDataHandler(data);
+                    if(data.Name === sessionName) {
+                        toFetch.push(item);
+                    }
+                }
+            }
+        });
+        if(toFetch.length > 0) {
+            this._datahandler.fetchUpdateItems(
+                toFetch.map((item) => item.id), (items) => this.onFetchTestCaseHistory(items), SessionFields.remaining
+            )
+        }
+    }
+
+    private onFetchTestCaseHistory(items: (TestDataWrapper|undefined)[]) {
+        items.forEach((item) => {
+            if(item) {
+                const session = item.getDataHandler() as TestSessionWrapper;
+                const tests = session.TestSessionInfo?.Tests;
+                if (tests) {
+                    Object.keys(tests).forEach((testuid) => this._testcases.get(testuid)?.updateHistory());
                 }
             }
         });
diff --git a/Engine/Source/Programs/Horde/HordeDashboard/src/backend/TestDataHandler.ts b/Engine/Source/Programs/Horde/HordeDashboard/src/backend/TestDataHandler.ts
--- a/Engine/Source/Programs/Horde/HordeDashboard/src/backend/TestDataHandler.ts
+++ b/Engine/Source/Programs/Horde/HordeDashboard/src/backend/TestDataHandler.ts
@@ -218,8 +218,8 @@ export class TestDataCollection {
     items?: Map<string, TestDataWrapper>;
     itemsByChange?: Map<number, TestDataWrapper[]>;
     cursor?: TestDataWrapper;
-    private nextPageIndex: number = 0;
-    private fetchCount: number = 100;
+    private nextIndex: number = 0;
+    private fetchCount: number = 20;
     private fetchMaxCount: number = 1000;
 
     desactivate() {
@@ -344,15 +344,30 @@ export class TestDataCollection {
             return [];
         }
 
-        if (!this.itemsByChange || this.itemsByChange.size === 1) {
-            await this.queueFetchItems(onFetch, filter, maxCount??this.fetchCount);
-        }
+        await this.fetchCursorHistory(onFetch, filter, maxCount, selector);
 
         const mapper = selector? (items: TestDataWrapper[]) => items.find(selector) as TestDataWrapper : (items: TestDataWrapper[]) => items[0];
 
         return Array.from(this.iterItemsByChange()).map(mapper);
     }
 
+    async fetchCursorHistory(onFetch?: OnFetchTestData, filter?: string, maxCount?: number, selector?: SelectorTestData) {
+        if (!this.cursor) {
+            return;
+        }
+
+        if (!this.itemsByChange || this.itemsByChange.size === 1) {
+            const key = `cursorHistory-${this.key}`;
+            if(this.isLoading(key)) {
+                await when(() => !this.loadingNotice.has(key));
+                return;
+            }
+            this.acquireLoadingNotice(key);
+            await this.queueFetchItems(onFetch, filter, maxCount??this.fetchCount);
+            this.releaseLoadingNotice(key);
+        }
+    }
+
     async queueFetchItems(onFetch?: OnFetchTestData, filter?: string, maxCount? : number) {
         if (maxCount) {
             this.fetchMaxCount = maxCount;
@@ -373,7 +388,7 @@ export class TestDataCollection {
             this.items = undefined;
             this.itemsByChange = undefined;
             this.cursor = undefined;
-            this.nextPageIndex = 0;
+            this.nextIndex = 0;
             this.fetchMaxCount = 1000; // the default
             this.releaseAllLoadingNotice();
         }
@@ -407,7 +422,8 @@ export class TestDataCollection {
         const streamId = this.streamId;
         let count = 0;
         try {
-            const pageData = await backend.getTestDataHistory(streamId, key, undefined, this.fetchCount, this.nextPageIndex++, filter);
+            const pageData = await backend.getTestDataHistory(streamId, key, undefined, this.fetchCount, this.nextIndex, filter);
+            this.nextIndex += this.fetchCount; 
 
             // Are we still fetching for the same collection?
             if(key === this.key && streamId === this.streamId) {
