Horde Storage - If ondemand replication fails from one region, try another

diff --git a/Engine/Source/Programs/Horde/HordeStorage/Horde.Storage/Implementation/Blob/IBlobService.cs b/Engine/Source/Programs/Horde/HordeStorage/Horde.Storage/Implementation/Blob/IBlobService.cs
--- a/Engine/Source/Programs/Horde/HordeStorage/Horde.Storage/Implementation/Blob/IBlobService.cs
+++ b/Engine/Source/Programs/Horde/HordeStorage/Horde.Storage/Implementation/Blob/IBlobService.cs
@@ -271,30 +271,39 @@ public class BlobService : IBlobService
             throw new BlobReplicationException(ns, blob, "Blob not found in any region");
 
         _logger.Information("On-demand replicating blob {Blob} in Namespace {Ns}", blob, ns);
-        SortedList<int, string> possiblePeers = _peerStatusService.GetPeersByLatency(blobInfo.Regions.ToList());
-
-        string bestPeer = possiblePeers.First().Value;
-        IPeerStatusService.PeerStatus? peerStatus = _peerStatusService.GetPeerStatus(bestPeer);
-        if (peerStatus == null)
-            throw new Exception($"Failed to find peer {bestPeer}");
-        PeerEndpoints peerEndpoint = peerStatus.Endpoints.First();
-        HttpClient httpClient = _httpClientFactory.CreateClient();
-        HttpRequestMessage blobRequest = BuildHttpRequest(HttpMethod.Get, $"{peerEndpoint.Url}/api/v1/blobs/{ns}/{blob}");
-        HttpResponseMessage blobResponse = await httpClient.SendAsync(blobRequest, HttpCompletionOption.ResponseHeadersRead);
-
-        if (blobResponse.StatusCode == HttpStatusCode.NotFound)
-        {
-            throw new BlobReplicationException(ns, blob, $"Failed to replicate {blob} in {ns} due to it not existing in region {bestPeer}.");
-        }
+        SortedList<int, string> possiblePeers = new SortedList<int, string>(_peerStatusService.GetPeersByLatency(blobInfo.Regions.ToList()));
 
-        if (blobResponse.StatusCode != HttpStatusCode.OK)
+        bool replicated = false;
+        foreach ((int latency, string? region) in possiblePeers)
         {
-            throw new BlobReplicationException(ns, blob, $"Failed to replicate {blob} in {ns} from region {bestPeer} due to bad http status code {blobResponse.StatusCode}.");
+            IPeerStatusService.PeerStatus? peerStatus = _peerStatusService.GetPeerStatus(region);
+            if (peerStatus == null)
+                throw new Exception($"Failed to find peer {region}");
+            PeerEndpoints peerEndpoint = peerStatus.Endpoints.First();
+            HttpClient httpClient = _httpClientFactory.CreateClient();
+            HttpRequestMessage blobRequest = BuildHttpRequest(HttpMethod.Get, $"{peerEndpoint.Url}/api/v1/blobs/{ns}/{blob}");
+            HttpResponseMessage blobResponse = await httpClient.SendAsync(blobRequest, HttpCompletionOption.ResponseHeadersRead);
+
+            if (blobResponse.StatusCode == HttpStatusCode.NotFound)
+            {
+                // try the next region
+                continue;
+            }
+
+            if (blobResponse.StatusCode != HttpStatusCode.OK)
+            {
+                throw new BlobReplicationException(ns, blob, $"Failed to replicate {blob} in {ns} from region {region} due to bad http status code {blobResponse.StatusCode}.");
+            }
+
+            await using Stream s = await blobResponse.Content.ReadAsStreamAsync();
+            using FilesystemBufferedPayload payload = await FilesystemBufferedPayload.Create(s);
+            await PutObject(ns, payload, blob);
+            replicated = true;
+            break;
         }
 
-        await using Stream s = await blobResponse.Content.ReadAsStreamAsync();
-        using FilesystemBufferedPayload payload = await FilesystemBufferedPayload.Create(s);
-        await PutObject(ns, payload, blob);
+        if (!replicated)
+            throw new BlobReplicationException(ns, blob, $"Failed to replicate {blob} in {ns} due to it not existing in any region");
 
         return await GetObject(ns, blob);
     }
