Horde Storage - Added log message to track the result of consistency check runs.

diff --git a/Engine/Source/Programs/Horde/HordeStorage/Horde.Storage/Implementation/Blob/BlobIndexConsistencyCheckService.cs b/Engine/Source/Programs/Horde/HordeStorage/Horde.Storage/Implementation/Blob/BlobIndexConsistencyCheckService.cs
--- a/Engine/Source/Programs/Horde/HordeStorage/Horde.Storage/Implementation/Blob/BlobIndexConsistencyCheckService.cs
+++ b/Engine/Source/Programs/Horde/HordeStorage/Horde.Storage/Implementation/Blob/BlobIndexConsistencyCheckService.cs
@@ -65,11 +65,16 @@ namespace Horde.Storage.Implementation
 
         private async Task RunConsistencyCheck()
         {
+            ulong countOfBlobsChecked = 0;
+            ulong countOfIncorrectBlobsFound = 0;
             string currentRegion = _jupiterSettings.CurrentValue.CurrentSite;
             await foreach (IBlobIndex.BlobInfo blobInfo in _blobIndex.GetAllBlobs())
             {
+                Interlocked.Increment(ref countOfBlobsChecked);
+
                 if (!blobInfo.Regions.Any())
                 {
+                    Interlocked.Increment(ref countOfIncorrectBlobsFound);
                     _logger.Warning("Blob {Blob} in namespace {Namespace} is not tracked to exist in any region", blobInfo.BlobIdentifier, blobInfo.Namespace);
 
                     if (await _blobService.ExistsInRootStore(blobInfo.Namespace, blobInfo.BlobIdentifier))
@@ -88,11 +93,12 @@ namespace Horde.Storage.Implementation
                         await _blobIndex.RemoveBlobFromIndex(blobInfo.Namespace, blobInfo.BlobIdentifier);
                     }
                 }
-
+                
                 if (blobInfo.Regions.Contains(currentRegion))
                 {
                     if (!await _blobService.ExistsInRootStore(blobInfo.Namespace, blobInfo.BlobIdentifier))
                     {
+                        Interlocked.Increment(ref countOfIncorrectBlobsFound);
                         _logger.Warning("Blob {Blob} in namespace {Namespace} did not exist in root store but is tracked as doing so in the blob index. Attempting to replicate it.", blobInfo.BlobIdentifier, blobInfo.Namespace);
                         try
                         {
@@ -110,7 +116,7 @@ namespace Horde.Storage.Implementation
                 }
             }
 
-            await Task.CompletedTask;
+            _logger.Information("Blob Index Consistency check finished, found {CountOfIncorrectBlobs} incorrect blobs. Processed {CountOfBlobs} blobs.", countOfIncorrectBlobsFound, countOfBlobsChecked);
         }
 
         protected override Task OnStopping(ConsistencyState state)
diff --git a/Engine/Source/Programs/Horde/HordeStorage/Horde.Storage/Implementation/Blob/BlobStoreConsistencyCheckService.cs b/Engine/Source/Programs/Horde/HordeStorage/Horde.Storage/Implementation/Blob/BlobStoreConsistencyCheckService.cs
--- a/Engine/Source/Programs/Horde/HordeStorage/Horde.Storage/Implementation/Blob/BlobStoreConsistencyCheckService.cs
+++ b/Engine/Source/Programs/Horde/HordeStorage/Horde.Storage/Implementation/Blob/BlobStoreConsistencyCheckService.cs
@@ -91,10 +91,12 @@ namespace Horde.Storage.Implementation
             List<NamespaceId> namespaces = await _refsStore.GetNamespaces().ToListAsync();
             namespaces.AddRange(await _referencesStore.GetNamespaces().ToListAsync());
 
-            long countOfBlobsChecked = 0;
             // technically this does not need to be run per namespace but per s3 bucket
             await foreach (NamespaceId ns in namespaces)
             {
+                ulong countOfBlobsChecked = 0;
+                ulong countOfIncorrectBlobsFound = 0;
+
                 using IScope scope = Tracer.Instance.StartActive("consistency_check.run");
                 scope.Span.ResourceName = ns.ToString();
 
@@ -112,12 +114,14 @@ namespace Horde.Storage.Implementation
                     {
                         _logger.Error("Mismatching hash for S3 object {Blob} in {Namespace}, new hash has {NewHash}. Deleting incorrect blob.", blob, ns, newHash);
 
+                        Interlocked.Increment(ref countOfIncorrectBlobsFound);
                         await s3Store.DeleteObject(ns, blob);
                         // update blob index tracking to indicate that we no longer have this blob in this region
                         await _blobIndex.RemoveBlobFromRegion(ns, blob);
                     }
                 }
-                
+
+                _logger.Information("Blob Store Consistency check finished for {Namespace}, found {CountOfIncorrectBlobs} incorrect blobs. Processed {CountOfBlobs} blobs.", ns, countOfIncorrectBlobsFound, countOfBlobsChecked);
             }
         }
 
