Horde Storage - Added option to control if blob index consistency checks deletes inconsistent content, disabled by default to simply run in report mode initially.

diff --git a/Engine/Source/Programs/Horde/HordeStorage/Horde.Storage/Implementation/Blob/BlobIndexConsistencyCheckService.cs b/Engine/Source/Programs/Horde/HordeStorage/Horde.Storage/Implementation/Blob/BlobIndexConsistencyCheckService.cs
--- a/Engine/Source/Programs/Horde/HordeStorage/Horde.Storage/Implementation/Blob/BlobIndexConsistencyCheckService.cs
+++ b/Engine/Source/Programs/Horde/HordeStorage/Horde.Storage/Implementation/Blob/BlobIndexConsistencyCheckService.cs
@@ -90,7 +90,8 @@ namespace Horde.Storage.Implementation
                         // this blob doesn't exist anywhere so we just cleanup the blob index
                         _logger.Warning("Blob {Blob} in namespace {Namespace} was removed from the blob index as it didnt exist anywhere", blobInfo.BlobIdentifier, blobInfo.Namespace);
 
-                        await _blobIndex.RemoveBlobFromIndex(blobInfo.Namespace, blobInfo.BlobIdentifier);
+                        if (_settings.CurrentValue.AllowDeletesInBlobIndex)
+                            await _blobIndex.RemoveBlobFromIndex(blobInfo.Namespace, blobInfo.BlobIdentifier);
                     }
                 }
                 
@@ -109,7 +110,8 @@ namespace Horde.Storage.Implementation
                             _logger.Error(e, "Failed to replicate Blob {Blob} in namespace {Namespace}. Unable to repair the blob index", blobInfo.BlobIdentifier, blobInfo.Namespace);
 
                             // we update the blob index to accurately reflect that we do not have the blob, this is not good though as it means a upload that we thought happened now lacks content
-                            await _blobIndex.RemoveBlobFromRegion(blobInfo.Namespace, blobInfo.BlobIdentifier);
+                            if (_settings.CurrentValue.AllowDeletesInBlobIndex)
+                                await _blobIndex.RemoveBlobFromRegion(blobInfo.Namespace, blobInfo.BlobIdentifier);
                         }
                     }
 
diff --git a/Engine/Source/Programs/Horde/HordeStorage/Horde.Storage/Settings.cs b/Engine/Source/Programs/Horde/HordeStorage/Horde.Storage/Settings.cs
--- a/Engine/Source/Programs/Horde/HordeStorage/Horde.Storage/Settings.cs
+++ b/Engine/Source/Programs/Horde/HordeStorage/Horde.Storage/Settings.cs
@@ -291,5 +291,6 @@ namespace Horde.Storage
         public bool EnableBlobStoreChecks { get; set; } = false;
         public bool EnableBlobIndexChecks { get; set; } = false;
         public double ConsistencyCheckPollFrequencySeconds { get; set; } = TimeSpan.FromHours(2).TotalSeconds;
+        public bool AllowDeletesInBlobIndex { get; set; } = false;
     }
 }
