Horde: Modify agent code to conform to NET coding style.

diff --git a/Engine/Source/Programs/Horde/Horde.Agent.Tests/Helper.cs b/Engine/Source/Programs/Horde/Horde.Agent.Tests/Helper.cs
--- a/Engine/Source/Programs/Horde/Horde.Agent.Tests/Helper.cs
+++ b/Engine/Source/Programs/Horde/Horde.Agent.Tests/Helper.cs
@@ -18,30 +18,15 @@ namespace Horde.Agent.Tests
 {
 	// Stub for fulfilling IOptionsMonitor interface during testing
 	// Copied from HordeServerTests until a good way to share code between these is decided.
-	public class TestOptionsMonitor<T> : IOptionsMonitor<T>, IDisposable
+	public class TestOptionsMonitor<T> : IOptions<T>
 		where T : class, new()
 	{
-		public TestOptionsMonitor(T CurrentValue)
+		public TestOptionsMonitor(T value)
 		{
-			this.CurrentValue = CurrentValue;
+			this.Value = value;
 		}
 
-		public T Get(string Name)
-		{
-			return CurrentValue;
-		}
-
-		public IDisposable OnChange(Action<T, string> listener)
-		{
-			return this;
-		}
-
-		public T CurrentValue { get; }
-
-		public void Dispose()
-		{
-			// Dummy stub to satisfy return value of OnChange 
-		}
+		public T Value { get; }
 	}
 
 	class RpcClientRefStub : IRpcClientRef
diff --git a/Engine/Source/Programs/Horde/Horde.Agent.Tests/WorkerServiceTest.cs b/Engine/Source/Programs/Horde/Horde.Agent.Tests/WorkerServiceTest.cs
--- a/Engine/Source/Programs/Horde/Horde.Agent.Tests/WorkerServiceTest.cs
+++ b/Engine/Source/Programs/Horde/Horde.Agent.Tests/WorkerServiceTest.cs
@@ -14,6 +14,7 @@ using HordeCommon.Rpc;
 using HordeCommon.Rpc.Messages;
 using HordeCommon.Rpc.Tasks;
 using Microsoft.Extensions.Logging;
+using Microsoft.Extensions.Logging.Abstractions;
 using Microsoft.Extensions.Logging.Console;
 using Microsoft.Extensions.Options;
 using Microsoft.VisualStudio.TestTools.UnitTesting;
@@ -27,12 +28,7 @@ namespace Horde.Agent.Tests
 
 		public WorkerServiceTest()
 		{
-			LoggerFactory LoggerFactory = new LoggerFactory();
-			ConsoleLoggerOptions LoggerOptions = new ConsoleLoggerOptions();
-			TestOptionsMonitor<ConsoleLoggerOptions> LoggerOptionsMon =
-				new TestOptionsMonitor<ConsoleLoggerOptions>(LoggerOptions);
-			LoggerFactory.AddProvider(new ConsoleLoggerProvider(LoggerOptionsMon));
-			WorkerLogger = LoggerFactory.CreateLogger<WorkerService>();
+			WorkerLogger = NullLogger<WorkerService>.Instance;
 		}
 
 		private WorkerService GetWorkerService(
@@ -48,7 +44,7 @@ namespace Horde.Agent.Tests
 			Settings.Server = "test";
 			Settings.WorkingDir = Path.GetTempPath();
 			Settings.Executor = ExecutorType.Test; // Not really used since the executor is overridden in the tests
-			IOptionsMonitor<AgentSettings> SettingsMonitor = new TestOptionsMonitor<AgentSettings>(Settings);
+			IOptions<AgentSettings> SettingsMonitor = new TestOptionsMonitor<AgentSettings>(Settings);
 
 			return new WorkerService(WorkerLogger, SettingsMonitor, null!, null!, CreateExecutor);
 		}
@@ -119,7 +115,7 @@ namespace Horde.Agent.Tests
 			});
 
 			WorkerService Ws = GetWorkerService((a, b, c) => Executor);
-			Ws.StepAbortPollInterval = TimeSpan.FromMilliseconds(1);
+			Ws._stepAbortPollInterval = TimeSpan.FromMilliseconds(1);
 			LeaseOutcome Outcome = (await Ws.ExecuteJobAsync(RpcConnection, "agentId1", "leaseId1", ExecuteJobTask,
 				WorkerLogger, Token)).Outcome;
 
@@ -147,7 +143,7 @@ namespace Horde.Agent.Tests
 			});
 
 			WorkerService Ws = GetWorkerService((a, b, c) => Executor);
-			Ws.StepAbortPollInterval = TimeSpan.FromMilliseconds(5);
+			Ws._stepAbortPollInterval = TimeSpan.FromMilliseconds(5);
 			
 			HordeRpcClientStub Client = new HordeRpcClientStub(WorkerLogger);
 			RpcConnectionStub RpcConnection = new RpcConnectionStub(null!, Client);
diff --git a/Engine/Source/Programs/Horde/Horde.Agent/AgentSettings.cs b/Engine/Source/Programs/Horde/Horde.Agent/AgentSettings.cs
--- a/Engine/Source/Programs/Horde/Horde.Agent/AgentSettings.cs
+++ b/Engine/Source/Programs/Horde/Horde.Agent/AgentSettings.cs
@@ -101,15 +101,15 @@ namespace Horde.Agent
 		/// <summary>
 		/// Checks whether the given certificate thumbprint should be trusted
 		/// </summary>
-		/// <param name="CertificateThumbprint">The cert thumbprint</param>
+		/// <param name="certificateThumbprint">The cert thumbprint</param>
 		/// <returns>True if the cert should be trusted</returns>
-		public bool IsTrustedCertificate(string CertificateThumbprint)
+		public bool IsTrustedCertificate(string certificateThumbprint)
 		{
-			if (Thumbprint != null && Thumbprint.Equals(CertificateThumbprint, StringComparison.OrdinalIgnoreCase))
+			if (Thumbprint != null && Thumbprint.Equals(certificateThumbprint, StringComparison.OrdinalIgnoreCase))
 			{
 				return true;
 			}
-			if (Thumbprints.Any(x => x.Equals(CertificateThumbprint, StringComparison.OrdinalIgnoreCase)))
+			if (Thumbprints.Any(x => x.Equals(certificateThumbprint, StringComparison.OrdinalIgnoreCase)))
 			{
 				return true;
 			}
@@ -214,10 +214,10 @@ namespace Horde.Agent
 		/// Gets the current server settings
 		/// </summary>
 		/// <returns>The current server settings</returns>
-		public ServerProfile GetServerProfile(string Name)
+		public ServerProfile GetServerProfile(string name)
 		{
-			ServerProfile? ServerProfile = ServerProfiles.FirstOrDefault(x => x.Name.Equals(Name, StringComparison.OrdinalIgnoreCase));
-			if (ServerProfile == null)
+			ServerProfile? serverProfile = ServerProfiles.FirstOrDefault(x => x.Name.Equals(name, StringComparison.OrdinalIgnoreCase));
+			if (serverProfile == null)
 			{
 				if (ServerProfiles.Count == 0)
 				{
@@ -225,10 +225,10 @@ namespace Horde.Agent
 				}
 				else
 				{
-					throw new Exception($"Unknown server profile name '{Name}' (valid profiles: {String.Join("/", ServerProfiles.Select(x => x.Name))})");
+					throw new Exception($"Unknown server profile name '{name}' (valid profiles: {String.Join("/", ServerProfiles.Select(x => x.Name))})");
 				}
 			}
-			return ServerProfile;
+			return serverProfile;
 		}
 
 		/// <summary>
@@ -259,12 +259,12 @@ namespace Horde.Agent
 		/// <summary>
 		/// Gets the configuration section for the active server profile
 		/// </summary>
-		/// <param name="ConfigSection"></param>
+		/// <param name="configSection"></param>
 		/// <returns></returns>
-		public static IConfigurationSection GetCurrentServerProfile(this IConfigurationSection ConfigSection)
+		public static IConfigurationSection GetCurrentServerProfile(this IConfigurationSection configSection)
 		{
-			string ProfileName = ConfigSection[nameof(AgentSettings.Server)];
-			return ConfigSection.GetSection(nameof(AgentSettings.ServerProfiles)).GetChildren().First(x => x[nameof(ServerProfile.Name)] == ProfileName);
+			string profileName = configSection[nameof(AgentSettings.Server)];
+			return configSection.GetSection(nameof(AgentSettings.ServerProfiles)).GetChildren().First(x => x[nameof(ServerProfile.Name)] == profileName);
 		}
 	}
 }
diff --git a/Engine/Source/Programs/Horde/Horde.Agent/Commands/Bundles/CreateCommand.cs b/Engine/Source/Programs/Horde/Horde.Agent/Commands/Bundles/CreateCommand.cs
--- a/Engine/Source/Programs/Horde/Horde.Agent/Commands/Bundles/CreateCommand.cs
+++ b/Engine/Source/Programs/Horde/Horde.Agent/Commands/Bundles/CreateCommand.cs
@@ -21,33 +21,33 @@ namespace Horde.Agent.Commands.Bundles
 	abstract class BundleCommandBase : Command
 	{
 		[CommandLine("-Server=")]
-		public string? Server = null;
+		public string? Server { get; set; } = null;
 
 		[CommandLine("-StorageDir=", Description = "Overrides the default storage server with a local directory")]
-		public DirectoryReference? StorageDir = null;
+		public DirectoryReference? StorageDir { get; set; } = null;
 
-		protected IStorageClient CreateStorageClient(ILogger Logger)
+		protected IStorageClient CreateStorageClient(ILogger logger)
 		{
 			if (StorageDir != null)
 			{
-				return new FileStorageClient(StorageDir, Logger);
+				return new FileStorageClient(StorageDir, logger);
 			}
 			else
 			{
-				IConfiguration Config = new ConfigurationBuilder()
+				IConfiguration config = new ConfigurationBuilder()
 					.AddJsonFile("appsettings.json", optional: false)
 					.AddJsonFile($"appsettings.{Environment.GetEnvironmentVariable("DOTNET_ENVIRONMENT")}.json", optional: true) // environment variable overrides, also used in k8s setups with Helm
 					.AddJsonFile("appsettings.User.json", optional: true)
 					.AddEnvironmentVariables()
 					.Build();
 
-				IConfigurationSection ConfigSection = Config.GetSection(AgentSettings.SectionName);
+				IConfigurationSection configSection = config.GetSection(AgentSettings.SectionName);
 
-				IServiceCollection Services = new ServiceCollection();
-				Services.AddHordeStorage(Settings => ConfigSection.GetCurrentServerProfile().GetSection(nameof(ServerProfile.Storage)).Bind(Settings));
+				IServiceCollection services = new ServiceCollection();
+				services.AddHordeStorage(settings => configSection.GetCurrentServerProfile().GetSection(nameof(ServerProfile.Storage)).Bind(settings));
 
-				IServiceProvider ServiceProvider = Services.BuildServiceProvider();
-				return ServiceProvider.GetRequiredService<IStorageClient>();
+				IServiceProvider serviceProvider = services.BuildServiceProvider();
+				return serviceProvider.GetRequiredService<IStorageClient>();
 			}
 		}
 	}
@@ -67,16 +67,14 @@ namespace Horde.Agent.Commands.Bundles
 		[CommandLine("-InputDir=", Required = true)]
 		public DirectoryReference InputDir { get; set; } = null!;
 
-		public override async Task<int> ExecuteAsync(ILogger Logger)
+		public override async Task<int> ExecuteAsync(ILogger logger)
 		{
-			IStorageClient StorageClient = base.CreateStorageClient(Logger);
-			using (MemoryCache Cache = new MemoryCache(new MemoryCacheOptions { SizeLimit = 50 * 1024 * 1024 }))
+			IStorageClient storageClient = base.CreateStorageClient(logger);
+			using (MemoryCache cache = new MemoryCache(new MemoryCacheOptions { SizeLimit = 50 * 1024 * 1024 }))
 			{
-				using (Bundle<DirectoryNode> NewBundle = Bundle.Create<DirectoryNode>(StorageClient, NamespaceId, new BundleOptions(), Cache))
-				{
-					await NewBundle.Root.CopyFromDirectoryAsync(InputDir.ToDirectoryInfo(), new ChunkingOptions(), Logger);
-					await NewBundle.WriteAsync(BucketId, RefId, CbObject.Empty, false);
-				}
+				using Bundle<DirectoryNode> bundle = Bundle.Create<DirectoryNode>(storageClient, NamespaceId, new BundleOptions(), cache);
+				await bundle.Root.CopyFromDirectoryAsync(InputDir.ToDirectoryInfo(), new ChunkingOptions(), logger);
+				await bundle.WriteAsync(BucketId, RefId, CbObject.Empty, false);
 			}
 
 			return 0;
diff --git a/Engine/Source/Programs/Horde/Horde.Agent/Commands/Bundles/DeltaCommand.cs b/Engine/Source/Programs/Horde/Horde.Agent/Commands/Bundles/DeltaCommand.cs
--- a/Engine/Source/Programs/Horde/Horde.Agent/Commands/Bundles/DeltaCommand.cs
+++ b/Engine/Source/Programs/Horde/Horde.Agent/Commands/Bundles/DeltaCommand.cs
@@ -3,22 +3,13 @@
 using EpicGames.Core;
 using EpicGames.Horde.Bundles;
 using EpicGames.Horde.Storage;
-using EpicGames.Horde.Storage.Impl;
-using EpicGames.Serialization;
-using Microsoft.Extensions.Caching.Memory;
-using Microsoft.Extensions.Configuration;
 using Microsoft.Extensions.Logging;
-using System;
-using System.Collections.Generic;
-using System.Diagnostics;
-using System.Linq;
-using System.Text;
 using System.Threading.Tasks;
 
 namespace Horde.Agent.Commands.Bundles
 {
 	[Command("bundle", "delta", "Desribe the delta between two bundles")]
-	class DeltaCommand : BundleCommandBase
+	internal class DeltaCommand : BundleCommandBase
 	{
 		[CommandLine("-Namespace=")]
 		public NamespaceId NamespaceId { get; set; } = new NamespaceId("default-ns");
@@ -32,10 +23,10 @@ namespace Horde.Agent.Commands.Bundles
 		[CommandLine("-Ref2=", Required = true)]
 		public RefId RefId2 { get; set; }
 
-		public override async Task<int> ExecuteAsync(ILogger Logger)
+		public override async Task<int> ExecuteAsync(ILogger logger)
 		{
-			IStorageClient StorageClient = CreateStorageClient(Logger);
-			await BundleTools.FindDeltaAsync(StorageClient, NamespaceId, BucketId, RefId1, RefId2, Logger);
+			IStorageClient storageClient = CreateStorageClient(logger);
+			await BundleTools.FindDeltaAsync(storageClient, NamespaceId, BucketId, RefId1, RefId2, logger);
 			return 0;
 		}
 	}
diff --git a/Engine/Source/Programs/Horde/Horde.Agent/Commands/Bundles/DumpCommand.cs b/Engine/Source/Programs/Horde/Horde.Agent/Commands/Bundles/DumpCommand.cs
--- a/Engine/Source/Programs/Horde/Horde.Agent/Commands/Bundles/DumpCommand.cs
+++ b/Engine/Source/Programs/Horde/Horde.Agent/Commands/Bundles/DumpCommand.cs
@@ -3,22 +3,13 @@
 using EpicGames.Core;
 using EpicGames.Horde.Bundles;
 using EpicGames.Horde.Storage;
-using EpicGames.Horde.Storage.Impl;
-using EpicGames.Serialization;
-using Microsoft.Extensions.Caching.Memory;
-using Microsoft.Extensions.Configuration;
-using Microsoft.Extensions.DependencyInjection;
 using Microsoft.Extensions.Logging;
-using System;
-using System.Collections.Generic;
-using System.IO;
-using System.Linq;
 using System.Threading.Tasks;
 
 namespace Horde.Agent.Commands.Bundles
 {
 	[Command("bundle", "dump", "Indexes a tree and creates a table of contents for each blob describing the contents")]
-	class DumpCommand : BundleCommandBase
+	internal class DumpCommand : BundleCommandBase
 	{
 		[CommandLine("-Namespace=")]
 		public NamespaceId NamespaceId { get; set; } = new NamespaceId("default-ns");
@@ -32,12 +23,12 @@ namespace Horde.Agent.Commands.Bundles
 		[CommandLine("-OutputDir=")]
 		public DirectoryReference? OutputDir { get; set; }
 
-		public override async Task<int> ExecuteAsync(ILogger Logger)
+		public override async Task<int> ExecuteAsync(ILogger logger)
 		{
 			OutputDir = DirectoryReference.Combine(Program.DataDir, "Bundles");
 
-			IStorageClient StorageClient = CreateStorageClient(Logger);
-			await BundleTools.WriteSummaryAsync(StorageClient, NamespaceId, BucketId, RefId, OutputDir);
+			IStorageClient storageClient = CreateStorageClient(logger);
+			await BundleTools.WriteSummaryAsync(storageClient, NamespaceId, BucketId, RefId, OutputDir);
 
 			return 0;
 		}
diff --git a/Engine/Source/Programs/Horde/Horde.Agent/Commands/Bundles/ExtractCommand.cs b/Engine/Source/Programs/Horde/Horde.Agent/Commands/Bundles/ExtractCommand.cs
--- a/Engine/Source/Programs/Horde/Horde.Agent/Commands/Bundles/ExtractCommand.cs
+++ b/Engine/Source/Programs/Horde/Horde.Agent/Commands/Bundles/ExtractCommand.cs
@@ -4,22 +4,14 @@ using EpicGames.Core;
 using EpicGames.Horde.Bundles;
 using EpicGames.Horde.Bundles.Nodes;
 using EpicGames.Horde.Storage;
-using EpicGames.Horde.Storage.Impl;
-using EpicGames.Serialization;
 using Microsoft.Extensions.Caching.Memory;
-using Microsoft.Extensions.Configuration;
-using Microsoft.Extensions.DependencyInjection;
 using Microsoft.Extensions.Logging;
-using System;
-using System.Collections.Generic;
-using System.IO;
-using System.Linq;
 using System.Threading.Tasks;
 
 namespace Horde.Agent.Commands.Bundles
 {
 	[Command("bundle", "extract", "Extracts data from a bundle to the local hard drive")]
-	class ExtractCommand : BundleCommandBase
+	internal class ExtractCommand : BundleCommandBase
 	{
 		[CommandLine("-Namespace=")]
 		public NamespaceId NamespaceId { get; set; } = new NamespaceId("default-ns");
@@ -33,13 +25,13 @@ namespace Horde.Agent.Commands.Bundles
 		[CommandLine("-OutputDir=", Required = true)]
 		public DirectoryReference OutputDir { get; set; } = null!;
 
-		public override async Task<int> ExecuteAsync(ILogger Logger)
+		public override async Task<int> ExecuteAsync(ILogger logger)
 		{
-			IStorageClient StorageClient = CreateStorageClient(Logger);
-			using (IMemoryCache Cache = new MemoryCache(new MemoryCacheOptions { SizeLimit = 1024 * 1024 * 1000 }))
+			IStorageClient storageClient = CreateStorageClient(logger);
+			using (IMemoryCache cache = new MemoryCache(new MemoryCacheOptions { SizeLimit = 1024 * 1024 * 1000 }))
 			{
-				using Bundle<DirectoryNode> NewBundle = await Bundle.ReadAsync<DirectoryNode>(StorageClient, NamespaceId, BucketId, RefId, new BundleOptions(), Cache);
-				await NewBundle.Root.CopyToDirectoryAsync(NewBundle, OutputDir.ToDirectoryInfo(), Logger);
+				using Bundle<DirectoryNode> bundle = await Bundle.ReadAsync<DirectoryNode>(storageClient, NamespaceId, BucketId, RefId, new BundleOptions(), cache);
+				await bundle.Root.CopyToDirectoryAsync(bundle, OutputDir.ToDirectoryInfo(), logger);
 			}
 			return 0;
 		}
diff --git a/Engine/Source/Programs/Horde/Horde.Agent/Commands/Compute/ComputeCommand.cs b/Engine/Source/Programs/Horde/Horde.Agent/Commands/Compute/ComputeCommand.cs
--- a/Engine/Source/Programs/Horde/Horde.Agent/Commands/Compute/ComputeCommand.cs
+++ b/Engine/Source/Programs/Horde/Horde.Agent/Commands/Compute/ComputeCommand.cs
@@ -50,273 +50,272 @@ namespace Horde.Agent.Commands
 		/// Input file describing the work to execute
 		/// </summary>
 		[CommandLine(Required = true)]
-		public FileReference Input = null!;
+		public FileReference Input { get; set; } = null!;
 
 		/// <summary>
 		/// The input directory. By default, the directory containing the input file will be used.
 		/// </summary>
 		[CommandLine]
-		public DirectoryReference? InputDir = null;
+		public DirectoryReference? InputDir { get; set; } = null;
 
 		/// <summary>
 		/// Apply a random salt to the cached value
 		/// </summary>
 		[CommandLine("-Salt")]
-		public bool RandomSalt;
+		public bool RandomSalt { get; set; }
 
 		/// <summary>
 		/// Add a known salt value to the cache
 		/// </summary>
 		[CommandLine("-Salt=")]
-		public string? Salt = null;
+		public string? Salt { get; set; } = null;
 						
 		/// <summary>
 		/// Skip checking if a result is already available in the action cache
 		/// </summary>
 		[CommandLine("-SkipCacheLookup")]
-		public bool SkipCacheLookup = false;
+		public bool SkipCacheLookup { get; set; } = false;
 		
 		/// <summary>
 		/// Directory to download the output files to. If not set, no results will be downloaded.
 		/// </summary>
 		[CommandLine("-OutputDir")]
-		public string? OutputDir = null;
+		public string? OutputDir { get; set; } = null;
 	
 		/// <summary>
 		/// Log verbosity level (use normal Serilog levels such as debug, warning or info)
 		/// </summary>
 		[CommandLine("-LogLevel")]
-		public string LogLevelStr = "debug";
+		public string LogLevelStr { get; set; } = "debug";
 
-		static (DirectoryTree, IoHash) CreateSandbox(DirectoryInfo BaseDirInfo, Dictionary<IoHash, byte[]> UploadList)
+		static (DirectoryTree, IoHash) CreateSandbox(DirectoryInfo baseDirInfo, Dictionary<IoHash, byte[]> uploadList)
 		{
-			DirectoryTree Tree = new DirectoryTree();
+			DirectoryTree tree = new DirectoryTree();
 
-			foreach (DirectoryInfo SubDirInfo in BaseDirInfo.EnumerateDirectories())
+			foreach (DirectoryInfo subDirInfo in baseDirInfo.EnumerateDirectories())
 			{
-				(DirectoryTree SubTree, IoHash SubDirHash) = CreateSandbox(SubDirInfo, UploadList);
-				Tree.Directories.Add(new DirectoryNode(SubDirInfo.Name, SubDirHash));
+				(DirectoryTree subTree, IoHash subDirHash) = CreateSandbox(subDirInfo, uploadList);
+				tree.Directories.Add(new DirectoryNode(subDirInfo.Name, subDirHash));
 			}
-			Tree.Directories.SortBy(x => x.Name, Utf8StringComparer.Ordinal);
+			tree.Directories.SortBy(x => x.Name, Utf8StringComparer.Ordinal);
 
-			foreach (FileInfo FileInfo in BaseDirInfo.EnumerateFiles())
+			foreach (FileInfo fileInfo in baseDirInfo.EnumerateFiles())
 			{
-				byte[] Data = File.ReadAllBytes(FileInfo.FullName);
-				IoHash Hash = IoHash.Compute(Data);
-				UploadList[Hash] = Data;
-				Tree.Files.Add(new FileNode(FileInfo.Name, Hash, FileInfo.Length, (int)FileInfo.Attributes));
+				byte[] data = File.ReadAllBytes(fileInfo.FullName);
+				IoHash hash = IoHash.Compute(data);
+				uploadList[hash] = data;
+				tree.Files.Add(new FileNode(fileInfo.Name, hash, fileInfo.Length, (int)fileInfo.Attributes));
 			}
-			Tree.Files.SortBy(x => x.Name, Utf8StringComparer.Ordinal);
+			tree.Files.SortBy(x => x.Name, Utf8StringComparer.Ordinal);
 
-			return (Tree, AddCbObject(UploadList, Tree));
+			return (tree, AddCbObject(uploadList, tree));
 		}
 
 		/// <inheritdoc/>
-		public override async Task<int> ExecuteAsync(ILogger Logger)
+		public override async Task<int> ExecuteAsync(ILogger logger)
 		{
 			InputDir ??= Input.Directory;
 
-			if (Enum.TryParse(LogLevelStr, true, out LogEventLevel LogEventLevel))
+			if (Enum.TryParse(LogLevelStr, true, out LogEventLevel logEventLevel))
 			{
-				Logging.LogLevelSwitch.MinimumLevel = LogEventLevel;
+				Logging.LogLevelSwitch.MinimumLevel = logEventLevel;
 			}
 			else
 			{
-				Console.WriteLine($"Unable to parse log level: {this.LogLevelStr}");
+				Console.WriteLine($"Unable to parse log level: {LogLevelStr}");
 				return 0;
 			}
 
-			IConfiguration Configuration = new ConfigurationBuilder()
+			IConfiguration configuration = new ConfigurationBuilder()
 				.AddEnvironmentVariables()
 				.AddJsonFile(Input.FullName)
 				.Build();
 
-			IHostBuilder HostBuilder = Host.CreateDefaultBuilder()
-				.ConfigureLogging(Builder =>
+			IHostBuilder hostBuilder = Host.CreateDefaultBuilder()
+				.ConfigureLogging(builder =>
 				{
-					Builder.SetMinimumLevel(LogLevel.Warning);
+					builder.SetMinimumLevel(LogLevel.Warning);
 				})
-				.ConfigureServices(Services =>
+				.ConfigureServices(services =>
 				{
-					Services.AddLogging();
+					services.AddLogging();
 
-					IConfigurationSection ComputeSettings = Configuration.GetSection(nameof(JsonComputeTask.ComputeServer));
-					Services.AddHordeCompute(Settings => ComputeSettings.Bind(Settings));
+					IConfigurationSection computeSettings = configuration.GetSection(nameof(JsonComputeTask.ComputeServer));
+					services.AddHordeCompute(settings => computeSettings.Bind(settings));
 
-					IConfigurationSection StorageSettings = Configuration.GetSection(nameof(JsonComputeTask.StorageServer));
-					Services.AddHordeStorage(Settings => StorageSettings.Bind(Settings));
+					IConfigurationSection storageSettings = configuration.GetSection(nameof(JsonComputeTask.StorageServer));
+					services.AddHordeStorage(settings => storageSettings.Bind(settings));
 				});
 
-			using (IHost Host = HostBuilder.Build())
+			using (IHost host = hostBuilder.Build())
 			{
-				IStorageClient StorageClient = Host.Services.GetRequiredService<IStorageClient>();
-				IComputeClient ComputeClient = Host.Services.GetRequiredService<IComputeClient>();
+				IStorageClient storageClient = host.Services.GetRequiredService<IStorageClient>();
+				IComputeClient computeClient = host.Services.GetRequiredService<IComputeClient>();
 
-				ByteString SaltBytes = ByteString.Empty;
+				ByteString saltBytes = ByteString.Empty;
 				if (Salt != null)
 				{
-					SaltBytes = ByteString.CopyFromUtf8(Salt);
+					saltBytes = ByteString.CopyFromUtf8(Salt);
 				}
 				else if(RandomSalt)
 				{
-					SaltBytes = ByteString.CopyFrom(Guid.NewGuid().ToByteArray());
+					saltBytes = ByteString.CopyFrom(Guid.NewGuid().ToByteArray());
 				}
 
-				if (SaltBytes.Length > 0)
+				if (saltBytes.Length > 0)
 				{
-					Logger.LogInformation("Using salt: {SaltBytes}", StringUtils.FormatHexString(SaltBytes.ToByteArray()));
+					logger.LogInformation("Using salt: {SaltBytes}", StringUtils.FormatHexString(saltBytes.ToByteArray()));
 				}
 
-				Dictionary<IoHash, byte[]> Blobs = new Dictionary<IoHash, byte[]>();
-				(_, IoHash SandboxHash) = CreateSandbox(InputDir.ToDirectoryInfo(), Blobs);
+				Dictionary<IoHash, byte[]> blobs = new Dictionary<IoHash, byte[]>();
+				(_, IoHash sandboxHash) = CreateSandbox(InputDir.ToDirectoryInfo(), blobs);
 
-				JsonComputeTask JsonComputeTask = new JsonComputeTask();
-				Configuration.Bind(JsonComputeTask);
+				JsonComputeTask jsonComputeTask = new JsonComputeTask();
+				configuration.Bind(jsonComputeTask);
 
-				Logger.LogInformation("compute server: {ServerUrl}", JsonComputeTask.ComputeServer.Url);
-				Logger.LogInformation("storage server: {ServerUrl}", JsonComputeTask.StorageServer.Url);
+				logger.LogInformation("compute server: {ServerUrl}", jsonComputeTask.ComputeServer.Url);
+				logger.LogInformation("storage server: {ServerUrl}", jsonComputeTask.StorageServer.Url);
 
-				JsonRequirements JsonRequirements = JsonComputeTask.Requirements;
-				Requirements Requirements = new Requirements(JsonRequirements.Condition ?? String.Empty);
-				Requirements.Resources = JsonRequirements.Resources.ToDictionary(x => x.Key, x => x.Value);
-				Requirements.Exclusive = JsonRequirements.Exclusive;
-				IoHash RequirementsHash = AddCbObject(Blobs, Requirements);
+				JsonRequirements jsonRequirements = jsonComputeTask.Requirements;
+				Requirements requirements = new Requirements(jsonRequirements.Condition ?? String.Empty);
+				requirements.Resources = jsonRequirements.Resources.ToDictionary(x => x.Key, x => x.Value);
+				requirements.Exclusive = jsonRequirements.Exclusive;
+				IoHash requirementsHash = AddCbObject(blobs, requirements);
 
-				ComputeTask Task = new ComputeTask(JsonComputeTask.Executable, JsonComputeTask.Arguments.ConvertAll<Utf8String>(x => x), JsonComputeTask.WorkingDirectory, SandboxHash);
-				Task.EnvVars = JsonComputeTask.EnvVars.ToDictionary(x => (Utf8String)x.Key, x => (Utf8String)x.Value);
-				Task.OutputPaths.AddRange(JsonComputeTask.OutputPaths.Select(x => (Utf8String)x));
-				Task.RequirementsHash = AddCbObject(Blobs, Requirements);
+				ComputeTask task = new ComputeTask(jsonComputeTask.Executable, jsonComputeTask.Arguments.ConvertAll<Utf8String>(x => x), jsonComputeTask.WorkingDirectory, sandboxHash);
+				task.EnvVars = jsonComputeTask.EnvVars.ToDictionary(x => (Utf8String)x.Key, x => (Utf8String)x.Value);
+				task.OutputPaths.AddRange(jsonComputeTask.OutputPaths.Select(x => (Utf8String)x));
+				task.RequirementsHash = AddCbObject(blobs, requirements);
 
-				await ExecuteAction(Logger, ComputeClient, StorageClient, JsonComputeTask.ClusterId, Task, Blobs);	
+				await ExecuteAction(logger, computeClient, storageClient, jsonComputeTask.ClusterId, task, blobs);	
 			}
 			return 0;
 		}
 
-		static IoHash AddCbObject<T>(Dictionary<IoHash, byte[]> HashToData, T Source)
+		static IoHash AddCbObject<T>(Dictionary<IoHash, byte[]> hashToData, T source)
 		{
-			CbObject Obj = CbSerializer.Serialize<T>(Source);
-			IoHash Hash = Obj.GetHash();
-			HashToData[Hash] = Obj.GetView().ToArray();
-			return Hash;
+			CbObject obj = CbSerializer.Serialize<T>(source);
+			IoHash hash = obj.GetHash();
+			hashToData[hash] = obj.GetView().ToArray();
+			return hash;
 		}
 
-		private async Task ExecuteAction(ILogger Logger, IComputeClient ComputeClient, IStorageClient StorageClient, ClusterId ClusterId, ComputeTask Task, Dictionary<IoHash, byte[]> UploadList)
+		private async Task ExecuteAction(ILogger logger, IComputeClient computeClient, IStorageClient storageClient, ClusterId clusterId, ComputeTask task, Dictionary<IoHash, byte[]> uploadList)
 		{
-			IComputeClusterInfo Cluster = await ComputeClient.GetClusterInfoAsync(ClusterId);
+			IComputeClusterInfo cluster = await computeClient.GetClusterInfoAsync(clusterId);
 
 			{
-				int Index = 0, TotalUploaded = 0, TotalSkipped = 0, UploadedBytes = 0, SkippedBytes = 0;
-				var Tasks = UploadList.Select(async (Pair) =>
+				int index = 0, totalUploaded = 0, totalSkipped = 0, uploadedBytes = 0, skippedBytes = 0;
+				IEnumerable<Task> tasks = uploadList.Select(async (pair) =>
 				{
-					if (await StorageClient.HasBlobAsync(Cluster.NamespaceId, Pair.Key))
+					if (await storageClient.HasBlobAsync(cluster.NamespaceId, pair.Key))
 					{
-						Logger.LogInformation("Skipped blob {Idx}/{Count}: {Hash} ({Length} bytes)", Interlocked.Increment(ref Index), UploadList.Count, Pair.Key, Pair.Value.Length);
-						Interlocked.Increment(ref TotalSkipped);
-						Interlocked.Add(ref SkippedBytes, Pair.Value.Length);
+						logger.LogInformation("Skipped blob {Idx}/{Count}: {Hash} ({Length} bytes)", Interlocked.Increment(ref index), uploadList.Count, pair.Key, pair.Value.Length);
+						Interlocked.Increment(ref totalSkipped);
+						Interlocked.Add(ref skippedBytes, pair.Value.Length);
 						return;
 					}
 
-					await StorageClient.WriteBlobFromMemoryAsync(Cluster.NamespaceId, Pair.Key, Pair.Value);
-					Logger.LogInformation("Uploaded blob {Idx}/{Count}: {Hash} ({Bytes} bytes)", Interlocked.Increment(ref Index), UploadList.Count, Pair.Key, Pair.Value.Length);
-					Interlocked.Increment(ref TotalUploaded);
-					Interlocked.Add(ref UploadedBytes, Pair.Value.Length);
+					await storageClient.WriteBlobFromMemoryAsync(cluster.NamespaceId, pair.Key, pair.Value);
+					logger.LogInformation("Uploaded blob {Idx}/{Count}: {Hash} ({Bytes} bytes)", Interlocked.Increment(ref index), uploadList.Count, pair.Key, pair.Value.Length);
+					Interlocked.Increment(ref totalUploaded);
+					Interlocked.Add(ref uploadedBytes, pair.Value.Length);
 				});
-				await System.Threading.Tasks.Task.WhenAll(Tasks);
+				await Task.WhenAll(tasks);
 
-				Logger.LogInformation("Uploaded {UploadedCount} blobs ({UploadedBytes} bytes), Skipped {SkippedCount} blobs ({SkippedBytes} bytes)", TotalUploaded, UploadedBytes, TotalSkipped, SkippedBytes);
+				logger.LogInformation("Uploaded {UploadedCount} blobs ({UploadedBytes} bytes), Skipped {SkippedCount} blobs ({SkippedBytes} bytes)", totalUploaded, uploadedBytes, totalSkipped, skippedBytes);
 			}
 
-			CbObject TaskObject = CbSerializer.Serialize(Task);
-			IoHash TaskHash = IoHash.Compute(TaskObject.GetView().Span);
-			RefId TaskRefId = new RefId(TaskHash);
-			await StorageClient.SetRefAsync(Cluster.NamespaceId, Cluster.RequestBucketId, TaskRefId, TaskObject);
-
-			ChannelId ChannelId = new ChannelId(Guid.NewGuid().ToString());
-			Logger.LogInformation("cluster: {ClusterId}", ClusterId);
-			Logger.LogInformation("channel: {ChannelId}", ChannelId);
-			Logger.LogInformation("task: {TaskHash}", TaskHash);
-			Logger.LogInformation("requirements: {RequirementsHash}", Task.RequirementsHash);
+			CbObject taskObject = CbSerializer.Serialize(task);
+			IoHash taskHash = IoHash.Compute(taskObject.GetView().Span);
+			RefId taskRefId = new RefId(taskHash);
+			await storageClient.SetRefAsync(cluster.NamespaceId, cluster.RequestBucketId, taskRefId, taskObject);
 
+			ChannelId channelId = new ChannelId(Guid.NewGuid().ToString());
+			logger.LogInformation("cluster: {ClusterId}", clusterId);
+			logger.LogInformation("channel: {ChannelId}", channelId);
+			logger.LogInformation("task: {TaskHash}", taskHash);
+			logger.LogInformation("requirements: {RequirementsHash}", task.RequirementsHash);
 
 			// Execute the action
-			await ComputeClient.AddTaskAsync(ClusterId, ChannelId, TaskRefId, Task.RequirementsHash, SkipCacheLookup);
+			await computeClient.AddTaskAsync(clusterId, channelId, taskRefId, task.RequirementsHash, SkipCacheLookup);
 
-			await foreach(IComputeTaskInfo Response in ComputeClient.GetTaskUpdatesAsync(ClusterId, ChannelId))
+			await foreach(IComputeTaskInfo response in computeClient.GetTaskUpdatesAsync(clusterId, channelId))
 			{
-				Logger.LogInformation("{OperationName}: Execution state: {State}", Response.TaskRefId, Response.State.ToString());
-				if (!String.IsNullOrEmpty(Response.AgentId) || !String.IsNullOrEmpty(Response.LeaseId))
+				logger.LogInformation("{OperationName}: Execution state: {State}", response.TaskRefId, response.State.ToString());
+				if (!String.IsNullOrEmpty(response.AgentId) || !String.IsNullOrEmpty(response.LeaseId))
 				{
-					Logger.LogInformation("{OperationName}: Running on agent {AgentId} under lease {LeaseId}", Response.TaskRefId, Response.AgentId, Response.LeaseId);
+					logger.LogInformation("{OperationName}: Running on agent {AgentId} under lease {LeaseId}", response.TaskRefId, response.AgentId, response.LeaseId);
 				}
-				if (Response.ResultRefId != null)
+				if (response.ResultRefId != null)
 				{
-					await HandleCompleteTask(StorageClient, Cluster.NamespaceId, Cluster.ResponseBucketId, Response.ResultRefId.Value, Logger);
+					await HandleCompleteTask(storageClient, cluster.NamespaceId, cluster.ResponseBucketId, response.ResultRefId.Value, logger);
 				}
-				if (Response.State == ComputeTaskState.Complete)
+				if (response.State == ComputeTaskState.Complete)
 				{
-					if (Response.Outcome != ComputeTaskOutcome.Success)
+					if (response.Outcome != ComputeTaskOutcome.Success)
 					{
-						Logger.LogError("{OperationName}: Outcome: {Outcome}, Detail: {Detail}", Response.TaskRefId, Response.Outcome.ToString(), Response.Detail ?? "(none)");
+						logger.LogError("{OperationName}: Outcome: {Outcome}, Detail: {Detail}", response.TaskRefId, response.Outcome.ToString(), response.Detail ?? "(none)");
 					}
 					break;
 				}
 			}
 		}
 
-		async Task HandleCompleteTask(IStorageClient StorageClient, NamespaceId NamespaceId, BucketId OutputBucketId, RefId OutputRefId, ILogger Logger)
+		async Task HandleCompleteTask(IStorageClient storageClient, NamespaceId namespaceId, BucketId outputBucketId, RefId outputRefId, ILogger logger)
 		{
-			ComputeTaskResult Result = await StorageClient.GetRefAsync<ComputeTaskResult>(NamespaceId, OutputBucketId, OutputRefId);
-			Logger.LogInformation("exit: {ExitCode}", Result.ExitCode);
+			ComputeTaskResult result = await storageClient.GetRefAsync<ComputeTaskResult>(namespaceId, outputBucketId, outputRefId);
+			logger.LogInformation("exit: {ExitCode}", result.ExitCode);
 
-			await LogTaskOutputAsync(StorageClient, "stdout", NamespaceId, Result.StdOutHash, Logger);
-			await LogTaskOutputAsync(StorageClient, "stderr", NamespaceId, Result.StdErrHash, Logger);
+			await LogTaskOutputAsync(storageClient, "stdout", namespaceId, result.StdOutHash, logger);
+			await LogTaskOutputAsync(storageClient, "stderr", namespaceId, result.StdErrHash, logger);
 
-			if (Result.OutputHash != null && OutputDir != null)
+			if (result.OutputHash != null && OutputDir != null)
 			{
-				await WriteOutputAsync(StorageClient, NamespaceId, Result.OutputHash.Value, new DirectoryReference(OutputDir));
+				await WriteOutputAsync(storageClient, namespaceId, result.OutputHash.Value, new DirectoryReference(OutputDir));
 			}
 		}
 
-		async Task LogTaskOutputAsync(IStorageClient StorageClient, string Channel, NamespaceId NamespaceId, IoHash? LogHash, ILogger Logger)
+		async Task LogTaskOutputAsync(IStorageClient storageClient, string channel, NamespaceId namespaceId, IoHash? logHash, ILogger logger)
 		{
-			if (LogHash != null)
+			if (logHash != null)
 			{
-				byte[] StdOutData = await StorageClient.ReadBlobToMemoryAsync(NamespaceId, LogHash.Value);
-				if (StdOutData.Length > 0)
+				byte[] stdOutData = await storageClient.ReadBlobToMemoryAsync(namespaceId, logHash.Value);
+				if (stdOutData.Length > 0)
 				{
-					foreach (string Line in Encoding.UTF8.GetString(StdOutData).Split('\n'))
+					foreach (string line in Encoding.UTF8.GetString(stdOutData).Split('\n'))
 					{
-						Logger.LogDebug("{Channel}: {Line}", Channel, Line);
+						logger.LogDebug("{Channel}: {Line}", channel, line);
 					}
 				}
 			}
 		}
 
-		async Task WriteOutputAsync(IStorageClient StorageClient, NamespaceId NamespaceId, IoHash TreeHash, DirectoryReference OutputDir)
+		async Task WriteOutputAsync(IStorageClient storageClient, NamespaceId namespaceId, IoHash treeHash, DirectoryReference outputDir)
 		{
-			DirectoryTree Tree = await StorageClient.ReadBlobAsync<DirectoryTree>(NamespaceId, TreeHash);
+			DirectoryTree tree = await storageClient.ReadBlobAsync<DirectoryTree>(namespaceId, treeHash);
 
-			List<Task> Tasks = new List<Task>();
-			foreach (FileNode File in Tree.Files)
+			List<Task> tasks = new List<Task>();
+			foreach (FileNode file in tree.Files)
 			{
-				FileReference OutputFile = FileReference.Combine(OutputDir, File.Name.ToString());
-				Tasks.Add(WriteObjectToFileAsync(StorageClient, NamespaceId, File.Hash, OutputFile));
+				FileReference outputFile = FileReference.Combine(outputDir, file.Name.ToString());
+				tasks.Add(WriteObjectToFileAsync(storageClient, namespaceId, file.Hash, outputFile));
 			}
-			foreach (DirectoryNode Directory in Tree.Directories)
+			foreach (DirectoryNode directory in tree.Directories)
 			{
-				DirectoryReference NextOutputDir = DirectoryReference.Combine(OutputDir, Directory.Name.ToString());
-				Tasks.Add(WriteOutputAsync(StorageClient, NamespaceId, Directory.Hash, NextOutputDir));
+				DirectoryReference nextOutputDir = DirectoryReference.Combine(outputDir, directory.Name.ToString());
+				tasks.Add(WriteOutputAsync(storageClient, namespaceId, directory.Hash, nextOutputDir));
 			}
 
-			await Task.WhenAll(Tasks);
+			await Task.WhenAll(tasks);
 		}
 
-		static async Task WriteObjectToFileAsync(IStorageClient StorageClient, NamespaceId NamespaceId, IoHash Hash, FileReference OutputFile)
+		static async Task WriteObjectToFileAsync(IStorageClient storageClient, NamespaceId namespaceId, IoHash hash, FileReference outputFile)
 		{
-			DirectoryReference.CreateDirectory(OutputFile.Directory);
+			DirectoryReference.CreateDirectory(outputFile.Directory);
 
-			byte[] Data = await StorageClient.ReadBlobToMemoryAsync(NamespaceId, Hash);
-			await FileReference.WriteAllBytesAsync(OutputFile, Data);
+			byte[] data = await storageClient.ReadBlobToMemoryAsync(namespaceId, hash);
+			await FileReference.WriteAllBytesAsync(outputFile, data);
 		}
 	}
 }
diff --git a/Engine/Source/Programs/Horde/Horde.Agent/Commands/Parse/ParseCommand.cs b/Engine/Source/Programs/Horde/Horde.Agent/Commands/Parse/ParseCommand.cs
--- a/Engine/Source/Programs/Horde/Horde.Agent/Commands/Parse/ParseCommand.cs
+++ b/Engine/Source/Programs/Horde/Horde.Agent/Commands/Parse/ParseCommand.cs
@@ -20,28 +20,28 @@ namespace Horde.Agent.Commands.Parse
 	{
 		[CommandLine("-File=", Required = true)]
 		[Description("Log file to parse rather than executing an external program.")]
-		FileReference InputFile = null!;
+		FileReference InputFile { get; set; } = null!;
 
 		[CommandLine("-Ignore=")]
 		[Description("Path to a file containing error patterns to ignore, one regex per line.")]
-		FileReference? IgnorePatternsFile = null;
+		FileReference? IgnorePatternsFile { get; set; } = null;
 
 		[CommandLine("-WorkspaceDir=")]
 		[Description("The root workspace directory")]
-		DirectoryReference? WorkspaceDir = null;
+		DirectoryReference? WorkspaceDir { get; set; } = null;
 
 		[CommandLine("-Stream=")]
 		[Description("The stream synced to the workspace")]
-		string? Stream = null;
+		string? Stream { get; set; } = null;
 
 		[CommandLine("-Change=")]
 		[Description("The changelist number that has been synced")]
-		int? Change = null;
+		int? Change { get; set; } = null;
 
-		public override async Task<int> ExecuteAsync(ILogger Logger)
+		public override async Task<int> ExecuteAsync(ILogger logger)
 		{
 			// Read all the ignore patterns
-			List<string> IgnorePatterns = new List<string>();
+			List<string> ignorePatterns = new List<string>();
 			if (IgnorePatternsFile != null)
 			{
 				if (!FileReference.Exists(IgnorePatternsFile))
@@ -50,37 +50,37 @@ namespace Horde.Agent.Commands.Parse
 				}
 
 				// Read all the ignore patterns
-				string[] Lines = FileReference.ReadAllLines(IgnorePatternsFile);
-				foreach (string Line in Lines)
+				string[] lines = FileReference.ReadAllLines(IgnorePatternsFile);
+				foreach (string line in lines)
 				{
-					string TrimLine = Line.Trim();
-					if (TrimLine.Length > 0 && !TrimLine.StartsWith("#"))
+					string trimLine = line.Trim();
+					if (trimLine.Length > 0 && !trimLine.StartsWith("#"))
 					{
-						IgnorePatterns.Add(TrimLine);
+						ignorePatterns.Add(trimLine);
 					}
 				}
 			}
 
 			// Read the file and pipe it through the event parser
-			using (FileStream InputStream = FileReference.Open(InputFile, FileMode.Open, FileAccess.Read))
+			using (FileStream inputStream = FileReference.Open(InputFile, FileMode.Open, FileAccess.Read))
 			{
-				LogParserContext Context = new LogParserContext();
-				Context.WorkspaceDir = WorkspaceDir;
-				Context.PerforceStream = Stream;
-				Context.PerforceChange = Change;
+				LogParserContext context = new LogParserContext();
+				context.WorkspaceDir = WorkspaceDir;
+				context.PerforceStream = Stream;
+				context.PerforceChange = Change;
 
-				using (LogParser Parser = new LogParser(Logger, Context, IgnorePatterns))
+				using (LogParser parser = new LogParser(logger, context, ignorePatterns))
 				{
-					byte[] Data = new byte[1024];
+					byte[] data = new byte[1024];
 					for (; ; )
 					{
-						int Length = await InputStream.ReadAsync(Data);
-						if(Length == 0)
+						int length = await inputStream.ReadAsync(data);
+						if(length == 0)
 						{
-							Parser.Flush();
+							parser.Flush();
 							break;
 						}
-						Parser.WriteData(Data.AsMemory(0, Length));
+						parser.WriteData(data.AsMemory(0, length));
 					}
 				}
 			}
diff --git a/Engine/Source/Programs/Horde/Horde.Agent/Commands/Service/InstallCommand.cs b/Engine/Source/Programs/Horde/Horde.Agent/Commands/Service/InstallCommand.cs
--- a/Engine/Source/Programs/Horde/Horde.Agent/Commands/Service/InstallCommand.cs
+++ b/Engine/Source/Programs/Horde/Horde.Agent/Commands/Service/InstallCommand.cs
@@ -26,79 +26,79 @@ namespace Horde.Agent.Commands.Service
 		/// Specifies the username for the service to run under
 		/// </summary>
 		[CommandLine("-UserName=")]
-		public string? UserName = null;
+		public string? UserName { get; set; } = null;
 
 		/// <summary>
 		/// Password for the username
 		/// </summary>
 		[CommandLine("-Password=")]
-		public string? Password = null;
+		public string? Password { get; set; } = null;
 
 		/// <summary>
 		/// The server profile to use
 		/// </summary>
 		[CommandLine("-Server=")]
-		public string? Server = null;
+		public string? Server { get; set; } = null;
 		
 		/// <summary>
 		/// Path to dotnet executable (dotnet.exe on Windows)
 		/// When left empty, the value of "dotnet" will be used.
 		/// </summary>
 		[CommandLine("-DotNetExecutable=")]
-		public string DotNetExecutable = "dotnet";
+		public string DotNetExecutable { get; set; } = "dotnet";
 
 		/// <summary>
 		/// Runs the service indefinitely
 		/// </summary>
-		/// <param name="Logger">Logger to use</param>
+		/// <param name="logger">Logger to use</param>
 		/// <returns>Exit code</returns>
-		public override Task<int> ExecuteAsync(ILogger Logger)
+		public override Task<int> ExecuteAsync(ILogger logger)
 		{
-			using (WindowsServiceManager ServiceManager = new WindowsServiceManager())
+			using (WindowsServiceManager serviceManager = new WindowsServiceManager())
 			{
-				using (WindowsService Service = ServiceManager.Open(ServiceName))
+				using (WindowsService service = serviceManager.Open(ServiceName))
 				{
-					if (Service.IsValid)
+					if (service.IsValid)
 					{
-						Logger.LogInformation("Stopping existing service...");
-						Service.Stop();
+						logger.LogInformation("Stopping existing service...");
+						service.Stop();
 
-						WindowsServiceStatus Status = Service.WaitForStatusChange(WindowsServiceStatus.Stopping, TimeSpan.FromSeconds(30.0));
-						if (Status != WindowsServiceStatus.Stopped)
+						WindowsServiceStatus status = service.WaitForStatusChange(WindowsServiceStatus.Stopping, TimeSpan.FromSeconds(30.0));
+						if (status != WindowsServiceStatus.Stopped)
 						{
-							Logger.LogError("Unable to stop service (status = {Status})", Status);
+							logger.LogError("Unable to stop service (status = {Status})", status);
 							return Task.FromResult(1);
 						}
 
-						Logger.LogInformation("Deleting service");
-						Service.Delete();
+						logger.LogInformation("Deleting service");
+						service.Delete();
 					}
 				}
 
-				Logger.LogInformation("Registering {ServiceName} service", ServiceName);
+				logger.LogInformation("Registering {ServiceName} service", ServiceName);
 
-				StringBuilder CommandLine = new StringBuilder();
-				CommandLine.AppendFormat("{0} \"{1}\" service run", DotNetExecutable, Assembly.GetEntryAssembly()!.Location);
+				StringBuilder commandLine = new StringBuilder();
+				commandLine.AppendFormat("{0} \"{1}\" service run", DotNetExecutable, Assembly.GetEntryAssembly()!.Location);
 				if(Server != null)
 				{
-					CommandLine.Append($" -server={Server}");
+					commandLine.Append($" -server={Server}");
 				}
 
-				using (WindowsService Service = ServiceManager.Create(ServiceName, "Horde Agent", CommandLine.ToString(), UserName, Password))
+				using (WindowsService service = serviceManager.Create(ServiceName, "Horde Agent", commandLine.ToString(), UserName, Password))
 				{
-					Service.SetDescription("Allows this machine to participate in a Horde farm.");
+					service.SetDescription("Allows this machine to participate in a Horde farm.");
 
-					Logger.LogInformation("Starting...");
-					Service.Start();
+					logger.LogInformation("Starting...");
+					service.Start();
 
-					WindowsServiceStatus Status = Service.WaitForStatusChange(WindowsServiceStatus.Starting, TimeSpan.FromSeconds(30.0));
-					if (Status != WindowsServiceStatus.Running)
+					WindowsServiceStatus status = service.WaitForStatusChange(WindowsServiceStatus.Starting, TimeSpan.FromSeconds(30.0));
+					if (status != WindowsServiceStatus.Running)
 					{
-						Logger.LogError("Unable to start service (status = {Status})", Status);
+						logger.LogError("Unable to start service (status = {Status})", status);
 						return Task.FromResult(1);
 					}
 
-					Logger.LogInformation("Done.");
+					logger.LogInformation("Done.");
 				}
 			}
 			return Task.FromResult(0);
diff --git a/Engine/Source/Programs/Horde/Horde.Agent/Commands/Service/RunCommand.cs b/Engine/Source/Programs/Horde/Horde.Agent/Commands/Service/RunCommand.cs
--- a/Engine/Source/Programs/Horde/Horde.Agent/Commands/Service/RunCommand.cs
+++ b/Engine/Source/Programs/Horde/Horde.Agent/Commands/Service/RunCommand.cs
@@ -39,97 +39,97 @@ namespace Horde.Agent.Modes.Service
 		/// Override for the server to use
 		/// </summary>
 		[CommandLine("-Server=")]
-		string? Server = null;
+		string? Server { get; set; } = null;
 
 		/// <summary>
 		/// Override the working directory
 		/// </summary>
 		[CommandLine("-WorkingDir=")]
-		string? WorkingDir = null;
+		string? WorkingDir { get; set; } = null;
 
 		/// <summary>
 		/// Runs the service indefinitely
 		/// </summary>
 		/// <returns>Exit code</returns>
-		public override async Task<int> ExecuteAsync(ILogger Logger)
+		public override async Task<int> ExecuteAsync(ILogger logger)
 		{
-			IHostBuilder HostBuilder = Host.CreateDefaultBuilder();
+			IHostBuilder hostBuilder = Host.CreateDefaultBuilder();
 
 			// Attempt to setup this process as a Windows service. A race condition inside Microsoft.Extensions.Hosting.WindowsServices.WindowsServiceHelpers.IsWindowsService
 			// can result in accessing the parent process after it's terminated, so catch any exceptions that it throws.
 			try
 			{
-				HostBuilder = HostBuilder.UseWindowsService();
+				hostBuilder = hostBuilder.UseWindowsService();
 			}
 			catch (InvalidOperationException)
 			{
 			}
 
-			HostBuilder = HostBuilder
-				.ConfigureAppConfiguration(Builder =>
+			hostBuilder = hostBuilder
+				.ConfigureAppConfiguration(builder =>
 				{
-					Dictionary<string, string> Overrides = new Dictionary<string, string>();
+					Dictionary<string, string> overrides = new Dictionary<string, string>();
 					if (Server != null)
 					{
-						Overrides.Add($"{AgentSettings.SectionName}:{nameof(AgentSettings.Server)}", Server);
+						overrides.Add($"{AgentSettings.SectionName}:{nameof(AgentSettings.Server)}", Server);
 					}
 					if (WorkingDir != null)
 					{
-						Overrides.Add($"{AgentSettings.SectionName}:{nameof(AgentSettings.WorkingDir)}", WorkingDir);
+						overrides.Add($"{AgentSettings.SectionName}:{nameof(AgentSettings.WorkingDir)}", WorkingDir);
 					}
-					Builder.AddInMemoryCollection(Overrides);
+					builder.AddInMemoryCollection(overrides);
 				})
-				.ConfigureLogging(Builder =>
+				.ConfigureLogging(builder =>
 				{
-					Builder.ClearProviders();
-					Builder.AddProvider(new Logging.HordeLoggerProvider());
-					Builder.AddFilter<Logging.HordeLoggerProvider>(null, LogLevel.Trace);
+					builder.ClearProviders();
+					builder.AddProvider(new Logging.HordeLoggerProvider());
+					builder.AddFilter<Logging.HordeLoggerProvider>(null, LogLevel.Trace);
 				})
-				.ConfigureServices((HostContext, Services) =>
+				.ConfigureServices((hostContext, services) =>
 				{
-					Services.Configure<HostOptions>(Options =>
+					services.Configure<HostOptions>(options =>
 					{
 						// Allow the worker service to terminate any before shutting down
-						Options.ShutdownTimeout = TimeSpan.FromSeconds(30);
+						options.ShutdownTimeout = TimeSpan.FromSeconds(30);
 					});
 
-					IConfigurationSection ConfigSection = HostContext.Configuration.GetSection(AgentSettings.SectionName);
-					Services.AddOptions<AgentSettings>().Configure(Options => ConfigSection.Bind(Options)).ValidateDataAnnotations();
+					IConfigurationSection configSection = hostContext.Configuration.GetSection(AgentSettings.SectionName);
+					services.AddOptions<AgentSettings>().Configure(options => configSection.Bind(options)).ValidateDataAnnotations();
 
-					AgentSettings Settings = new AgentSettings();
-					ConfigSection.Bind(Settings);
+					AgentSettings settings = new AgentSettings();
+					configSection.Bind(settings);
 
-					ServerProfile ServerProfile = Settings.GetCurrentServerProfile();
-					ConfigureTracing(ServerProfile.Environment, Program.Version);
+					ServerProfile serverProfile = settings.GetCurrentServerProfile();
+					ConfigureTracing(serverProfile.Environment, Program.Version);
 
-					Logging.SetEnv(ServerProfile.Environment);
+					Logging.SetEnv(serverProfile.Environment);
 
-					Services.AddHttpClient(Program.HordeServerClientName, Config =>
+					services.AddHttpClient(Program.HordeServerClientName, config =>
 					{
-						Config.BaseAddress = new Uri(ServerProfile.Url);
-						Config.DefaultRequestHeaders.Add("Accept", "application/json");
-						Config.Timeout = TimeSpan.FromSeconds(300); // Need to make sure this doesn't cancel any long running gRPC streaming calls (eg. session update)
+						config.BaseAddress = new Uri(serverProfile.Url);
+						config.DefaultRequestHeaders.Add("Accept", "application/json");
+						config.Timeout = TimeSpan.FromSeconds(300); // Need to make sure this doesn't cancel any long running gRPC streaming calls (eg. session update)
 					})
 					.ConfigurePrimaryHttpMessageHandler(() =>
 					{
-						HttpClientHandler Handler = new HttpClientHandler();
-						Handler.ServerCertificateCustomValidationCallback += (Sender, Cert, Chain, Errors) => CertificateHelper.CertificateValidationCallBack(Logger, Sender, Cert, Chain, Errors, ServerProfile);
-						return Handler;
+						HttpClientHandler handler = new HttpClientHandler();
+						handler.ServerCertificateCustomValidationCallback += (sender, cert, chain, errors) => CertificateHelper.CertificateValidationCallBack(logger, sender, cert, chain, errors, serverProfile);
+						return handler;
 					})
-					.AddTransientHttpErrorPolicy(Builder =>
+					.AddTransientHttpErrorPolicy(builder =>
 					{
-						return Builder.WaitAndRetryAsync(new[] { TimeSpan.FromSeconds(1), TimeSpan.FromSeconds(5), TimeSpan.FromSeconds(10) });
+						return builder.WaitAndRetryAsync(new[] { TimeSpan.FromSeconds(1), TimeSpan.FromSeconds(5), TimeSpan.FromSeconds(10) });
 					});
 
-					Services.AddHordeStorage(Settings => ConfigSection.GetCurrentServerProfile().GetSection(nameof(ServerProfile.Storage)).Bind(Settings));
+					services.AddHordeStorage(settings => configSection.GetCurrentServerProfile().GetSection(nameof(serverProfile.Storage)).Bind(settings));
 
-					Services.AddSingleton<GrpcService>();
-					Services.AddHostedService<WorkerService>();
+					services.AddSingleton<GrpcService>();
+					services.AddHostedService<WorkerService>();
 				});
 
 			try
 			{
-				await HostBuilder.Build().RunAsync();
+				await hostBuilder.Build().RunAsync();
 			}
 			catch (OperationCanceledException)
 			{
@@ -139,18 +139,18 @@ namespace Horde.Agent.Modes.Service
 			return 0;
 		}
 
-		static void ConfigureTracing(string Environment, string Version)
+		static void ConfigureTracing(string environment, string version)
 		{
-			TracerSettings Settings = TracerSettings.FromDefaultSources();
-			Settings.Environment = Environment;
-			Settings.ServiceName = "hordeagent";
-			Settings.ServiceVersion = Version;
-			Settings.LogsInjectionEnabled = true;
+			TracerSettings settings = TracerSettings.FromDefaultSources();
+			settings.Environment = environment;
+			settings.ServiceName = "hordeagent";
+			settings.ServiceVersion = version;
+			settings.LogsInjectionEnabled = true;
 
-			Tracer.Instance = new Tracer(Settings);
+			Tracer.Instance = new Tracer(settings);
 
-			ITracer OpenTracer = Datadog.Trace.OpenTracing.OpenTracingTracerFactory.WrapTracer(Tracer.Instance);
-			GlobalTracer.Register(OpenTracer);
+			ITracer openTracer = Datadog.Trace.OpenTracing.OpenTracingTracerFactory.WrapTracer(Tracer.Instance);
+			GlobalTracer.Register(openTracer);
 		}
 	}
 }
diff --git a/Engine/Source/Programs/Horde/Horde.Agent/Commands/Service/UninstallCommand.cs b/Engine/Source/Programs/Horde/Horde.Agent/Commands/Service/UninstallCommand.cs
--- a/Engine/Source/Programs/Horde/Horde.Agent/Commands/Service/UninstallCommand.cs
+++ b/Engine/Source/Programs/Horde/Horde.Agent/Commands/Service/UninstallCommand.cs
@@ -25,32 +25,32 @@ namespace Horde.Agent.Commands.Service
 		/// <summary>
 		/// Runs the service indefinitely
 		/// </summary>
-		/// <param name="Logger">Logger to use</param>
+		/// <param name="logger">Logger to use</param>
 		/// <returns>Exit code</returns>
-		public override Task<int> ExecuteAsync(ILogger Logger)
+		public override Task<int> ExecuteAsync(ILogger logger)
 		{
-			using (WindowsServiceManager ServiceManager = new WindowsServiceManager())
+			using (WindowsServiceManager serviceManager = new WindowsServiceManager())
 			{
-				using (WindowsService Service = ServiceManager.Open(ServiceName))
+				using (WindowsService service = serviceManager.Open(ServiceName))
 				{
-					if (Service.IsValid)
+					if (service.IsValid)
 					{
-						Logger.LogInformation("Stopping existing service...");
-						Service.Stop();
+						logger.LogInformation("Stopping existing service...");
+						service.Stop();
 
-						WindowsServiceStatus Status = Service.WaitForStatusChange(WindowsServiceStatus.Stopping, TimeSpan.FromSeconds(30.0));
-						if (Status != WindowsServiceStatus.Stopped)
+						WindowsServiceStatus status = service.WaitForStatusChange(WindowsServiceStatus.Stopping, TimeSpan.FromSeconds(30.0));
+						if (status != WindowsServiceStatus.Stopped)
 						{
-							Logger.LogError("Unable to stop service (status = {Status})", Status);
+							logger.LogError("Unable to stop service (status = {Status})", status);
 							return Task.FromResult(1);
 						}
 
-						Logger.LogInformation("Deleting service");
-						Service.Delete();
+						logger.LogInformation("Deleting service");
+						service.Delete();
 					}
 					else
 					{
-						Logger.LogInformation("Unable to find service {ServiceName}", ServiceName);
+						logger.LogInformation("Unable to find service {ServiceName}", ServiceName);
 					}
 				}
 			}
diff --git a/Engine/Source/Programs/Horde/Horde.Agent/Commands/Service/UpgradeCommand.cs b/Engine/Source/Programs/Horde/Horde.Agent/Commands/Service/UpgradeCommand.cs
--- a/Engine/Source/Programs/Horde/Horde.Agent/Commands/Service/UpgradeCommand.cs
+++ b/Engine/Source/Programs/Horde/Horde.Agent/Commands/Service/UpgradeCommand.cs
@@ -30,167 +30,167 @@ namespace Horde.Agent.Commands
 		/// The process id to replace
 		/// </summary>
 		[CommandLine("-ProcessId=", Required = true)]
-		int ProcessId = -1;
+		int ProcessId { get; set; } = -1;
 
 		/// <summary>
 		/// The target directory to install to
 		/// </summary>
 		[CommandLine("-TargetDir=", Required = true)]
-		DirectoryReference TargetDir = null!;
+		DirectoryReference TargetDir { get; set; } = null!;
 
 		/// <summary>
 		/// Arguments to forwar to the target executable
 		/// </summary>
 		[CommandLine("-Arguments=", Required = true)]
-		string Arguments = null!;
+		string Arguments { get; set; } = null!;
 
 		/// <summary>
 		/// Upgrades the application to a new version
 		/// </summary>
-		/// <param name="Logger">The log output device</param>
+		/// <param name="logger">The log output device</param>
 		/// <returns>True if the upgrade succeeded</returns>
-		public override Task<int> ExecuteAsync(ILogger Logger)
+		public override Task<int> ExecuteAsync(ILogger logger)
 		{
 			// Stop the other process
-			Logger.LogInformation("Attempting to perform upgrade on process {ProcessId}", ProcessId);
-			using (Process OtherProcess = Process.GetProcessById(ProcessId))
+			logger.LogInformation("Attempting to perform upgrade on process {ProcessId}", ProcessId);
+			using (Process otherProcess = Process.GetProcessById(ProcessId))
 			{
 				// Get the directory containing the target application
-				DirectoryInfo TargetDir = new DirectoryInfo(this.TargetDir.FullName);
-				HashSet<string> TargetFiles = new HashSet<string>(TargetDir.EnumerateFiles("*", SearchOption.AllDirectories).Select(x => x.FullName), StringComparer.OrdinalIgnoreCase);
+				DirectoryInfo targetDir = new DirectoryInfo(TargetDir.FullName);
+				HashSet<string> targetFiles = new HashSet<string>(targetDir.EnumerateFiles("*", SearchOption.AllDirectories).Select(x => x.FullName), StringComparer.OrdinalIgnoreCase);
 
 				// Find all the source files
-				DirectoryInfo SourceDir = new DirectoryInfo(Path.GetDirectoryName(Assembly.GetExecutingAssembly().Location));
-				HashSet<string> SourceFiles = new HashSet<string>(SourceDir.EnumerateFiles("*", SearchOption.AllDirectories).Select(x => x.FullName), StringComparer.OrdinalIgnoreCase);
+				DirectoryInfo sourceDir = new DirectoryInfo(Path.GetDirectoryName(Assembly.GetExecutingAssembly().Location));
+				HashSet<string> sourceFiles = new HashSet<string>(sourceDir.EnumerateFiles("*", SearchOption.AllDirectories).Select(x => x.FullName), StringComparer.OrdinalIgnoreCase);
 
 				// Exclude all the source files from the list of target files, since we may be in a subdirectory
-				TargetFiles.ExceptWith(SourceFiles);
+				targetFiles.ExceptWith(sourceFiles);
 
 				// Ignore any files that are in the saved directory
-				string SourceDataDir = Path.Combine(SourceDir.FullName, "Saved") + Path.DirectorySeparatorChar;
-				SourceFiles.RemoveWhere(x => x.StartsWith(SourceDataDir, StringComparison.OrdinalIgnoreCase));
+				string sourceDataDir = Path.Combine(sourceDir.FullName, "Saved") + Path.DirectorySeparatorChar;
+				sourceFiles.RemoveWhere(x => x.StartsWith(sourceDataDir, StringComparison.OrdinalIgnoreCase));
 
-				string TargetDataDir = Path.Combine(TargetDir.FullName, "Saved") + Path.DirectorySeparatorChar;
-				TargetFiles.RemoveWhere(x => x.StartsWith(TargetDataDir, StringComparison.OrdinalIgnoreCase));
+				string targetDataDir = Path.Combine(targetDir.FullName, "Saved") + Path.DirectorySeparatorChar;
+				targetFiles.RemoveWhere(x => x.StartsWith(targetDataDir, StringComparison.OrdinalIgnoreCase));
 
 				// Copy all the files into the target directory
-				List<Tuple<string, string>> RenameFiles = new List<Tuple<string, string>>();
-				foreach (string SourceFile in SourceFiles)
+				List<Tuple<string, string>> renameFiles = new List<Tuple<string, string>>();
+				foreach (string sourceFile in sourceFiles)
 				{
-					if (!SourceFile.StartsWith(SourceDir.FullName, StringComparison.OrdinalIgnoreCase))
+					if (!sourceFile.StartsWith(sourceDir.FullName, StringComparison.OrdinalIgnoreCase))
 					{
-						throw new InvalidDataException($"Expected {SourceFile} to be under {SourceDir.FullName}");
+						throw new InvalidDataException($"Expected {sourceFile} to be under {sourceDir.FullName}");
 					}
 
-					string TargetFile = TargetDir.FullName + SourceFile.Substring(SourceDir.FullName.Length);
-					Directory.CreateDirectory(Path.GetDirectoryName(TargetFile));
+					string targetFile = targetDir.FullName + sourceFile.Substring(sourceDir.FullName.Length);
+					Directory.CreateDirectory(Path.GetDirectoryName(targetFile));
 
-					string TargetFileBeforeRename = TargetFile + ".new";
-					Logger.LogDebug("Copying {SourceFile} to {TargetFileBeforeRename}", SourceFile, TargetFileBeforeRename);
-					File.Copy(SourceFile, TargetFileBeforeRename, true);
+					string targetFileBeforeRename = targetFile + ".new";
+					logger.LogDebug("Copying {SourceFile} to {TargetFileBeforeRename}", sourceFile, targetFileBeforeRename);
+					File.Copy(sourceFile, targetFileBeforeRename, true);
 
-					RenameFiles.Add(Tuple.Create(TargetFileBeforeRename, TargetFile));
-					TargetFiles.Remove(TargetFileBeforeRename);
+					renameFiles.Add(Tuple.Create(targetFileBeforeRename, targetFile));
+					targetFiles.Remove(targetFileBeforeRename);
 				}
 
 				if (RuntimeInformation.IsOSPlatform(OSPlatform.Windows))
 				{
-					UpgradeWindowsService(Logger, OtherProcess, TargetFiles, RenameFiles);
+					UpgradeWindowsService(logger, otherProcess, targetFiles, renameFiles);
 				}
 				else if (RuntimeInformation.IsOSPlatform(OSPlatform.OSX))
 				{
-					UpgradeMacService(Logger, OtherProcess, TargetFiles, RenameFiles);
+					UpgradeMacService(logger, otherProcess, targetFiles, renameFiles);
 				}
 				else if (RuntimeInformation.IsOSPlatform(OSPlatform.Linux))
 				{
-					UpgradeLinuxService(Logger, OtherProcess, TargetFiles, RenameFiles);
+					UpgradeLinuxService(logger, otherProcess, targetFiles, renameFiles);
 				}
 				else
 				{
-					Logger.LogError("Agent is not running a platform that supports Upgrades. Platform: {0}", RuntimeInformation.OSDescription);
+					logger.LogError("Agent is not running a platform that supports Upgrades. Platform: {0}", RuntimeInformation.OSDescription);
 					return Task.FromResult(-1);
 				}
 			}
-			Logger.LogInformation("Upgrade complete");
+			logger.LogInformation("Upgrade complete");
 			return Task.FromResult(0);
 		}
 
-		void UpgradeFilesInPlace(ILogger Logger, HashSet<string> TargetFiles, List<Tuple<string, string>> RenameFiles)
+		void UpgradeFilesInPlace(ILogger logger, HashSet<string> targetFiles, List<Tuple<string, string>> renameFiles)
 		{
 			// Remove all the target files
-			foreach (string TargetFile in TargetFiles)
+			foreach (string targetFile in targetFiles)
 			{
-				Logger.LogDebug("Deleting {File}", TargetFile);
-				File.SetAttributes(TargetFile, FileAttributes.Normal);
-				File.Delete(TargetFile);
+				logger.LogDebug("Deleting {File}", targetFile);
+				File.SetAttributes(targetFile, FileAttributes.Normal);
+				File.Delete(targetFile);
 			}
 
 			// Rename all the new files into place
-			foreach (Tuple<string, string> Pair in RenameFiles)
+			foreach (Tuple<string, string> pair in renameFiles)
 			{
-				Logger.LogDebug("Renaming {SourceFile} to {TargetFile}", Pair.Item1, Pair.Item2);
-				File.Move(Pair.Item1, Pair.Item2, true);
+				logger.LogDebug("Renaming {SourceFile} to {TargetFile}", pair.Item1, pair.Item2);
+				File.Move(pair.Item1, pair.Item2, true);
 			}
 		}
 
-		void UpgradeMacService(ILogger Logger, Process OtherProcess, HashSet<string> TargetFiles, List<Tuple<string, string>> RenameFiles)
+		void UpgradeMacService(ILogger logger, Process otherProcess, HashSet<string> targetFiles, List<Tuple<string, string>> renameFiles)
 		{
-			UpgradeFilesInPlace(Logger, TargetFiles, RenameFiles);
-			Logger.LogDebug("Upgrade completed, restarting...");
-			OtherProcess.Kill();
+			UpgradeFilesInPlace(logger, targetFiles, renameFiles);
+			logger.LogDebug("Upgrade completed, restarting...");
+			otherProcess.Kill();
 		}
 
-		void UpgradeLinuxService(ILogger Logger, Process OtherProcess, HashSet<string> TargetFiles, List<Tuple<string, string>> RenameFiles)
+		void UpgradeLinuxService(ILogger logger, Process otherProcess, HashSet<string> targetFiles, List<Tuple<string, string>> renameFiles)
 		{
-			UpgradeFilesInPlace(Logger, TargetFiles, RenameFiles);
-			Logger.LogDebug("Upgrade completed, restarting...");
-			OtherProcess.Kill();
+			UpgradeFilesInPlace(logger, targetFiles, renameFiles);
+			logger.LogDebug("Upgrade completed, restarting...");
+			otherProcess.Kill();
 		}
 
-		void UpgradeWindowsService(ILogger Logger, Process OtherProcess, HashSet<string> TargetFiles, List<Tuple<string, string>> RenameFiles)
+		void UpgradeWindowsService(ILogger logger, Process otherProcess, HashSet<string> targetFiles, List<Tuple<string, string>> renameFiles)
 		{
 			// Try to get the service associated with the passed-in process id
-			using (ServiceController? Service = GetServiceForProcess(ProcessId))
+			using (ServiceController? service = GetServiceForProcess(ProcessId))
 			{
 				// Stop the process
-				if (Service == null)
+				if (service == null)
 				{
-					Logger.LogInformation("Terminating other process");
-					OtherProcess.Kill();
+					logger.LogInformation("Terminating other process");
+					otherProcess.Kill();
 				}
 				else
 				{
-					Logger.LogInformation("Stopping service");
-					Service.Stop();
+					logger.LogInformation("Stopping service");
+					service.Stop();
 				}
-				OtherProcess.WaitForExit();
+				otherProcess.WaitForExit();
 
-				UpgradeFilesInPlace(Logger, TargetFiles, RenameFiles);
+				UpgradeFilesInPlace(logger, targetFiles, renameFiles);
 
 				// Run the new application
-				if (Service == null)
+				if (service == null)
 				{
-					string DriverFileName = "dotnet";
-					string AssemblyFileName = Path.Combine(TargetDir.FullName, Path.GetFileName(Assembly.GetExecutingAssembly().Location));
+					string driverFileName = "dotnet";
+					string assemblyFileName = Path.Combine(TargetDir.FullName, Path.GetFileName(Assembly.GetExecutingAssembly().Location));
 
-					StringBuilder DriverArguments = new StringBuilder();
-					DriverArguments.AppendArgument(AssemblyFileName);
-					DriverArguments.Append(' ');
-					DriverArguments.Append(Arguments);
+					StringBuilder driverArguments = new StringBuilder();
+					driverArguments.AppendArgument(assemblyFileName);
+					driverArguments.Append(' ');
+					driverArguments.Append(Arguments);
 
-					StringBuilder Launch = new StringBuilder();
-					Launch.AppendArgument(DriverFileName);
-					Launch.Append(' ');
-					Launch.Append(DriverArguments);
-					Logger.LogInformation("Launching: {Launch}", Launch.ToString());
+					StringBuilder launch = new StringBuilder();
+					launch.AppendArgument(driverFileName);
+					launch.Append(' ');
+					launch.Append(driverArguments);
+					logger.LogInformation("Launching: {Launch}", launch.ToString());
 
-					using Process NewProcess = Process.Start(DriverFileName, DriverArguments.ToString());
+					using Process newProcess = Process.Start(driverFileName, driverArguments.ToString());
 				}
 				else
 				{
 					// Start the service again
-					Logger.LogInformation("Restarting service");
-					Service.Start();
+					logger.LogInformation("Restarting service");
+					service.Start();
 				}
 			}
 		}
@@ -198,23 +198,23 @@ namespace Horde.Agent.Commands
 		/// <summary>
 		/// Try to find the service controller for the given process id
 		/// </summary>
-		/// <param name="ProcessId">The process id to search for</param>
+		/// <param name="processId">The process id to search for</param>
 		/// <returns>The service controller corresponding to this process</returns>
-		static ServiceController? GetServiceForProcess(int ProcessId)
+		static ServiceController? GetServiceForProcess(int processId)
 		{
 			try
 			{
-				using (ManagementObjectSearcher Searcher = new ManagementObjectSearcher($"SELECT Name FROM Win32_Service WHERE ProcessId={ProcessId}"))
+				using (ManagementObjectSearcher searcher = new ManagementObjectSearcher($"SELECT Name FROM Win32_Service WHERE ProcessId={processId}"))
 				{
-					foreach (ManagementObject Service in Searcher.Get())
+					foreach (ManagementObject service in searcher.Get())
 					{
-						PropertyData Property = Service.Properties["Name"];
-						if (Property != null)
+						PropertyData property = service.Properties["Name"];
+						if (property != null)
 						{
-							string? Name = Property.Value as string;
-							if (Name != null)
+							string? name = property.Value as string;
+							if (name != null)
 							{
-								return new ServiceController(Name);
+								return new ServiceController(name);
 							}
 						}
 					}
diff --git a/Engine/Source/Programs/Horde/Horde.Agent/Commands/Utilities/CapsCommand.cs b/Engine/Source/Programs/Horde/Horde.Agent/Commands/Utilities/CapsCommand.cs
--- a/Engine/Source/Programs/Horde/Horde.Agent/Commands/Utilities/CapsCommand.cs
+++ b/Engine/Source/Programs/Horde/Horde.Agent/Commands/Utilities/CapsCommand.cs
@@ -28,15 +28,15 @@ namespace Horde.Agent.Commands.Utilities
 	class CapsCommand : Command
 	{
 		/// <inheritdoc/>
-		public override async Task<int> ExecuteAsync(ILogger Logger)
+		public override async Task<int> ExecuteAsync(ILogger logger)
 		{
-			AgentCapabilities Capabilities = await WorkerService.GetAgentCapabilities(DirectoryReference.GetCurrentDirectory(), Logger);
-			foreach (DeviceCapabilities Device in Capabilities.Devices)
+			AgentCapabilities capabilities = await WorkerService.GetAgentCapabilities(DirectoryReference.GetCurrentDirectory(), logger);
+			foreach (DeviceCapabilities device in capabilities.Devices)
 			{
-				Logger.LogInformation("Device: {Name}", Device.Handle);
-				foreach (string Property in Device.Properties)
+				logger.LogInformation("Device: {Name}", device.Handle);
+				foreach (string property in device.Properties)
 				{
-					Logger.LogInformation("  {Property}", Property);
+					logger.LogInformation("  {Property}", property);
 				}
 			}
 			return 0;
diff --git a/Engine/Source/Programs/Horde/Horde.Agent/Commands/Utilities/CreateCertCommand.cs b/Engine/Source/Programs/Horde/Horde.Agent/Commands/Utilities/CreateCertCommand.cs
--- a/Engine/Source/Programs/Horde/Horde.Agent/Commands/Utilities/CreateCertCommand.cs
+++ b/Engine/Source/Programs/Horde/Horde.Agent/Commands/Utilities/CreateCertCommand.cs
@@ -25,66 +25,66 @@ namespace Horde.Agent.Commands.Certs
 	class CreateCertCommand : Command
 	{
 		[CommandLine("-Server=")]
-		public string? Server = null;
+		public string? Server { get; set; } = null;
 
 		[CommandLine("-DnsName=")]
-		public string? DnsName;
+		public string? DnsName { get; set; }
 
 		[CommandLine("-PrivateCert=")]
-		public string? PrivateCertFile;
+		public string? PrivateCertFile { get; set; }
 
 		[CommandLine("-Environment=")]
-		public string? Environment = "Development";
+		public string? Environment { get; set; } = "Development";
 
 		/// <summary>
 		/// Main entry point for this command
 		/// </summary>
-		/// <param name="Logger"></param>
+		/// <param name="logger"></param>
 		/// <returns>Async task</returns>
-		public override Task<int> ExecuteAsync(ILogger Logger)
+		public override Task<int> ExecuteAsync(ILogger logger)
 		{
 			if (DnsName == null)
 			{
-				IConfigurationRoot Config = new ConfigurationBuilder()
+				IConfigurationRoot config = new ConfigurationBuilder()
 					.SetBasePath(Directory.GetCurrentDirectory())
 					.AddJsonFile("appsettings.json", true)
 					.AddJsonFile($"appsettings.{Environment}.json", true)
 					.Build();
 
-				AgentSettings Settings = new AgentSettings();
-				Config.GetSection("Horde").Bind(Settings);
+				AgentSettings settings = new AgentSettings();
+				config.GetSection("Horde").Bind(settings);
 
 				if(Server != null)
 				{
-					Settings.Server = Server;
+					settings.Server = Server;
 				}
 
-				DnsName = Settings.GetCurrentServerProfile().Url;
+				DnsName = settings.GetCurrentServerProfile().Url;
 				DnsName = Regex.Replace(DnsName, @"^[a-zA-Z]+://", "");
 				DnsName = Regex.Replace(DnsName, "/.*$", "");
 			}
 
 			if (PrivateCertFile == null)
 			{
-				FileReference SolutionFile = FileReference.Combine(new FileReference(Assembly.GetExecutingAssembly().Location).Directory, "..", "..", "..", "..", "Horde.sln");
-				if (!FileReference.Exists(SolutionFile))
+				FileReference solutionFile = FileReference.Combine(new FileReference(Assembly.GetExecutingAssembly().Location).Directory, "..", "..", "..", "..", "Horde.sln");
+				if (!FileReference.Exists(solutionFile))
 				{
-					Logger.LogError("The -PrivateCertFile=... arguments must be specified when running outside the default build directory");
+					logger.LogError("The -PrivateCertFile=... arguments must be specified when running outside the default build directory");
 					return Task.FromResult(1);
 				}
-				PrivateCertFile = FileReference.Combine(SolutionFile.Directory, "HordeServer", "Certs", $"ServerToAgent-{Environment}.pfx").FullName;
+				PrivateCertFile = FileReference.Combine(solutionFile.Directory, "HordeServer", "Certs", $"ServerToAgent-{Environment}.pfx").FullName;
 			}
 
-			Logger.LogInformation("Creating certificate for {DnsName}", DnsName);
+			logger.LogInformation("Creating certificate for {DnsName}", DnsName);
 
-			byte[] PrivateCertData = CertificateUtils.CreateSelfSignedCert(DnsName, "Horde Server");
+			byte[] privateCertData = CertificateUtils.CreateSelfSignedCert(DnsName, "Horde Server");
 
-			Logger.LogInformation("Writing private cert: {PrivateCert}", new FileReference(PrivateCertFile).FullName);
-			File.WriteAllBytes(PrivateCertFile, PrivateCertData);
+			logger.LogInformation("Writing private cert: {PrivateCert}", new FileReference(PrivateCertFile).FullName);
+			File.WriteAllBytes(PrivateCertFile, privateCertData);
 
-			using X509Certificate2 Certificate = new X509Certificate2(PrivateCertData);
-			Logger.LogInformation("Certificate thumbprint is {Thumbprint}", Certificate.Thumbprint);
-			Logger.LogInformation("Add this thumbprint to list of trusted servers in appsettings.json to trust this server.");
+			using X509Certificate2 certificate = new X509Certificate2(privateCertData);
+			logger.LogInformation("Certificate thumbprint is {Thumbprint}", certificate.Thumbprint);
+			logger.LogInformation("Add this thumbprint to list of trusted servers in appsettings.json to trust this server.");
 
 			return Task.FromResult(0);
 		}
diff --git a/Engine/Source/Programs/Horde/Horde.Agent/Commands/Utilities/PublishCommand.cs b/Engine/Source/Programs/Horde/Horde.Agent/Commands/Utilities/PublishCommand.cs
--- a/Engine/Source/Programs/Horde/Horde.Agent/Commands/Utilities/PublishCommand.cs
+++ b/Engine/Source/Programs/Horde/Horde.Agent/Commands/Utilities/PublishCommand.cs
@@ -29,98 +29,98 @@ namespace Horde.Agent.Commands.Utilities
 		/// The server to upload the software to
 		/// </summary>
 		[CommandLine("-Server=", Required = true)]
-		public string Server = null!;
+		public string Server { get; set; } = null!;
 
 		/// <summary>
 		/// Access token used to authenticate with the server
 		/// </summary>
 		[CommandLine("-Token=", Required = true)]
-		public string Token = null!;
+		public string Token { get; set; } = null!;
 
 		/// <summary>
 		/// The input directory
 		/// </summary>
 		[CommandLine("-InputDir=", Required = true)]
-		public DirectoryReference InputDir = null!;
+		public DirectoryReference InputDir { get; set; } = null!;
 
 		/// <summary>
 		/// The channel to post to
 		/// </summary>
 		[CommandLine("-Channel=", Required = true)]
-		public string Channel = null!;
+		public string Channel { get; set; } = null!;
 
 		/// <summary>
 		/// Main entry point for this command
 		/// </summary>
-		/// <param name="Logger"></param>
+		/// <param name="logger"></param>
 		/// <returns>Async task</returns>
-		public override async Task<int> ExecuteAsync(ILogger Logger)
+		public override async Task<int> ExecuteAsync(ILogger logger)
 		{
 			// Create the zip archive
-			Logger.LogInformation("Creating archive...");
-			byte[] ArchiveData;
-			using (MemoryStream OutputStream = new MemoryStream())
+			logger.LogInformation("Creating archive...");
+			byte[] archiveData;
+			using (MemoryStream outputStream = new MemoryStream())
 			{
-				using (ZipArchive OutputArchive = new ZipArchive(OutputStream, ZipArchiveMode.Create, true))
+				using (ZipArchive outputArchive = new ZipArchive(outputStream, ZipArchiveMode.Create, true))
 				{
-					foreach (FileReference InputFile in DirectoryReference.EnumerateFiles(InputDir, "*", SearchOption.AllDirectories))
+					foreach (FileReference inputFile in DirectoryReference.EnumerateFiles(InputDir, "*", SearchOption.AllDirectories))
 					{
-						string RelativePath = InputFile.MakeRelativeTo(InputDir).Replace(Path.DirectorySeparatorChar, '/');
-						if (!RelativePath.Equals("Log.json", StringComparison.OrdinalIgnoreCase) && !RelativePath.Equals("Log.txt", StringComparison.OrdinalIgnoreCase))
+						string relativePath = inputFile.MakeRelativeTo(InputDir).Replace(Path.DirectorySeparatorChar, '/');
+						if (!relativePath.Equals("Log.json", StringComparison.OrdinalIgnoreCase) && !relativePath.Equals("Log.txt", StringComparison.OrdinalIgnoreCase))
 						{
-							Logger.LogInformation("  Adding {RelativePath}", RelativePath);
+							logger.LogInformation("  Adding {RelativePath}", relativePath);
 
-							ZipArchiveEntry OutputEntry = OutputArchive.CreateEntry(RelativePath, CompressionLevel.Optimal);
+							ZipArchiveEntry outputEntry = outputArchive.CreateEntry(relativePath, CompressionLevel.Optimal);
 
-							using Stream InputEntryStream = FileReference.Open(InputFile, FileMode.Open, FileAccess.Read, FileShare.Read);
-							using Stream OutputEntryStream = OutputEntry.Open();
+							using Stream inputEntryStream = FileReference.Open(inputFile, FileMode.Open, FileAccess.Read, FileShare.Read);
+							using Stream outputEntryStream = outputEntry.Open();
 
-							await InputEntryStream.CopyToAsync(OutputEntryStream);
+							await inputEntryStream.CopyToAsync(outputEntryStream);
 						}
 					}
 				}
-				ArchiveData = OutputStream.ToArray();
+				archiveData = outputStream.ToArray();
 			}
-			Logger.LogInformation("Complete ({Size:n0} bytes)", ArchiveData.Length);
+			logger.LogInformation("Complete ({Size:n0} bytes)", archiveData.Length);
 
 			// Read the application settings
-			IConfigurationRoot Config = new ConfigurationBuilder()
+			IConfigurationRoot config = new ConfigurationBuilder()
 				.SetBasePath(Program.AppDir.FullName)
 				.AddJsonFile("appsettings.json", true)
 				.AddJsonFile("appsettings.Production.json", true)
 				.Build();
 
-			AgentSettings Settings = new AgentSettings();
-			Config.GetSection("Horde").Bind(Settings);
+			AgentSettings settings = new AgentSettings();
+			config.GetSection("Horde").Bind(settings);
 
 			// Get the server we're using
-			ServerProfile ServerProfile = Settings.GetServerProfile(Server);
+			ServerProfile serverProfile = settings.GetServerProfile(Server);
 
 			// Create the http message handler
-			HttpClientHandler Handler = new HttpClientHandler();
-			Handler.ServerCertificateCustomValidationCallback += (Sender, Cert, Chain, Errors) => CertificateHelper.CertificateValidationCallBack(Logger, Sender, Cert, Chain, Errors, ServerProfile);
+			HttpClientHandler handler = new HttpClientHandler();
+			handler.ServerCertificateCustomValidationCallback += (sender, cert, chain, errors) => CertificateHelper.CertificateValidationCallBack(logger, sender, cert, chain, errors, serverProfile);
 
 			// Create a new http client for uploading
-			HttpClient HttpClient = new HttpClient(Handler);
-			HttpClient.BaseAddress = new Uri(ServerProfile.Url);
-			HttpClient.DefaultRequestHeaders.Authorization = new AuthenticationHeaderValue("Bearer", Token);
+			HttpClient httpClient = new HttpClient(handler);
+			httpClient.BaseAddress = new Uri(serverProfile.Url);
+			httpClient.DefaultRequestHeaders.Authorization = new AuthenticationHeaderValue("Bearer", Token);
 
-			GrpcChannelOptions ChannelOptions = new GrpcChannelOptions
+			GrpcChannelOptions channelOptions = new GrpcChannelOptions
 			{
-				HttpClient = HttpClient,
+				HttpClient = httpClient,
 				MaxReceiveMessageSize = 100 * 1024 * 1024, // 100 MB
 				MaxSendMessageSize = 100 * 1024 * 1024 // 100 MB
 			};
-			using (GrpcChannel RpcChannel = GrpcChannel.ForAddress(ServerProfile.Url, ChannelOptions))
+			using (GrpcChannel rpcChannel = GrpcChannel.ForAddress(serverProfile.Url, channelOptions))
 			{
-				HordeRpc.HordeRpcClient RpcClient = new HordeRpc.HordeRpcClient(RpcChannel);
+				HordeRpc.HordeRpcClient rpcClient = new HordeRpc.HordeRpcClient(rpcChannel);
 
-				UploadSoftwareRequest UploadRequest = new UploadSoftwareRequest();
-				UploadRequest.Channel = Channel;
-				UploadRequest.Data = Google.Protobuf.ByteString.CopyFrom(ArchiveData);
+				UploadSoftwareRequest uploadRequest = new UploadSoftwareRequest();
+				uploadRequest.Channel = Channel;
+				uploadRequest.Data = Google.Protobuf.ByteString.CopyFrom(archiveData);
 
-				UploadSoftwareResponse UploadResponse = await RpcClient.UploadSoftwareAsync(UploadRequest);
-				Logger.LogInformation("Created software (version={SoftwareId} channel={Default})", UploadResponse.Version, UploadRequest.Channel);
+				UploadSoftwareResponse uploadResponse = await rpcClient.UploadSoftwareAsync(uploadRequest);
+				logger.LogInformation("Created software (version={SoftwareId} channel={Default})", uploadResponse.Version, uploadRequest.Channel);
 			}
 
 			return 0;
diff --git a/Engine/Source/Programs/Horde/Horde.Agent/Commands/Workspace/CleanCommand.cs b/Engine/Source/Programs/Horde/Horde.Agent/Commands/Workspace/CleanCommand.cs
--- a/Engine/Source/Programs/Horde/Horde.Agent/Commands/Workspace/CleanCommand.cs
+++ b/Engine/Source/Programs/Horde/Horde.Agent/Commands/Workspace/CleanCommand.cs
@@ -19,11 +19,11 @@ namespace Horde.Agent.Commands.Workspace
 	{
 		[CommandLine("-Incremental")]
 		[Description("Performs an incremental sync, without removing intermediates")]
-		public bool bIncrementalSync = false;
+		public bool Incremental { get; set; } = false;
 
-		protected override Task ExecuteAsync(IPerforceConnection Perforce, ManagedWorkspace Repo, ILogger Logger)
+		protected override Task ExecuteAsync(IPerforceConnection perforce, ManagedWorkspace repo, ILogger logger)
 		{
-			return Repo.CleanAsync(!bIncrementalSync, CancellationToken.None);
+			return repo.CleanAsync(!Incremental, CancellationToken.None);
 		}
 	}
 }
diff --git a/Engine/Source/Programs/Horde/Horde.Agent/Commands/Workspace/ClearCommand.cs b/Engine/Source/Programs/Horde/Horde.Agent/Commands/Workspace/ClearCommand.cs
--- a/Engine/Source/Programs/Horde/Horde.Agent/Commands/Workspace/ClearCommand.cs
+++ b/Engine/Source/Programs/Horde/Horde.Agent/Commands/Workspace/ClearCommand.cs
@@ -16,9 +16,9 @@ namespace Horde.Agent.Commands.Workspace
 	[Command("Workspace", "Clear", "Empties the staging directory of any files, returning them to the cache")]
 	class ClearCommand : WorkspaceCommand
 	{
-		protected override Task ExecuteAsync(IPerforceConnection Perforce, ManagedWorkspace Repo, ILogger Logger)
+		protected override Task ExecuteAsync(IPerforceConnection perforce, ManagedWorkspace repo, ILogger logger)
 		{
-			return Repo.ClearAsync(CancellationToken.None);
+			return repo.ClearAsync(CancellationToken.None);
 		}
 	}
 }
diff --git a/Engine/Source/Programs/Horde/Horde.Agent/Commands/Workspace/DumpCommand.cs b/Engine/Source/Programs/Horde/Horde.Agent/Commands/Workspace/DumpCommand.cs
--- a/Engine/Source/Programs/Horde/Horde.Agent/Commands/Workspace/DumpCommand.cs
+++ b/Engine/Source/Programs/Horde/Horde.Agent/Commands/Workspace/DumpCommand.cs
@@ -15,9 +15,9 @@ namespace Horde.Agent.Commands.Workspace
 	[Command("Workspace", "Dump", "Dumps the contents of the repository to the log for analysis")]
 	class Dump : WorkspaceCommand
 	{
-		protected override Task ExecuteAsync(IPerforceConnection Perforce, ManagedWorkspace Repo, ILogger Logger)
+		protected override Task ExecuteAsync(IPerforceConnection perforce, ManagedWorkspace repo, ILogger logger)
 		{
-			Repo.Dump();
+			repo.Dump();
 			return Task.CompletedTask;
 		}
 	}
diff --git a/Engine/Source/Programs/Horde/Horde.Agent/Commands/Workspace/PopulateCacheCommand.cs b/Engine/Source/Programs/Horde/Horde.Agent/Commands/Workspace/PopulateCacheCommand.cs
--- a/Engine/Source/Programs/Horde/Horde.Agent/Commands/Workspace/PopulateCacheCommand.cs
+++ b/Engine/Source/Programs/Horde/Horde.Agent/Commands/Workspace/PopulateCacheCommand.cs
@@ -19,34 +19,34 @@ namespace Horde.Agent.Commands.Workspace
 	{
 		[CommandLine("-ClientAndStream=")]
 		[Description("Specifies client and stream pairs, in the format Client:Stream")]
-		List<string> ClientAndStreamParams = new List<string>();
+		List<string> ClientAndStreamParams { get; set; } = new List<string>();
 
 		[CommandLine("-Filter=")]
 		[Description("Filters for the files to sync, in P4 syntax (eg. /Engine/...)")]
-		List<string> Filters = new List<string>();
+		List<string> Filters { get; set; } = new List<string>();
 
 		[CommandLine("-FakeSync")]
 		[Description("Simulates the sync without actually fetching any files")]
-		bool bFakeSync = false;
+		bool FakeSync { get; set; } = false;
 
-		protected override async Task ExecuteAsync(IPerforceConnection Perforce, ManagedWorkspace Repo, ILogger Logger)
+		protected override async Task ExecuteAsync(IPerforceConnection perforce, ManagedWorkspace repo, ILogger logger)
 		{
-			List<string> ExpandedFilters = ExpandFilters(Filters);
+			List<string> expandedFilters = ExpandFilters(Filters);
 
-			List<PopulateRequest> PopulateRequests = new List<PopulateRequest>();
-			foreach (string ClientAndStreamParam in ClientAndStreamParams)
+			List<PopulateRequest> populateRequests = new List<PopulateRequest>();
+			foreach (string clientAndStreamParam in ClientAndStreamParams)
 			{
-				int Idx = ClientAndStreamParam.IndexOf(':');
-				if (Idx == -1)
+				int idx = clientAndStreamParam.IndexOf(':');
+				if (idx == -1)
 				{
 					throw new FatalErrorException("Expected -ClientAndStream=<ClientName>:<StreamName>");
 				}
 
-				using IPerforceConnection PerforceClient = await Perforce.WithClientAsync(ClientAndStreamParam.Substring(0, Idx));
-				PopulateRequests.Add(new PopulateRequest(PerforceClient, ClientAndStreamParam.Substring(Idx + 1), ExpandedFilters));
+				using IPerforceConnection perforceClient = await perforce.WithClientAsync(clientAndStreamParam.Substring(0, idx));
+				populateRequests.Add(new PopulateRequest(perforceClient, clientAndStreamParam.Substring(idx + 1), expandedFilters));
 			}
 
-			await Repo.PopulateAsync(PopulateRequests, bFakeSync, CancellationToken.None);
+			await repo.PopulateAsync(populateRequests, FakeSync, CancellationToken.None);
 		}
 	}
 }
diff --git a/Engine/Source/Programs/Horde/Horde.Agent/Commands/Workspace/PurgeCacheCommand.cs b/Engine/Source/Programs/Horde/Horde.Agent/Commands/Workspace/PurgeCacheCommand.cs
--- a/Engine/Source/Programs/Horde/Horde.Agent/Commands/Workspace/PurgeCacheCommand.cs
+++ b/Engine/Source/Programs/Horde/Horde.Agent/Commands/Workspace/PurgeCacheCommand.cs
@@ -17,17 +17,17 @@ namespace Horde.Agent.Commands.Workspace
 	class PurgeCacheCommand : WorkspaceCommand
 	{
 		[CommandLine("-Size=")]
-		string? SizeParam = null;
+		string? SizeParam { get; set; } = null;
 
-		protected override Task ExecuteAsync(IPerforceConnection Perforce, ManagedWorkspace Repo, ILogger Logger)
+		protected override Task ExecuteAsync(IPerforceConnection perforce, ManagedWorkspace repo, ILogger logger)
 		{
-			long Size = 0;
+			long size = 0;
 			if (SizeParam != null)
 			{
-				Size = ParseSize(SizeParam);
+				size = ParseSize(SizeParam);
 			}
 
-			return Repo.PurgeAsync(Size, CancellationToken.None);
+			return repo.PurgeAsync(size, CancellationToken.None);
 		}
 	}
 }
diff --git a/Engine/Source/Programs/Horde/Horde.Agent/Commands/Workspace/RepairCommand.cs b/Engine/Source/Programs/Horde/Horde.Agent/Commands/Workspace/RepairCommand.cs
--- a/Engine/Source/Programs/Horde/Horde.Agent/Commands/Workspace/RepairCommand.cs
+++ b/Engine/Source/Programs/Horde/Horde.Agent/Commands/Workspace/RepairCommand.cs
@@ -16,9 +16,9 @@ namespace Horde.Agent.Commands.Workspace
 	[Command("Workspace", "RepairCache", "Checks the integrity of the cache, and removes any invalid files")]
 	class RepairCommand : WorkspaceCommand
 	{
-		protected override Task ExecuteAsync(IPerforceConnection Perforce, ManagedWorkspace Repo, ILogger Logger)
+		protected override Task ExecuteAsync(IPerforceConnection perforce, ManagedWorkspace repo, ILogger logger)
 		{
-			return Repo.RepairAsync(CancellationToken.None);
+			return repo.RepairAsync(CancellationToken.None);
 		}
 	}
 }
diff --git a/Engine/Source/Programs/Horde/Horde.Agent/Commands/Workspace/RevertCommand.cs b/Engine/Source/Programs/Horde/Horde.Agent/Commands/Workspace/RevertCommand.cs
--- a/Engine/Source/Programs/Horde/Horde.Agent/Commands/Workspace/RevertCommand.cs
+++ b/Engine/Source/Programs/Horde/Horde.Agent/Commands/Workspace/RevertCommand.cs
@@ -19,12 +19,12 @@ namespace Horde.Agent.Commands.Workspace
 	{
 		[CommandLine("-Client=", Required = true)]
 		[Description("Client to revert all files for")]
-		string ClientName = null!;
+		string ClientName { get; set; } = null!;
 
-		protected override async Task ExecuteAsync(IPerforceConnection Perforce, ManagedWorkspace Repo, ILogger Logger)
+		protected override async Task ExecuteAsync(IPerforceConnection perforce, ManagedWorkspace repo, ILogger logger)
 		{
-			using IPerforceConnection PerforceClient = await Perforce.WithClientAsync(ClientName);
-			await Repo.RevertAsync(PerforceClient, CancellationToken.None);
+			using IPerforceConnection perforceClient = await perforce.WithClientAsync(ClientName);
+			await repo.RevertAsync(perforceClient, CancellationToken.None);
 		}
 	}
 }
diff --git a/Engine/Source/Programs/Horde/Horde.Agent/Commands/Workspace/SetupCommand.cs b/Engine/Source/Programs/Horde/Horde.Agent/Commands/Workspace/SetupCommand.cs
--- a/Engine/Source/Programs/Horde/Horde.Agent/Commands/Workspace/SetupCommand.cs
+++ b/Engine/Source/Programs/Horde/Horde.Agent/Commands/Workspace/SetupCommand.cs
@@ -19,16 +19,16 @@ namespace Horde.Agent.Commands.Workspace
 	{
 		[CommandLine("-Client=", Required = true)]
 		[Description("Name of the client to create")]
-		string ClientName = null!;
+		string ClientName { get; set; } = null!;
 
 		[CommandLine("-Stream=", Required = true)]
 		[Description("Name of the stream to configure")]
-		string StreamName = null!;
+		string StreamName { get; set; } = null!;
 
-		protected override async Task ExecuteAsync(IPerforceConnection Perforce, ManagedWorkspace Repo, ILogger Logger)
+		protected override async Task ExecuteAsync(IPerforceConnection perforce, ManagedWorkspace repo, ILogger logger)
 		{
-			using IPerforceConnection PerforceClient = await Perforce.WithClientAsync(ClientName);
-			await Repo.SetupAsync(PerforceClient, StreamName, CancellationToken.None);
+			using IPerforceConnection perforceClient = await perforce.WithClientAsync(ClientName);
+			await repo.SetupAsync(perforceClient, StreamName, CancellationToken.None);
 		}
 	}
 }
diff --git a/Engine/Source/Programs/Horde/Horde.Agent/Commands/Workspace/SizesCommand.cs b/Engine/Source/Programs/Horde/Horde.Agent/Commands/Workspace/SizesCommand.cs
--- a/Engine/Source/Programs/Horde/Horde.Agent/Commands/Workspace/SizesCommand.cs
+++ b/Engine/Source/Programs/Horde/Horde.Agent/Commands/Workspace/SizesCommand.cs
@@ -19,20 +19,20 @@ namespace Horde.Agent.Commands.Workspace
 	{
 		[CommandLine("-TempClient=", Required = true)]
 		[Description("Name of a temporary client to switch between streams gathering metadata. Will be created if it does not exist.")]
-		string TempClientName = null!;
+		string TempClientName { get; set; } = null!;
 
 		[CommandLine("-Stream=")]
 		[Description("Streams that should be included in the output")]
-		List<string> StreamNames = new List<string>();
+		List<string> StreamNames { get; set; } = new List<string>();
 
 		[CommandLine("-Filter=")]
 		[Description("Filters for the files to sync, in P4 syntax (eg. /Engine/...)")]
-		List<string> Filters = new List<string>();
+		List<string> Filters { get; set; } = new List<string>();
 
-		protected override async Task ExecuteAsync(IPerforceConnection Perforce, ManagedWorkspace Repo, ILogger Logger)
+		protected override async Task ExecuteAsync(IPerforceConnection perforce, ManagedWorkspace repo, ILogger logger)
 		{
-			using IPerforceConnection PerforceClient = await Perforce.WithClientAsync(TempClientName);
-			await Repo.StatsAsync(PerforceClient, StreamNames, Filters, CancellationToken.None);
+			using IPerforceConnection perforceClient = await perforce.WithClientAsync(TempClientName);
+			await repo.StatsAsync(perforceClient, StreamNames, Filters, CancellationToken.None);
 		}
 	}
 }
diff --git a/Engine/Source/Programs/Horde/Horde.Agent/Commands/Workspace/StatusCommand.cs b/Engine/Source/Programs/Horde/Horde.Agent/Commands/Workspace/StatusCommand.cs
--- a/Engine/Source/Programs/Horde/Horde.Agent/Commands/Workspace/StatusCommand.cs
+++ b/Engine/Source/Programs/Horde/Horde.Agent/Commands/Workspace/StatusCommand.cs
@@ -15,9 +15,9 @@ namespace Horde.Agent.Commands.Workspace
 	[Command("Workspace", "Status", "Prints information about the state of the cache and workspace")]
 	class StatusCommand : WorkspaceCommand
 	{
-		protected override Task ExecuteAsync(IPerforceConnection Perforce, ManagedWorkspace Repo, ILogger Logger)
+		protected override Task ExecuteAsync(IPerforceConnection perforce, ManagedWorkspace repo, ILogger logger)
 		{
-			Repo.Status();
+			repo.Status();
 			return Task.CompletedTask;
 		}
 	}
diff --git a/Engine/Source/Programs/Horde/Horde.Agent/Commands/Workspace/SyncCommand.cs b/Engine/Source/Programs/Horde/Horde.Agent/Commands/Workspace/SyncCommand.cs
--- a/Engine/Source/Programs/Horde/Horde.Agent/Commands/Workspace/SyncCommand.cs
+++ b/Engine/Source/Programs/Horde/Horde.Agent/Commands/Workspace/SyncCommand.cs
@@ -19,47 +19,47 @@ namespace Horde.Agent.Commands.Workspace
 	{
 		[CommandLine("-Client=", Required = true)]
 		[Description("Name of the client to sync. Will be created if it does not exist.")]
-		public string ClientName = null!;
+		public string ClientName { get; set; } = null!;
 
 		[CommandLine("-Stream=", Required = true)]
 		[Description("The stream to sync")]
-		public string StreamName = null!;
+		public string StreamName { get; set; } = null!;
 
 		[CommandLine("-Change=")]
 		[Description("The change to sync. May be a changelist number, or 'Latest'")]
-		public string Change = "Latest";
+		public string Change { get; set; } = "Latest";
 
 		[CommandLine("-Preflight=")]
 		[Description("The change to unshelve into the workspace")]
-		public int PreflightChange = -1;
+		public int PreflightChange { get; set; } = -1;
 
 		[CommandLine]
 		[Description("Optional path to a cache file used to store workspace metadata. Using a location on a network share allows multiple machines syncing the same CL to only query Perforce state once.")]
-		FileReference? CacheFile = null;
+		FileReference? CacheFile { get; set; } = null;
 
 		[CommandLine("-Filter=")]
 		[Description("Filters for the files to sync, in P4 syntax (eg. /Engine/...)")]
-		public List<string> Filters = new List<string>();
+		public List<string> Filters { get; set; } = new List<string>();
 
 		[CommandLine("-Incremental")]
 		[Description("Performs an incremental sync, without removing intermediates")]
-		public bool bIncrementalSync = false;
+		public bool IncrementalSync { get; set; } = false;
 
 		[CommandLine("-FakeSync")]
 		[Description("Simulates the sync without actually fetching any files")]
-		public bool bFakeSync = false;
+		public bool FakeSync { get; set; } = false;
 
-		protected override async Task ExecuteAsync(IPerforceConnection Perforce, ManagedWorkspace Repo, ILogger Logger)
+		protected override async Task ExecuteAsync(IPerforceConnection perforce, ManagedWorkspace repo, ILogger logger)
 		{
-			int ChangeNumber = ParseChangeNumberOrLatest(Change);
-			List<string> ExpandedFilters = ExpandFilters(Filters);
+			int changeNumber = ParseChangeNumberOrLatest(Change);
+			List<string> expandedFilters = ExpandFilters(Filters);
 
-			using IPerforceConnection PerforceClient = await Perforce.WithClientAsync(ClientName);
-			await Repo.SyncAsync(PerforceClient, StreamName, ChangeNumber, ExpandedFilters, !bIncrementalSync, bFakeSync, CacheFile, CancellationToken.None);
+			using IPerforceConnection perforceClient = await perforce.WithClientAsync(ClientName);
+			await repo.SyncAsync(perforceClient, StreamName, changeNumber, expandedFilters, !IncrementalSync, FakeSync, CacheFile, CancellationToken.None);
 
 			if (PreflightChange != -1)
 			{
-				await Repo.UnshelveAsync(PerforceClient, StreamName, PreflightChange, CancellationToken.None);
+				await repo.UnshelveAsync(perforceClient, StreamName, PreflightChange, CancellationToken.None);
 			}
 		}
 	}
