Horde: Modify agent code to conform to NET coding style.

diff --git a/Engine/Source/Programs/Horde/Horde.Agent/Parser/Matchers/LogChannelEventMatcher.cs b/Engine/Source/Programs/Horde/Horde.Agent/Parser/Matchers/LogChannelEventMatcher.cs
--- a/Engine/Source/Programs/Horde/Horde.Agent/Parser/Matchers/LogChannelEventMatcher.cs
+++ b/Engine/Source/Programs/Horde/Horde.Agent/Parser/Matchers/LogChannelEventMatcher.cs
@@ -20,42 +20,34 @@ namespace Horde.Agent.Parser.Matchers
 	class LogChannelEventMatcher : ILogEventMatcher
 	{
 		/// <inheritdoc/>
-		public LogEventMatch? Match(ILogCursor Input)
+		public LogEventMatch? Match(ILogCursor input)
 		{
-			Match? Match;
-			if(Input.TryMatch(@"^(\s*)(?<channel>[a-zA-Z_][a-zA-Z0-9_]*):\s*(?<severity>Error|Warning|Display): ", out Match))
+			Match? match;
+			if(input.TryMatch(@"^(\s*)(?<channel>[a-zA-Z_][a-zA-Z0-9_]*):\s*(?<severity>Error|Warning|Display): ", out match))
 			{
-				LogEventBuilder Builder = new LogEventBuilder(Input);
-				Builder.Annotate(Match.Groups["channel"], LogEventMarkup.Channel);
-				Builder.Annotate(Match.Groups["severity"], LogEventMarkup.Severity);
+				LogEventBuilder builder = new LogEventBuilder(input);
+				builder.Annotate(match.Groups["channel"], LogEventMarkup.Channel);
+				builder.Annotate(match.Groups["severity"], LogEventMarkup.Severity);
 				
-				if (Builder.Next.TryMatch(@"^\s+", out Match? Indent))
+				if (builder.Next.TryMatch(@"^\s+", out Match? indent))
 				{						
-					while (Builder.Next.IsMatch($"^{Indent.Value}"))
+					while (builder.Next.IsMatch($"^{indent.Value}"))
 					{
-						Builder.MoveNext();
+						builder.MoveNext();
 					}					
 				}
 
-				LogLevel Level;
-				switch(Match.Groups["severity"].Value)
+				LogLevel level = match.Groups["severity"].Value switch
 				{
-					case "Error":
-						Level = LogLevel.Error;
-						break;
-					case "Warning":
-						Level = LogLevel.Warning;
-						break;
-					default:
-						Level = LogLevel.Information;
-						break;
-				}
+					"Error" => LogLevel.Error,
+					"Warning" => LogLevel.Warning,
+					_ => LogLevel.Information,
+				};
 
-				return Builder.ToMatch(LogEventPriority.Low, Level, KnownLogEvents.Engine_LogChannel);
+				return builder.ToMatch(LogEventPriority.Low, level, KnownLogEvents.Engine_LogChannel);
 				
 			}
 			return null;
 		}
-
 	}
 }
diff --git a/Engine/Source/Programs/Horde/Horde.Agent/Parser/Matchers/MicrosoftEventMatcher.cs b/Engine/Source/Programs/Horde/Horde.Agent/Parser/Matchers/MicrosoftEventMatcher.cs
--- a/Engine/Source/Programs/Horde/Horde.Agent/Parser/Matchers/MicrosoftEventMatcher.cs
+++ b/Engine/Source/Programs/Horde/Horde.Agent/Parser/Matchers/MicrosoftEventMatcher.cs
@@ -16,43 +16,43 @@ namespace Horde.Agent.Parser.Matchers
 	/// </summary>
 	class MicrosoftEventMatcher : ILogEventMatcher
 	{
-		ILogContext Context;
+		private readonly ILogContext _context;
 
-		public MicrosoftEventMatcher(ILogContext Context)
+		public MicrosoftEventMatcher(ILogContext context)
 		{
-			this.Context = Context;
+			_context = context;
 		}
 
-		public LogEventMatch? Match(ILogCursor Cursor)
+		public LogEventMatch? Match(ILogCursor cursor)
 		{
 			// filename(line# [, column#]) | toolname} : [ any text ] {error | warning} code+number:localizable string [ any text ]
 
-			Match? Match;
-			if (Cursor.TryMatch(@"(?<severity>(?:error|warning)) (?<code>[a-zA-Z]+[0-9]+)\s*:", out Match))
+			Match? match;
+			if (cursor.TryMatch(@"(?<severity>(?:error|warning)) (?<code>[a-zA-Z]+[0-9]+)\s*:", out match))
 			{
-				string Prefix = Cursor.CurrentLine!.Substring(0, Match.Index);
+				string prefix = cursor.CurrentLine!.Substring(0, match.Index);
 
-				Match FileOrToolMatch = Regex.Match(Prefix, @"^\s*(.*[^\s])\s*:");
-				if (FileOrToolMatch.Success)
+				Match fileOrToolMatch = Regex.Match(prefix, @"^\s*(.*[^\s])\s*:");
+				if (fileOrToolMatch.Success)
 				{
-					LogEventBuilder Builder = new LogEventBuilder(Cursor);
+					LogEventBuilder builder = new LogEventBuilder(cursor);
 
-					Match FileMatch = Regex.Match(FileOrToolMatch.Value, @"^\s*(?<file>.*)\((?<line>\d+)(?:, (?<column>\d+))?\)\s*:$");
-					if (FileMatch.Success)
+					Match fileMatch = Regex.Match(fileOrToolMatch.Value, @"^\s*(?<file>.*)\((?<line>\d+)(?:, (?<column>\d+))?\)\s*:$");
+					if (fileMatch.Success)
 					{
-						Builder.AnnotateSourceFile(FileMatch.Groups["file"], Context, "");
-						Builder.Annotate(FileMatch.Groups["line"], LogEventMarkup.LineNumber);
-						Builder.TryAnnotate(FileMatch.Groups["column"]);
+						builder.AnnotateSourceFile(fileMatch.Groups["file"], _context, "");
+						builder.Annotate(fileMatch.Groups["line"], LogEventMarkup.LineNumber);
+						builder.TryAnnotate(fileMatch.Groups["column"]);
 					}
 					else
 					{
-						Builder.Annotate("tool", FileOrToolMatch.Groups[1], LogEventMarkup.ToolName);
+						builder.Annotate("tool", fileOrToolMatch.Groups[1], LogEventMarkup.ToolName);
 					}
 
-					Group Severity = Match.Groups["severity"];
-					Builder.Annotate(Severity);
-					Builder.Annotate(Match.Groups["code"]);
-					return Builder.ToMatch(LogEventPriority.Normal, Severity.Value.Equals("error")? LogLevel.Error : LogLevel.Warning, KnownLogEvents.Microsoft);
+					Group severity = match.Groups["severity"];
+					builder.Annotate(severity);
+					builder.Annotate(match.Groups["code"]);
+					return builder.ToMatch(LogEventPriority.Normal, severity.Value.Equals("error")? LogLevel.Error : LogLevel.Warning, KnownLogEvents.Microsoft);
 				}
 			}
 			return null;
diff --git a/Engine/Source/Programs/Horde/Horde.Agent/Parser/Matchers/SourceFileLineEventMatcher.cs b/Engine/Source/Programs/Horde/Horde.Agent/Parser/Matchers/SourceFileLineEventMatcher.cs
--- a/Engine/Source/Programs/Horde/Horde.Agent/Parser/Matchers/SourceFileLineEventMatcher.cs
+++ b/Engine/Source/Programs/Horde/Horde.Agent/Parser/Matchers/SourceFileLineEventMatcher.cs
@@ -33,46 +33,46 @@ namespace Horde.Agent.Parser.Matchers
 		const string LinePattern =
 			@"(?<line>\d+)";
 
-		ILogContext Context;
+		readonly ILogContext _context;
 
-		public SourceFileLineEventMatcher(ILogContext Context)
+		public SourceFileLineEventMatcher(ILogContext context)
 		{
-			this.Context = Context;
+			_context = context;
 		}
 
 		/// <inheritdoc/>
-		public LogEventMatch? Match(ILogCursor Input)
+		public LogEventMatch? Match(ILogCursor input)
 		{
-			Match? Match;
-			if (Input.TryMatch($"^\\s*{SeverityPattern}: {FilePattern}(?:\\({LinePattern}\\))?: ", out Match))
+			Match? match;
+			if (input.TryMatch($"^\\s*{SeverityPattern}: {FilePattern}(?:\\({LinePattern}\\))?: ", out match))
 			{
-				LogLevel Level = GetLogLevelFromSeverity(Match);
+				LogLevel level = GetLogLevelFromSeverity(match);
 
-				LogEventBuilder Builder = new LogEventBuilder(Input);
+				LogEventBuilder builder = new LogEventBuilder(input);
 
-				Builder.AnnotateSourceFile(Match.Groups["file"], Context, "");
-				Builder.Annotate(Match.Groups["severity"], LogEventMarkup.Severity);
-				Builder.TryAnnotate(Match.Groups["line"], LogEventMarkup.LineNumber);
+				builder.AnnotateSourceFile(match.Groups["file"], _context, "");
+				builder.Annotate(match.Groups["severity"], LogEventMarkup.Severity);
+				builder.TryAnnotate(match.Groups["line"], LogEventMarkup.LineNumber);
 
-				EventId EventId;
-				if (Input.IsMatch("copyright"))
+				EventId eventId;
+				if (input.IsMatch("copyright"))
 				{
-					EventId = KnownLogEvents.AutomationTool_MissingCopyright;
+					eventId = KnownLogEvents.AutomationTool_MissingCopyright;
 				}
 				else
 				{
-					EventId = KnownLogEvents.AutomationTool_SourceFileLine;
+					eventId = KnownLogEvents.AutomationTool_SourceFileLine;
 				}
 
-				return Builder.ToMatch(LogEventPriority.AboveNormal, Level, EventId);
+				return builder.ToMatch(LogEventPriority.AboveNormal, level, eventId);
 			}
 			return null;
 		}
 
-		static LogLevel GetLogLevelFromSeverity(Match Match)
+		static LogLevel GetLogLevelFromSeverity(Match match)
 		{
-			string Severity = Match.Groups["severity"].Value;
-			if (Severity.Equals("WARNING", StringComparison.Ordinal))
+			string severity = match.Groups["severity"].Value;
+			if (severity.Equals("WARNING", StringComparison.Ordinal))
 			{
 				return LogLevel.Warning;
 			}
diff --git a/Engine/Source/Programs/Horde/Horde.Agent/Parser/Matchers/SystemicEventMatcher.cs b/Engine/Source/Programs/Horde/Horde.Agent/Parser/Matchers/SystemicEventMatcher.cs
--- a/Engine/Source/Programs/Horde/Horde.Agent/Parser/Matchers/SystemicEventMatcher.cs
+++ b/Engine/Source/Programs/Horde/Horde.Agent/Parser/Matchers/SystemicEventMatcher.cs
@@ -16,20 +16,20 @@ namespace Horde.Agent.Parser.Matchers
 {
 	class SystemicEventMatcher : ILogEventMatcher
 	{
-		public LogEventMatch? Match(ILogCursor Cursor)
+		public LogEventMatch? Match(ILogCursor cursor)
 		{
-			if (Cursor.IsMatch(@"LogDerivedDataCache: .*queries/writes will be limited"))
+			if (cursor.IsMatch(@"LogDerivedDataCache: .*queries/writes will be limited"))
 			{
-				return new LogEventBuilder(Cursor).ToMatch(LogEventPriority.High, LogLevel.Information, KnownLogEvents.Systemic_SlowDDC);
+				return new LogEventBuilder(cursor).ToMatch(LogEventPriority.High, LogLevel.Information, KnownLogEvents.Systemic_SlowDDC);
 			}
-			if (Cursor.IsMatch(@"^\s*ERROR: Error: EC_OK"))
+			if (cursor.IsMatch(@"^\s*ERROR: Error: EC_OK"))
 			{
-				LogEventBuilder Builder = new LogEventBuilder(Cursor);
-				while (Builder.Next.IsMatch(@"^\s*ERROR:\s*$"))
+				LogEventBuilder builder = new LogEventBuilder(cursor);
+				while (builder.Next.IsMatch(@"^\s*ERROR:\s*$"))
 				{
-					Builder.MoveNext();
+					builder.MoveNext();
 				}
-				return new LogEventBuilder(Cursor).ToMatch(LogEventPriority.High, LogLevel.Information, KnownLogEvents.Systemic_PdbUtil);
+				return new LogEventBuilder(cursor).ToMatch(LogEventPriority.High, LogLevel.Information, KnownLogEvents.Systemic_PdbUtil);
 			}
 			return null;
 		}
diff --git a/Engine/Source/Programs/Horde/Horde.Agent/Parser/Matchers/XoreaxEventMatcher.cs b/Engine/Source/Programs/Horde/Horde.Agent/Parser/Matchers/XoreaxEventMatcher.cs
--- a/Engine/Source/Programs/Horde/Horde.Agent/Parser/Matchers/XoreaxEventMatcher.cs
+++ b/Engine/Source/Programs/Horde/Horde.Agent/Parser/Matchers/XoreaxEventMatcher.cs
@@ -16,42 +16,42 @@ namespace Horde.Agent.Parser.Matchers
 {
 	class XoreaxEventMatcher : ILogEventMatcher
 	{
-		public LogEventMatch? Match(ILogCursor Cursor)
+		public LogEventMatch? Match(ILogCursor cursor)
 		{
-			if (Cursor.IsMatch(@"\(BuildService.exe\) is not running"))
+			if (cursor.IsMatch(@"\(BuildService.exe\) is not running"))
 			{
-				LogEventBuilder Builder = new LogEventBuilder(Cursor);
-				return Builder.ToMatch(LogEventPriority.High, LogLevel.Information, KnownLogEvents.Systemic_Xge_ServiceNotRunning);
+				LogEventBuilder builder = new LogEventBuilder(cursor);
+				return builder.ToMatch(LogEventPriority.High, LogLevel.Information, KnownLogEvents.Systemic_Xge_ServiceNotRunning);
 			}
 
-			if (Cursor.IsMatch(@"BUILD FAILED: (.*)xgConsole\.exe(.*)"))
+			if (cursor.IsMatch(@"BUILD FAILED: (.*)xgConsole\.exe(.*)"))
 			{
-				LogEventBuilder Builder = new LogEventBuilder(Cursor);
-				return Builder.ToMatch(LogEventPriority.High, LogLevel.Error, KnownLogEvents.Systemic_Xge_BuildFailed);
+				LogEventBuilder builder = new LogEventBuilder(cursor);
+				return builder.ToMatch(LogEventPriority.High, LogLevel.Error, KnownLogEvents.Systemic_Xge_BuildFailed);
 			}
 
-			if (Cursor.IsMatch(@"^\s*--------------------Build System Warning[- ]"))
+			if (cursor.IsMatch(@"^\s*--------------------Build System Warning[- ]"))
 			{
-				LogEventBuilder Builder = new LogEventBuilder(Cursor);
-				if (Builder.Next.TryMatch(@"^(\s*)([^ ].*):", out Match? Prefix))
+				LogEventBuilder builder = new LogEventBuilder(cursor);
+				if (builder.Next.TryMatch(@"^(\s*)([^ ].*):", out Match? prefix))
 				{
-					Builder.MoveNext();
+					builder.MoveNext();
 
-					string Message = Prefix.Groups[2].Value;
+					string message = prefix.Groups[2].Value;
 
-					EventId EventId = KnownLogEvents.Systemic_Xge;
-					if (Regex.IsMatch(Message, "Failed to connect to Coordinator"))
+					EventId eventId = KnownLogEvents.Systemic_Xge;
+					if (Regex.IsMatch(message, "Failed to connect to Coordinator"))
 					{
-						EventId = KnownLogEvents.Systemic_Xge_Standalone;
+						eventId = KnownLogEvents.Systemic_Xge_Standalone;
 					}
 
-					string PrefixPattern = $"^{Prefix.Groups[1].Value}\\s";
-					while (Builder.Next.IsMatch(PrefixPattern))
+					string prefixPattern = $"^{prefix.Groups[1].Value}\\s";
+					while (builder.Next.IsMatch(prefixPattern))
 					{
-						Builder.MoveNext();
+						builder.MoveNext();
 					}
 
-					return Builder.ToMatch(LogEventPriority.High, LogLevel.Information, EventId);
+					return builder.ToMatch(LogEventPriority.High, LogLevel.Information, eventId);
 				}
 			}
 			return null;
diff --git a/Engine/Source/Programs/Horde/Horde.Agent/Program.cs b/Engine/Source/Programs/Horde/Horde.Agent/Program.cs
--- a/Engine/Source/Programs/Horde/Horde.Agent/Program.cs
+++ b/Engine/Source/Programs/Horde/Horde.Agent/Program.cs
@@ -65,15 +65,15 @@ namespace Horde.Agent
 		/// <summary>
 		/// Entry point
 		/// </summary>
-		/// <param name="Args">Command-line arguments</param>
+		/// <param name="args">Command-line arguments</param>
 		/// <returns>Exit code</returns>
-		public static async Task<int> Main(string[] Args)
+		public static async Task<int> Main(string[] args)
 		{
-			Program.Args = Args;
+			Program.Args = args;
 
-			IServiceCollection Services = new ServiceCollection();
-			Services.AddCommandsFromAssembly(Assembly.GetExecutingAssembly());
-			Services.AddLogging(Builder => Builder.AddProvider(new Logging.HordeLoggerProvider()));
+			IServiceCollection services = new ServiceCollection();
+			services.AddCommandsFromAssembly(Assembly.GetExecutingAssembly());
+			services.AddLogging(builder => builder.AddProvider(new Logging.HordeLoggerProvider()));
 
 			// Enable unencrypted HTTP/2 for gRPC channel without TLS
 			AppContext.SetSwitch("System.Net.Http.SocketsHttpHandler.Http2UnencryptedSupport", true);
@@ -83,15 +83,15 @@ namespace Horde.Agent
 				// Prioritize agent execution time over any job its running.
 				// We've seen file copying starving the agent communication to the Horde server, causing a disconnect.
 				// Increasing the process priority is speculative fix to combat this.
-				using (Process Process = Process.GetCurrentProcess())
+				using (Process process = Process.GetCurrentProcess())
 				{
-					Process.PriorityClass = ProcessPriorityClass.High;
+					process.PriorityClass = ProcessPriorityClass.High;
 				}
 			}
 
 			// Execute all the commands
-			IServiceProvider ServiceProvider = Services.BuildServiceProvider();
-			return await CommandHost.RunAsync(new CommandLineArguments(Args), ServiceProvider, typeof(Horde.Agent.Modes.Service.RunCommand));
+			IServiceProvider serviceProvider = services.BuildServiceProvider();
+			return await CommandHost.RunAsync(new CommandLineArguments(args), serviceProvider, typeof(Horde.Agent.Modes.Service.RunCommand));
 		}
 
 		/// <summary>
@@ -127,10 +127,10 @@ namespace Horde.Agent
 		{
 			if (RuntimeInformation.IsOSPlatform(OSPlatform.Windows))
 			{
-				DirectoryReference? ProgramDataDir = DirectoryReference.GetSpecialFolder(Environment.SpecialFolder.CommonApplicationData);
-				if (ProgramDataDir != null)
+				DirectoryReference? programDataDir = DirectoryReference.GetSpecialFolder(Environment.SpecialFolder.CommonApplicationData);
+				if (programDataDir != null)
 				{
-					return DirectoryReference.Combine(ProgramDataDir, "HordeAgent");
+					return DirectoryReference.Combine(programDataDir, "HordeAgent");
 				}
 			}
 			return GetAppDir();
diff --git a/Engine/Source/Programs/Horde/Horde.Agent/Services/GrpcService.cs b/Engine/Source/Programs/Horde/Horde.Agent/Services/GrpcService.cs
--- a/Engine/Source/Programs/Horde/Horde.Agent/Services/GrpcService.cs
+++ b/Engine/Source/Programs/Horde/Horde.Agent/Services/GrpcService.cs
@@ -22,31 +22,18 @@ namespace Horde.Agent.Services
 	/// </summary>
 	class GrpcService
 	{
-		/// <summary>
-		/// Settings for the current server
-		/// </summary>
-		ServerProfile ServerProfile;
-
-		/// <summary>
-		/// Options for the agent
-		/// </summary>
-		IOptions<AgentSettings> Settings;
-
-		/// <summary>
-		/// Logger instance
-		/// </summary>
-		ILogger Logger;
+		private readonly ServerProfile _serverProfile;
+		private readonly ILogger _logger;
 
 		/// <summary>
 		/// Constructor
 		/// </summary>
-		/// <param name="Settings"></param>
-		/// <param name="Logger"></param>
-		public GrpcService(IOptions<AgentSettings> Settings, ILogger<GrpcService> Logger)
+		/// <param name="settings"></param>
+		/// <param name="logger"></param>
+		public GrpcService(IOptions<AgentSettings> settings, ILogger<GrpcService> logger)
 		{
-			this.Settings = Settings;
-			this.ServerProfile = Settings.Value.GetCurrentServerProfile();
-			this.Logger = Logger;
+			_serverProfile = settings.Value.GetCurrentServerProfile();
+			_logger = logger;
 		}
 
 		/// <summary>
@@ -55,22 +42,22 @@ namespace Horde.Agent.Services
 		/// <returns>New grpc channel</returns>
 		public GrpcChannel CreateGrpcChannel()
 		{
-			return CreateGrpcChannel(ServerProfile.Token);
+			return CreateGrpcChannel(_serverProfile.Token);
 		}
 		
 		/// <summary>
 		/// Create a GRPC channel with the given bearer token
 		/// </summary>
 		/// <returns>New grpc channel</returns>
-		public GrpcChannel CreateGrpcChannel(string? BearerToken)
+		public GrpcChannel CreateGrpcChannel(string? bearerToken)
 		{
-			if (BearerToken == null)
+			if (bearerToken == null)
 			{
-				return CreateGrpcChannel(ServerProfile.Url, null);
+				return CreateGrpcChannel(_serverProfile.Url, null);
 			}
 			else
 			{
-				return CreateGrpcChannel(ServerProfile.Url, new AuthenticationHeaderValue("Bearer", BearerToken));
+				return CreateGrpcChannel(_serverProfile.Url, new AuthenticationHeaderValue("Bearer", bearerToken));
 			}
 		}
 
@@ -78,33 +65,33 @@ namespace Horde.Agent.Services
 		/// Create a GRPC channel with the given auth header value
 		/// </summary>
 		/// <returns>New grpc channel</returns>
-		public GrpcChannel CreateGrpcChannel(string Address, AuthenticationHeaderValue? AuthHeaderValue)
+		public GrpcChannel CreateGrpcChannel(string address, AuthenticationHeaderValue? authHeaderValue)
 		{
-			HttpClientHandler CustomCertHandler = new HttpClientHandler();
-			CustomCertHandler.ServerCertificateCustomValidationCallback += (Sender, Cert, Chain, Errors) => CertificateHelper.CertificateValidationCallBack(Logger, Sender, Cert, Chain, Errors, ServerProfile);
+			HttpClientHandler customCertHandler = new HttpClientHandler();
+			customCertHandler.ServerCertificateCustomValidationCallback += (sender, cert, chain, errors) => CertificateHelper.CertificateValidationCallBack(_logger, sender, cert, chain, errors, _serverProfile);
 
-			TimeSpan[] RetryDelay = { TimeSpan.FromSeconds(1), TimeSpan.FromSeconds(5), TimeSpan.FromSeconds(10) };
-			IAsyncPolicy<HttpResponseMessage> Policy = HttpPolicyExtensions.HandleTransientHttpError().WaitAndRetryAsync(RetryDelay);
-			PolicyHttpMessageHandler RetryHandler = new PolicyHttpMessageHandler(Policy);
-			RetryHandler.InnerHandler = CustomCertHandler;
+			TimeSpan[] retryDelay = { TimeSpan.FromSeconds(1), TimeSpan.FromSeconds(5), TimeSpan.FromSeconds(10) };
+			IAsyncPolicy<HttpResponseMessage> policy = HttpPolicyExtensions.HandleTransientHttpError().WaitAndRetryAsync(retryDelay);
+			PolicyHttpMessageHandler retryHandler = new PolicyHttpMessageHandler(policy);
+			retryHandler.InnerHandler = customCertHandler;
 
-			HttpClient HttpClient = new HttpClient(RetryHandler);
-			HttpClient.DefaultRequestHeaders.Add("Accept", "application/json");
-			if (AuthHeaderValue != null)
+			HttpClient httpClient = new HttpClient(retryHandler);
+			httpClient.DefaultRequestHeaders.Add("Accept", "application/json");
+			if (authHeaderValue != null)
 			{
-				HttpClient.DefaultRequestHeaders.Authorization = AuthHeaderValue;
+				httpClient.DefaultRequestHeaders.Authorization = authHeaderValue;
 			}
 
-			HttpClient.Timeout = TimeSpan.FromSeconds(210); // Need to make sure this doesn't cancel any long running gRPC streaming calls (eg. session update)
+			httpClient.Timeout = TimeSpan.FromSeconds(210); // Need to make sure this doesn't cancel any long running gRPC streaming calls (eg. session update)
 
-			Logger.LogInformation("Connecting to rpc server {BaseUrl}", ServerProfile.Url);
-			return GrpcChannel.ForAddress(Address, new GrpcChannelOptions
+			_logger.LogInformation("Connecting to rpc server {BaseUrl}", _serverProfile.Url);
+			return GrpcChannel.ForAddress(address, new GrpcChannelOptions
 			{
 				// Required payloads coming from CAS service can be large
 				MaxReceiveMessageSize = 1024 * 1024 * 1024, // 1 GB
 				MaxSendMessageSize = 1024 * 1024 * 1024, // 1 GB
 				
-				HttpClient = HttpClient,
+				HttpClient = httpClient,
 				DisposeHttpClient = true
 			});
 		}
