Horde: Modify agent code to conform to NET coding style.

diff --git a/Engine/Source/Programs/Horde/Horde.Agent/Services/WorkerService.cs b/Engine/Source/Programs/Horde/Horde.Agent/Services/WorkerService.cs
--- a/Engine/Source/Programs/Horde/Horde.Agent/Services/WorkerService.cs
+++ b/Engine/Source/Programs/Horde/Horde.Agent/Services/WorkerService.cs
@@ -59,26 +59,26 @@ namespace Horde.Agent.Services
 			/// <summary>
 			/// The worker lease state
 			/// </summary>
-			public Lease Lease;
+			public Lease Lease { get; set; }
 
 			/// <summary>
 			/// The task being executed for this lease
 			/// </summary>
-			public Task? Task;
+			public Task? Task { get; set; }
 
 			/// <summary>
 			/// Source for cancellation tokens for this session.
 			/// </summary>
-			public CancellationTokenSource CancellationTokenSource;
+			public CancellationTokenSource CancellationTokenSource { get; }
 
 			/// <summary>
 			/// Constructor
 			/// </summary>
-			/// <param name="Lease">The worker lease state</param>
-			public LeaseInfo(Lease Lease)
+			/// <param name="lease">The worker lease state</param>
+			public LeaseInfo(Lease lease)
 			{
-				this.Lease = Lease;
-				this.CancellationTokenSource = new CancellationTokenSource();
+				Lease = lease;
+				CancellationTokenSource = new CancellationTokenSource();
 			}
 		}
 
@@ -115,208 +115,208 @@ namespace Horde.Agent.Services
 			/// <summary>
 			/// Constructor for 
 			/// </summary>
-			/// <param name="Outcome"></param>
-			private LeaseResult(LeaseOutcome Outcome)
+			/// <param name="outcome"></param>
+			private LeaseResult(LeaseOutcome outcome)
 			{
-				this.Outcome = Outcome;
+				Outcome = outcome;
 			}
 
 			/// <summary>
 			/// Constructor for successful results
 			/// </summary>
-			/// <param name="Output"></param>
-			public LeaseResult(byte[]? Output)
+			/// <param name="output"></param>
+			public LeaseResult(byte[]? output)
 			{
-				this.Outcome = LeaseOutcome.Success;
-				this.Output = Output;
+				Outcome = LeaseOutcome.Success;
+				Output = output;
 			}
 		}
 
 		/// <summary>
 		/// List of processes that should be terminated before running a job
 		/// </summary>
-		private HashSet<string> ProcessNamesToTerminate = new HashSet<string>(StringComparer.OrdinalIgnoreCase);
+		private readonly HashSet<string> _processNamesToTerminate = new HashSet<string>(StringComparer.OrdinalIgnoreCase);
 
 		/// <summary>
 		/// Object used for controlling access to the access tokens and active sessions list
 		/// </summary>
-		object LockObject = new object();
+		readonly object _lockObject = new object();
 
 		/// <summary>
 		/// Sink for log messages
 		/// </summary>
-		ILogger<WorkerService> Logger;
+		readonly ILogger _logger;
 
 		/// <summary>
 		/// The current settings instance
 		/// </summary>
-		AgentSettings Settings;
+		readonly AgentSettings _settings;
 
 		/// <summary>
 		/// Settings for the current server
 		/// </summary>
-		ServerProfile ServerProfile;
+		readonly ServerProfile _serverProfile;
 
 		/// <summary>
 		/// The grpc service instance
 		/// </summary>
-		GrpcService GrpcService;
+		readonly GrpcService _grpcService;
 
 		/// <summary>
 		/// Client interface to the storage system
 		/// </summary>
-		IStorageClient StorageClient;
+		readonly IStorageClient _storageClient;
 
 		/// <summary>
 		/// The working directory
 		/// </summary>
-		DirectoryReference WorkingDir;
+		readonly DirectoryReference _workingDir;
 
 		/// <summary>
 		/// The list of active leases.
 		/// </summary>
-		List<LeaseInfo> ActiveLeases = new List<LeaseInfo>();
+		readonly List<LeaseInfo> _activeLeases = new List<LeaseInfo>();
 
 		/// <summary>
 		/// Whether the agent is currently in an unhealthy state
 		/// </summary>
-		bool Unhealthy = false;
+		readonly bool _unhealthy = false;
 
 		/// <summary>
 		/// Whether the 
 		/// </summary>
-		bool RequestShutdown = false;
+		bool _requestShutdown = false;
 
 		/// <summary>
 		/// Whether to restart after shutting down
 		/// </summary>
-		bool RestartAfterShutdown = false;
+		bool _restartAfterShutdown = false;
 
 		/// <summary>
 		/// Time the the service started
 		/// </summary>
-		static readonly DateTimeOffset StartTime = DateTimeOffset.Now;
+		static readonly DateTimeOffset s_startTime = DateTimeOffset.Now;
 
 		/// <summary>
 		/// Time at which the computer started
 		/// </summary>
-		static readonly DateTimeOffset BootTime = DateTimeOffset.Now - TimeSpan.FromTicks(Environment.TickCount64 * TimeSpan.TicksPerMillisecond);
+		static readonly DateTimeOffset s_bootTime = DateTimeOffset.Now - TimeSpan.FromTicks(Environment.TickCount64 * TimeSpan.TicksPerMillisecond);
 
 		/// <summary>
 		/// Task completion source used to trigger the background thread to update the leases. Must take a lock on LockObject before 
 		/// </summary>
-		AsyncEvent UpdateLeasesEvent = new AsyncEvent();
+		readonly AsyncEvent _updateLeasesEvent = new AsyncEvent();
 
 		/// <summary>
 		/// Number of times UpdateSession has failed
 		/// </summary>
-		int UpdateSessionFailures;
+		int _updateSessionFailures;
 
 		/// <summary>
 		/// Function for creating a new executor. Primarily done this way to aid testing. 
 		/// </summary>
-		private readonly Func<IRpcConnection, ExecuteJobTask, BeginBatchResponse, IExecutor> CreateExecutor;
+		private readonly Func<IRpcConnection, ExecuteJobTask, BeginBatchResponse, IExecutor> _createExecutor;
 
 		/// <summary>
 		/// How often to poll the server checking if a step has been aborted
 		/// Exposed as internal to ease testing. 
 		/// </summary>
-		internal TimeSpan StepAbortPollInterval = TimeSpan.FromSeconds(5);
+		internal TimeSpan _stepAbortPollInterval = TimeSpan.FromSeconds(5);
 
 		/// <summary>
 		/// Constructor. Registers with the server and starts accepting connections.
 		/// </summary>
-		/// <param name="Logger">Log sink</param>
-		/// <param name="OptionsMonitor">The current settings</param>
-		/// <param name="GrpcService">Instance of the Grpc service</param>
-		/// <param name="StorageClient">Instance of the storage client</param>
-		/// <param name="CreateExecutor"></param>
-		public WorkerService(ILogger<WorkerService> Logger, IOptionsMonitor<AgentSettings> OptionsMonitor, GrpcService GrpcService, IStorageClient StorageClient, Func<IRpcConnection, ExecuteJobTask, BeginBatchResponse, IExecutor>? CreateExecutor = null)
+		/// <param name="logger">Log sink</param>
+		/// <param name="options">The current settings</param>
+		/// <param name="grpcService">Instance of the Grpc service</param>
+		/// <param name="storageClient">Instance of the storage client</param>
+		/// <param name="createExecutor"></param>
+		public WorkerService(ILogger<WorkerService> logger, IOptions<AgentSettings> options, GrpcService grpcService, IStorageClient storageClient, Func<IRpcConnection, ExecuteJobTask, BeginBatchResponse, IExecutor>? createExecutor = null)
 		{
-			this.Logger = Logger;
-			this.Settings = OptionsMonitor.CurrentValue;
-			this.ServerProfile = Settings.GetCurrentServerProfile();
-			this.GrpcService = GrpcService;
-			this.StorageClient = StorageClient;
+			_logger = logger;
+			_settings = options.Value;
+			_serverProfile = _settings.GetCurrentServerProfile();
+			_grpcService = grpcService;
+			_storageClient = storageClient;
 
-			if (Settings.WorkingDir == null)
+			if (_settings.WorkingDir == null)
 			{
 				throw new Exception("WorkingDir is not set. Unable to run service.");
 			}
 
-			DirectoryReference BaseDir = new FileReference(Assembly.GetExecutingAssembly().Location).Directory;
+			DirectoryReference baseDir = new FileReference(Assembly.GetExecutingAssembly().Location).Directory;
 
-			WorkingDir = DirectoryReference.Combine(BaseDir, Settings.WorkingDir);
-			Logger.LogInformation("Using working directory {WorkingDir}", WorkingDir);
-			DirectoryReference.CreateDirectory(WorkingDir);
+			_workingDir = DirectoryReference.Combine(baseDir, _settings.WorkingDir);
+			logger.LogInformation("Using working directory {WorkingDir}", _workingDir);
+			DirectoryReference.CreateDirectory(_workingDir);
 
-			ProcessNamesToTerminate.UnionWith(Settings.ProcessNamesToTerminate);
+			_processNamesToTerminate.UnionWith(_settings.ProcessNamesToTerminate);
 
-			if (CreateExecutor == null)
+			if (createExecutor == null)
 			{
-				this.CreateExecutor = Settings.Executor switch
+				_createExecutor = _settings.Executor switch
 				{
-					ExecutorType.Test => (RpcClient, ExecuteTask, Batch) =>
-						new TestExecutor(RpcClient, ExecuteTask.JobId, ExecuteTask.BatchId, Batch.AgentType),
-					ExecutorType.Local => (RpcClient, ExecuteTask, Batch) =>
-						new LocalExecutor(RpcClient, ExecuteTask.JobId, ExecuteTask.BatchId, Batch.AgentType, Settings.LocalExecutor),
-					ExecutorType.Perforce => (RpcClient, ExecuteTask, Batch) =>
-						new PerforceExecutor(RpcClient, ExecuteTask.JobId, ExecuteTask.BatchId, Batch.AgentType, ExecuteTask.AutoSdkWorkspace, ExecuteTask.Workspace, WorkingDir),
-					_ => throw new InvalidDataException($"Unknown executor type '{Settings.Executor}'")
+					ExecutorType.Test => (rpcClient, executeTask, batch) =>
+						new TestExecutor(rpcClient, executeTask.JobId, executeTask.BatchId, batch.AgentType),
+					ExecutorType.Local => (rpcClient, executeTask, batch) =>
+						new LocalExecutor(rpcClient, executeTask.JobId, executeTask.BatchId, batch.AgentType, _settings.LocalExecutor),
+					ExecutorType.Perforce => (rpcClient, executeTask, batch) =>
+						new PerforceExecutor(rpcClient, executeTask.JobId, executeTask.BatchId, batch.AgentType, executeTask.AutoSdkWorkspace, executeTask.Workspace, _workingDir),
+					_ => throw new InvalidDataException($"Unknown executor type '{_settings.Executor}'")
 				};
 			}
 			else
 			{
-				this.CreateExecutor = CreateExecutor;
+				_createExecutor = createExecutor;
 			}
 		}
 
 		/// <summary>
 		/// Executes the ServerTaskAsync method and swallows the exception for the task being cancelled. This allows waiting for it to terminate.
 		/// </summary>
-		/// <param name="StoppingToken">Indicates that the service is trying to stop</param>
-		protected override async Task ExecuteAsync(CancellationToken StoppingToken)
+		/// <param name="stoppingToken">Indicates that the service is trying to stop</param>
+		protected override async Task ExecuteAsync(CancellationToken stoppingToken)
 		{
 			try
 			{
-				await ExecuteInnerAsync(StoppingToken);
+				await ExecuteInnerAsync(stoppingToken);
 			}
-			catch (Exception Ex)
+			catch (Exception ex)
 			{
-				Logger.LogCritical(Ex, "Unhandled exception");
+				_logger.LogCritical(ex, "Unhandled exception");
 			}
 		}
 
 		/// <summary>
 		/// Background task to cycle access tokens and update the state of the agent with the server.
 		/// </summary>
-		/// <param name="StoppingToken">Indicates that the service is trying to stop</param>
-		async Task ExecuteInnerAsync(CancellationToken StoppingToken)
+		/// <param name="stoppingToken">Indicates that the service is trying to stop</param>
+		async Task ExecuteInnerAsync(CancellationToken stoppingToken)
 		{
 			// Print the server info
-			Logger.LogInformation("Server: {Server}", ServerProfile.Url);
-			Logger.LogInformation("Arguments: {Arguments}", Environment.CommandLine);
+			_logger.LogInformation("Server: {Server}", _serverProfile.Url);
+			_logger.LogInformation("Arguments: {Arguments}", Environment.CommandLine);
 
 			// Show the current client id
-			string Version = Program.Version;
-			Logger.LogInformation("Version: {Version}", Version);
+			string version = Program.Version;
+			_logger.LogInformation("Version: {Version}", version);
 
 			// Keep trying to start an agent session with the server
-			while (!StoppingToken.IsCancellationRequested)
+			while (!stoppingToken.IsCancellationRequested)
 			{
-				Stopwatch SessionTime = Stopwatch.StartNew();
+				Stopwatch sessionTime = Stopwatch.StartNew();
 				try
 				{
-					if (ActiveLeases.Count == 0)
+					if (_activeLeases.Count == 0)
 					{
-						await HandleSessionAsync(StoppingToken);
+						await HandleSessionAsync(stoppingToken);
 					}
 				}
-				catch (Exception Ex)
+				catch (Exception ex)
 				{
-					if (!StoppingToken.IsCancellationRequested || !IsCancellationException(Ex))
+					if (!stoppingToken.IsCancellationRequested || !IsCancellationException(ex))
 					{
-						Logger.LogError(Ex, "Exception while executing session. Restarting.");
+						_logger.LogError(ex, "Exception while executing session. Restarting.");
 					}
 				}
 
@@ -324,44 +324,44 @@ namespace Horde.Agent.Services
 				{
 					await DrainLeasesAsync();
 				}
-				catch (Exception Ex)
+				catch (Exception ex)
 				{
-					Logger.LogError(Ex, "Exception while draining leases. Agent may be in an inconsistent state.");
+					_logger.LogError(ex, "Exception while draining leases. Agent may be in an inconsistent state.");
 				}
 
-				if (RequestShutdown)
+				if (_requestShutdown)
 				{
-					Logger.LogInformation("Initiating shutdown (restart={Restart})", RestartAfterShutdown);
-					if (Shutdown.InitiateShutdown(RestartAfterShutdown, Logger))
+					_logger.LogInformation("Initiating shutdown (restart={Restart})", _restartAfterShutdown);
+					if (Shutdown.InitiateShutdown(_restartAfterShutdown, _logger))
 					{
-						for (int Idx = 10; Idx > 0; Idx--)
+						for (int idx = 10; idx > 0; idx--)
 						{
-							Logger.LogInformation("Waiting for shutdown ({Count})", Idx);
+							_logger.LogInformation("Waiting for shutdown ({Count})", idx);
 							await Task.Delay(TimeSpan.FromSeconds(60.0));
 						}
-						Logger.LogInformation("Shutdown aborted.");
+						_logger.LogInformation("Shutdown aborted.");
 					}
-					RequestShutdown = RestartAfterShutdown = false;
+					_requestShutdown = _restartAfterShutdown = false;
 				}
-				else if (SessionTime.Elapsed < TimeSpan.FromSeconds(2.0))
+				else if (sessionTime.Elapsed < TimeSpan.FromSeconds(2.0))
 				{
-					Logger.LogInformation("Waiting 5 seconds before restarting session...");
+					_logger.LogInformation("Waiting 5 seconds before restarting session...");
 					await Task.Delay(TimeSpan.FromSeconds(5.0));
 				}
 			}
 		}
 
-		async Task WaitForMutexAsync(Mutex Mutex, CancellationToken StoppingToken)
+		async Task WaitForMutexAsync(Mutex mutex, CancellationToken stoppingToken)
 		{
 			try
 			{
-				if (!Mutex.WaitOne(0))
+				if (!mutex.WaitOne(0))
 				{
-					Logger.LogError("Another instance of HordeAgent is already running. Waiting until it terminates.");
-					while (!Mutex.WaitOne(0))
+					_logger.LogError("Another instance of HordeAgent is already running. Waiting until it terminates.");
+					while (!mutex.WaitOne(0))
 					{
-						StoppingToken.ThrowIfCancellationRequested();
-						await Task.Delay(TimeSpan.FromSeconds(1.0), StoppingToken);
+						stoppingToken.ThrowIfCancellationRequested();
+						await Task.Delay(TimeSpan.FromSeconds(1.0), stoppingToken);
 					}
 				}
 			}
@@ -372,49 +372,49 @@ namespace Horde.Agent.Services
 
 		async Task DrainLeasesAsync()
 		{
-			for (int Idx = 0; Idx < ActiveLeases.Count; Idx++)
+			for (int idx = 0; idx < _activeLeases.Count; idx++)
 			{
-				LeaseInfo ActiveLease = ActiveLeases[Idx];
-				if (ActiveLease.Task == null)
+				LeaseInfo activeLease = _activeLeases[idx];
+				if (activeLease.Task == null)
 				{
-					ActiveLeases.RemoveAt(Idx--);
-					Logger.LogInformation("Removed lease {LeaseId}", ActiveLease.Lease.Id);
+					_activeLeases.RemoveAt(idx--);
+					_logger.LogInformation("Removed lease {LeaseId}", activeLease.Lease.Id);
 				}
 				else
 				{
-					Logger.LogInformation("Cancelling active lease {LeaseId}", ActiveLease.Lease.Id);
-					ActiveLease.CancellationTokenSource.Cancel();
+					_logger.LogInformation("Cancelling active lease {LeaseId}", activeLease.Lease.Id);
+					activeLease.CancellationTokenSource.Cancel();
 				}
 			}
 
-			while (ActiveLeases.Count > 0)
+			while (_activeLeases.Count > 0)
 			{
-				List<Task> Tasks = ActiveLeases.Select(x => x.Task!).ToList();
-				Tasks.Add(Task.Delay(TimeSpan.FromMinutes(1.0)));
-				await Task.WhenAny(Tasks);
+				List<Task> tasks = _activeLeases.Select(x => x.Task!).ToList();
+				tasks.Add(Task.Delay(TimeSpan.FromMinutes(1.0)));
+				await Task.WhenAny(tasks);
 
-				for (int Idx = 0; Idx < ActiveLeases.Count; Idx++)
+				for (int idx = 0; idx < _activeLeases.Count; idx++)
 				{
-					LeaseInfo ActiveLease = ActiveLeases[Idx];
-					if (ActiveLease.Task!.IsCompleted)
+					LeaseInfo activeLease = _activeLeases[idx];
+					if (activeLease.Task!.IsCompleted)
 					{
-						ActiveLeases.RemoveAt(Idx--);
+						_activeLeases.RemoveAt(idx--);
 						try
 						{
-							await ActiveLease.Task;
+							await activeLease.Task;
 						}
 						catch (OperationCanceledException)
 						{
 						}
-						catch (Exception Ex)
+						catch (Exception ex)
 						{
-							Logger.LogError(Ex, "Lease {LeaseId} threw an exception while terminating", ActiveLease.Lease.Id);
+							_logger.LogError(ex, "Lease {LeaseId} threw an exception while terminating", activeLease.Lease.Id);
 						}
-						Logger.LogInformation("Lease {LeaseId} has completed", ActiveLease.Lease.Id);
+						_logger.LogInformation("Lease {LeaseId} has completed", activeLease.Lease.Id);
 					}
 					else
 					{
-						Logger.LogInformation("Still waiting for lease {LeaseId} to terminate...", ActiveLease.Lease.Id);
+						_logger.LogInformation("Still waiting for lease {LeaseId} to terminate...", activeLease.Lease.Id);
 					}
 				}
 			}
@@ -423,231 +423,231 @@ namespace Horde.Agent.Services
 		/// <summary>
 		/// Handles the lifetime of an agent session.
 		/// </summary>
-		/// <param name="StoppingToken">Indicates that the service is trying to stop</param>
+		/// <param name="stoppingToken">Indicates that the service is trying to stop</param>
 		/// <returns>Async task</returns>
-		async Task HandleSessionAsync(CancellationToken StoppingToken)
+		async Task HandleSessionAsync(CancellationToken stoppingToken)
 		{
 			// Make sure there's only one instance of the agent running
-			using Mutex SingleInstanceMutex = new Mutex(false, "Global\\HordeAgent-DB828ACB-0AA5-4D32-A62A-21D4429B1014");
-			await WaitForMutexAsync(SingleInstanceMutex, StoppingToken);
+			using Mutex singleInstanceMutex = new Mutex(false, "Global\\HordeAgent-DB828ACB-0AA5-4D32-A62A-21D4429B1014");
+			await WaitForMutexAsync(singleInstanceMutex, stoppingToken);
 
 			// Terminate any remaining child processes from other instances
-			TerminateProcesses(Logger, StoppingToken);
+			TerminateProcesses(_logger, stoppingToken);
 
 			// Show the worker capabilities
-			AgentCapabilities Capabilities = await GetAgentCapabilities(WorkingDir, Logger);
-			if (Capabilities.Properties.Count > 0)
+			AgentCapabilities capabilities = await GetAgentCapabilities(_workingDir, _logger);
+			if (capabilities.Properties.Count > 0)
 			{
-				Logger.LogInformation("Global:");
-				foreach (string Property in Capabilities.Properties)
+				_logger.LogInformation("Global:");
+				foreach (string property in capabilities.Properties)
 				{
-					Logger.LogInformation("  {AgentProperty}", Property);
+					_logger.LogInformation("  {AgentProperty}", property);
 				}
 			}
-			foreach (DeviceCapabilities Device in Capabilities.Devices)
+			foreach (DeviceCapabilities device in capabilities.Devices)
 			{
-				Logger.LogInformation($"{Device.Handle} Device:");
-				foreach (string Property in Device.Properties)
+				_logger.LogInformation($"{device.Handle} Device:");
+				foreach (string property in device.Properties)
 				{
-					Logger.LogInformation("   {DeviceProperty}", Property);
+					_logger.LogInformation("   {DeviceProperty}", property);
 				}
 			}
 
 			// Mount all the necessary network shares. Currently only supported on Windows.
-			if (Settings.ShareMountingEnabled && RuntimeInformation.IsOSPlatform(OSPlatform.Windows))
+			if (_settings.ShareMountingEnabled && RuntimeInformation.IsOSPlatform(OSPlatform.Windows))
 			{
-				foreach (MountNetworkShare Share in Settings.Shares)
+				foreach (MountNetworkShare share in _settings.Shares)
 				{
-					if (Share.MountPoint != null && Share.RemotePath != null)
+					if (share.MountPoint != null && share.RemotePath != null)
 					{
-						Logger.LogInformation("Mounting {RemotePath} as {MountPoint}", Share.RemotePath, Share.MountPoint);
-						NetworkShare.Mount(Share.MountPoint, Share.RemotePath);
+						_logger.LogInformation("Mounting {RemotePath} as {MountPoint}", share.RemotePath, share.MountPoint);
+						NetworkShare.Mount(share.MountPoint, share.RemotePath);
 					}
 				}
 			}
 
 			// Create the session
-			CreateSessionResponse CreateSessionResponse;
-			using (GrpcChannel Channel = GrpcService.CreateGrpcChannel(ServerProfile.Token))
+			CreateSessionResponse createSessionResponse;
+			using (GrpcChannel channel = _grpcService.CreateGrpcChannel(_serverProfile.Token))
 			{
-				HordeRpc.HordeRpcClient RpcClient = new HordeRpc.HordeRpcClient(Channel);
+				HordeRpc.HordeRpcClient rpcClient = new HordeRpc.HordeRpcClient(channel);
 
 				// Create the session information
-				CreateSessionRequest SessionRequest = new CreateSessionRequest();
-				SessionRequest.Name = Settings.GetAgentName();
-				SessionRequest.Status = AgentStatus.Ok;
-				SessionRequest.Capabilities = Capabilities;
-				SessionRequest.Version = Program.Version;
+				CreateSessionRequest sessionRequest = new CreateSessionRequest();
+				sessionRequest.Name = _settings.GetAgentName();
+				sessionRequest.Status = AgentStatus.Ok;
+				sessionRequest.Capabilities = capabilities;
+				sessionRequest.Version = Program.Version;
 
 				// Create a session
-				CreateSessionResponse = await RpcClient.CreateSessionAsync(SessionRequest, null, null, StoppingToken);
-				Logger.LogInformation("Session started ({SessionId})", CreateSessionResponse.SessionId);
+				createSessionResponse = await rpcClient.CreateSessionAsync(sessionRequest, null, null, stoppingToken);
+				_logger.LogInformation("Session started ({SessionId})", createSessionResponse.SessionId);
 			}
 
 			// Open a connection to the server
-			await using (IRpcConnection RpcCon = RpcConnection.Create(() => GrpcService.CreateGrpcChannel(CreateSessionResponse.Token), Logger))
+			await using (IRpcConnection rpcCon = RpcConnection.Create(() => _grpcService.CreateGrpcChannel(createSessionResponse.Token), _logger))
 			{
 				// Track how many updates we get in 10 seconds. We'll start rate limiting this if it looks like we've got a problem that's causing us to spam the server.
-				Stopwatch UpdateTimer = Stopwatch.StartNew();
-				Queue<TimeSpan> UpdateTimes = new Queue<TimeSpan>();
+				Stopwatch updateTimer = Stopwatch.StartNew();
+				Queue<TimeSpan> updateTimes = new Queue<TimeSpan>();
 
 				// Loop until we're ready to exit
-				Stopwatch UpdateCapabilitiesTimer = Stopwatch.StartNew();
+				Stopwatch updateCapabilitiesTimer = Stopwatch.StartNew();
 				for (; ; )
 				{
-					Task WaitTask = UpdateLeasesEvent.Task;
+					Task waitTask = _updateLeasesEvent.Task;
 
 					// Flag for whether the service is stopping
-					bool bStopping = StoppingToken.IsCancellationRequested || RequestShutdown;
+					bool stopping = stoppingToken.IsCancellationRequested || _requestShutdown;
 
 					// Build the next update request
-					UpdateSessionRequest UpdateSessionRequest = new UpdateSessionRequest();
-					UpdateSessionRequest.AgentId = CreateSessionResponse.AgentId;
-					UpdateSessionRequest.SessionId = CreateSessionResponse.SessionId;
+					UpdateSessionRequest updateSessionRequest = new UpdateSessionRequest();
+					updateSessionRequest.AgentId = createSessionResponse.AgentId;
+					updateSessionRequest.SessionId = createSessionResponse.SessionId;
 
 					// Get the new the lease states. If a restart is requested and we have no active leases, signal to the server that we're stopping.
-					lock (LockObject)
+					lock (_lockObject)
 					{
-						foreach (LeaseInfo LeaseInfo in ActiveLeases)
+						foreach (LeaseInfo leaseInfo in _activeLeases)
 						{
-							UpdateSessionRequest.Leases.Add(new Lease(LeaseInfo.Lease));
+							updateSessionRequest.Leases.Add(new Lease(leaseInfo.Lease));
 						}
-						if (RequestShutdown && ActiveLeases.Count == 0)
+						if (_requestShutdown && _activeLeases.Count == 0)
 						{
-							bStopping = true;
+							stopping = true;
 						}
 					}
 
 					// Get the new agent status
-					if (bStopping)
+					if (stopping)
 					{
-						UpdateSessionRequest.Status = AgentStatus.Stopping;
+						updateSessionRequest.Status = AgentStatus.Stopping;
 					}
-					else if (Unhealthy)
+					else if (_unhealthy)
 					{
-						UpdateSessionRequest.Status = AgentStatus.Unhealthy;
+						updateSessionRequest.Status = AgentStatus.Unhealthy;
 					}
 					else
 					{
-						UpdateSessionRequest.Status = AgentStatus.Ok;
+						updateSessionRequest.Status = AgentStatus.Ok;
 					}
 
 					// Update the capabilities every 5m
-					if (UpdateCapabilitiesTimer.Elapsed > TimeSpan.FromMinutes(5.0))
+					if (updateCapabilitiesTimer.Elapsed > TimeSpan.FromMinutes(5.0))
 					{
-						UpdateSessionRequest.Capabilities = await GetAgentCapabilities(WorkingDir, Logger);
-						UpdateCapabilitiesTimer.Restart();
+						updateSessionRequest.Capabilities = await GetAgentCapabilities(_workingDir, _logger);
+						updateCapabilitiesTimer.Restart();
 					}
 
 					// Complete the wait task if we subsequently stop
-					using (bStopping ? (CancellationTokenRegistration?)null : StoppingToken.Register(() => UpdateLeasesEvent.Set()))
+					using (stopping ? (CancellationTokenRegistration?)null : stoppingToken.Register(() => _updateLeasesEvent.Set()))
 					{
 						// Update the state with the server
-						UpdateSessionResponse? UpdateSessionResponse = null;
-						using (IRpcClientRef? RpcClientRef = RpcCon.TryGetClientRef(new RpcContext()))
+						UpdateSessionResponse? updateSessionResponse = null;
+						using (IRpcClientRef? rpcClientRef = rpcCon.TryGetClientRef(new RpcContext()))
 						{
-							if (RpcClientRef == null)
+							if (rpcClientRef == null)
 							{
-								await Task.WhenAny(Task.Delay(TimeSpan.FromSeconds(5.0)), WaitTask);
+								await Task.WhenAny(Task.Delay(TimeSpan.FromSeconds(5.0)), waitTask);
 							}
 							else
 							{
-								UpdateSessionResponse = await UpdateSessionAsync(RpcClientRef, UpdateSessionRequest, WaitTask);
+								updateSessionResponse = await UpdateSessionAsync(rpcClientRef, updateSessionRequest, waitTask);
 							}
 						}
 
-						lock (LockObject)
+						lock (_lockObject)
 						{
 							// Now reconcile the local state to match what the server reports
-							if (UpdateSessionResponse != null)
+							if (updateSessionResponse != null)
 							{
 								// Remove any leases which have completed
-								ActiveLeases.RemoveAll(x => (x.Lease.State == LeaseState.Completed || x.Lease.State == LeaseState.Cancelled) && !UpdateSessionResponse.Leases.Any(y => y.Id == x.Lease.Id && y.State != LeaseState.Cancelled));
+								_activeLeases.RemoveAll(x => (x.Lease.State == LeaseState.Completed || x.Lease.State == LeaseState.Cancelled) && !updateSessionResponse.Leases.Any(y => y.Id == x.Lease.Id && y.State != LeaseState.Cancelled));
 
 								// Create any new leases and cancel any running leases
-								foreach (Lease ServerLease in UpdateSessionResponse.Leases)
+								foreach (Lease serverLease in updateSessionResponse.Leases)
 								{
-									if (ServerLease.State == LeaseState.Cancelled)
+									if (serverLease.State == LeaseState.Cancelled)
 									{
-										LeaseInfo? Info = ActiveLeases.FirstOrDefault(x => x.Lease.Id == ServerLease.Id);
-										if (Info != null)
+										LeaseInfo? info = _activeLeases.FirstOrDefault(x => x.Lease.Id == serverLease.Id);
+										if (info != null)
 										{
-											Logger.LogInformation("Cancelling lease {LeaseId}", ServerLease.Id);
-											Info.CancellationTokenSource.Cancel();
+											_logger.LogInformation("Cancelling lease {LeaseId}", serverLease.Id);
+											info.CancellationTokenSource.Cancel();
 										}
 									}
-									if (ServerLease.State == LeaseState.Pending && !ActiveLeases.Any(x => x.Lease.Id == ServerLease.Id))
+									if (serverLease.State == LeaseState.Pending && !_activeLeases.Any(x => x.Lease.Id == serverLease.Id))
 									{
-										ServerLease.State = LeaseState.Active;
+										serverLease.State = LeaseState.Active;
 
-										Logger.LogInformation("Adding lease {LeaseId}", ServerLease.Id);
-										LeaseInfo Info = new LeaseInfo(ServerLease);
-										Info.Task = Task.Run(() => HandleLeaseAsync(RpcCon, CreateSessionResponse.AgentId, Info));
-										ActiveLeases.Add(Info);
+										_logger.LogInformation("Adding lease {LeaseId}", serverLease.Id);
+										LeaseInfo info = new LeaseInfo(serverLease);
+										info.Task = Task.Run(() => HandleLeaseAsync(rpcCon, createSessionResponse.AgentId, info));
+										_activeLeases.Add(info);
 									}
 								}
 							}
 
 							// If there's nothing still running and cancellation was requested, exit
-							if (ActiveLeases.Count == 0 && UpdateSessionRequest.Status == AgentStatus.Stopping)
+							if (_activeLeases.Count == 0 && updateSessionRequest.Status == AgentStatus.Stopping)
 							{
-								Logger.LogInformation("No leases are active. Agent is stopping.");
+								_logger.LogInformation("No leases are active. Agent is stopping.");
 								break;
 							}
 						}
 					}
 
 					// Update the historical update times
-					TimeSpan UpdateTime = UpdateTimer.Elapsed;
-					while (UpdateTimes.TryPeek(out TimeSpan FirstTime) && FirstTime < UpdateTime - TimeSpan.FromMinutes(1.0))
+					TimeSpan updateTime = updateTimer.Elapsed;
+					while (updateTimes.TryPeek(out TimeSpan firstTime) && firstTime < updateTime - TimeSpan.FromMinutes(1.0))
 					{
-						UpdateTimes.Dequeue();
+						updateTimes.Dequeue();
 					}
-					UpdateTimes.Enqueue(UpdateTime);
+					updateTimes.Enqueue(updateTime);
 
 					// If we're updating too much, introduce an artificial delay
-					if (UpdateTimes.Count > 60)
+					if (updateTimes.Count > 60)
 					{
-						Logger.LogWarning("Agent is issuing large number of UpdateSession() calls. Delaying for 10 seconds.");
+						_logger.LogWarning("Agent is issuing large number of UpdateSession() calls. Delaying for 10 seconds.");
 						await Task.Delay(TimeSpan.FromSeconds(10.0));
 					}
 				}
 
-				Logger.LogInformation("Disposing RpcConnection");
+				_logger.LogInformation("Disposing RpcConnection");
 			}
 		}
 
 		/// <summary>
 		/// Wrapper for <see cref="UpdateSessionInternalAsync"/> which filters/logs exceptions
 		/// </summary>
-		/// <param name="RpcClientRef">The RPC client connection</param>
-		/// <param name="UpdateSessionRequest">The session update request</param>
-		/// <param name="WaitTask">Task which can be used to jump out of the update early</param>
+		/// <param name="rpcClientRef">The RPC client connection</param>
+		/// <param name="updateSessionRequest">The session update request</param>
+		/// <param name="waitTask">Task which can be used to jump out of the update early</param>
 		/// <returns>Response from the call</returns>
-		async Task<UpdateSessionResponse?> UpdateSessionAsync(IRpcClientRef RpcClientRef, UpdateSessionRequest UpdateSessionRequest, Task WaitTask)
+		async Task<UpdateSessionResponse?> UpdateSessionAsync(IRpcClientRef rpcClientRef, UpdateSessionRequest updateSessionRequest, Task waitTask)
 		{
-			UpdateSessionResponse? UpdateSessionResponse = null;
+			UpdateSessionResponse? updateSessionResponse = null;
 			try
 			{
-				UpdateSessionResponse = await UpdateSessionInternalAsync(RpcClientRef, UpdateSessionRequest, WaitTask);
-				UpdateSessionFailures = 0;
+				updateSessionResponse = await UpdateSessionInternalAsync(rpcClientRef, updateSessionRequest, waitTask);
+				_updateSessionFailures = 0;
 			}
-			catch (RpcException Ex)
+			catch (RpcException ex)
 			{
-				if (++UpdateSessionFailures >= 3)
+				if (++_updateSessionFailures >= 3)
 				{
 					throw;
 				}
-				else if (Ex.StatusCode == StatusCode.Unavailable)
+				else if (ex.StatusCode == StatusCode.Unavailable)
 				{
-					Logger.LogInformation(Ex, "Service unavailable while calling UpdateSessionAsync(), will retry");
+					_logger.LogInformation(ex, "Service unavailable while calling UpdateSessionAsync(), will retry");
 				}
 				else
 				{
-					Logger.LogError(Ex, "Error while executing RPC. Will retry.");
+					_logger.LogError(ex, "Error while executing RPC. Will retry.");
 				}
 			}
-			return UpdateSessionResponse;
+			return updateSessionResponse;
 		}
 
 		/// <summary>
@@ -658,173 +658,173 @@ namespace Horde.Agent.Services
 		/// until we want to terminate the call (see [REDACTED_PROFILE_URL]/grpc/issues/8277). In order to do that, we need to make a 
 		/// bidirectional streaming call, even though we only expect one response/response.
 		/// </summary>
-		/// <param name="RpcClientRef">The RPC client</param>
-		/// <param name="Request">The session update request</param>
-		/// <param name="WaitTask">Task to use to terminate the wait</param>
+		/// <param name="rpcClientRef">The RPC client</param>
+		/// <param name="request">The session update request</param>
+		/// <param name="waitTask">Task to use to terminate the wait</param>
 		/// <returns>The response object</returns>
-		async Task<UpdateSessionResponse?> UpdateSessionInternalAsync(IRpcClientRef RpcClientRef, UpdateSessionRequest Request, Task WaitTask)
+		async Task<UpdateSessionResponse?> UpdateSessionInternalAsync(IRpcClientRef rpcClientRef, UpdateSessionRequest request, Task waitTask)
 		{
-			DateTime Deadline = DateTime.UtcNow + TimeSpan.FromMinutes(2.0);
-			using (AsyncDuplexStreamingCall<UpdateSessionRequest, UpdateSessionResponse> Call = RpcClientRef.Client.UpdateSession(deadline: Deadline))
+			DateTime deadline = DateTime.UtcNow + TimeSpan.FromMinutes(2.0);
+			using (AsyncDuplexStreamingCall<UpdateSessionRequest, UpdateSessionResponse> call = rpcClientRef.Client.UpdateSession(deadline: deadline))
 			{
-				Logger.LogDebug("Updating session {SessionId} (Status={Status})", Request.SessionId, Request.Status);
+				_logger.LogDebug("Updating session {SessionId} (Status={Status})", request.SessionId, request.Status);
 
 				// Write the request to the server
-				await Call.RequestStream.WriteAsync(Request);
+				await call.RequestStream.WriteAsync(request);
 
 				// Wait until the server responds or we need to trigger a new update
-				Task<bool> MoveNextAsync = Call.ResponseStream.MoveNext();
+				Task<bool> moveNextAsync = call.ResponseStream.MoveNext();
 
-				Task Task = await Task.WhenAny(MoveNextAsync, WaitTask, RpcClientRef.DisposingTask);
-				if(Task == WaitTask)
+				Task task = await Task.WhenAny(moveNextAsync, waitTask, rpcClientRef.DisposingTask);
+				if(task == waitTask)
 				{
-					Logger.LogDebug("Cancelling long poll from client side (new update)");
+					_logger.LogDebug("Cancelling long poll from client side (new update)");
 				}
-				else if (Task == RpcClientRef.DisposingTask)
+				else if (task == rpcClientRef.DisposingTask)
 				{
-					Logger.LogDebug("Cancelling long poll from client side (server migration)");
+					_logger.LogDebug("Cancelling long poll from client side (server migration)");
 				}
 
 				// Close the request stream to indicate that we're finished
-				await Call.RequestStream.CompleteAsync();
+				await call.RequestStream.CompleteAsync();
 
 				// Wait for a response or a new update to come in, then close the request stream
-				UpdateSessionResponse? Response = null;
-				while (await MoveNextAsync)
+				UpdateSessionResponse? response = null;
+				while (await moveNextAsync)
 				{
-					Response = Call.ResponseStream.Current;
-					MoveNextAsync = Call.ResponseStream.MoveNext();
+					response = call.ResponseStream.Current;
+					moveNextAsync = call.ResponseStream.MoveNext();
 				}
-				return Response;
+				return response;
 			}
 		}
 
 		/// <summary>
 		/// Handle a lease request
 		/// </summary>
-		/// <param name="RpcConnection">The RPC connection to the server</param>
-		/// <param name="AgentId">The agent id</param>
-		/// <param name="LeaseInfo">Information about the lease</param>
+		/// <param name="rpcConnection">The RPC connection to the server</param>
+		/// <param name="agentId">The agent id</param>
+		/// <param name="leaseInfo">Information about the lease</param>
 		/// <returns>Async task</returns>
-		async Task HandleLeaseAsync(IRpcConnection RpcConnection, string AgentId, LeaseInfo LeaseInfo)
+		async Task HandleLeaseAsync(IRpcConnection rpcConnection, string agentId, LeaseInfo leaseInfo)
 		{
-			using IScope Scope = GlobalTracer.Instance.BuildSpan("HandleLease").WithResourceName(LeaseInfo.Lease.Id).StartActive();
-			Scope.Span.SetTag("LeaseId", LeaseInfo.Lease.Id);
-			Scope.Span.SetTag("AgentId", AgentId);
+			using IScope scope = GlobalTracer.Instance.BuildSpan("HandleLease").WithResourceName(leaseInfo.Lease.Id).StartActive();
+			scope.Span.SetTag("LeaseId", leaseInfo.Lease.Id);
+			scope.Span.SetTag("AgentId", agentId);
 //			using IDisposable TraceProperty = LogContext.PushProperty("dd.trace_id", CorrelationIdentifier.TraceId.ToString());
 //			using IDisposable SpanProperty = LogContext.PushProperty("dd.span_id", CorrelationIdentifier.SpanId.ToString());
 
-			Logger.LogInformation("Handling lease {LeaseId}", LeaseInfo.Lease.Id);
+			_logger.LogInformation("Handling lease {LeaseId}", leaseInfo.Lease.Id);
 
 			// Get the lease outcome
-			LeaseResult Result = LeaseResult.Failed;
+			LeaseResult result = LeaseResult.Failed;
 			try
 			{
-				Result = await HandleLeasePayloadAsync(RpcConnection, AgentId, LeaseInfo);
+				result = await HandleLeasePayloadAsync(rpcConnection, agentId, leaseInfo);
 			}
-			catch (Exception Ex)
+			catch (Exception ex)
 			{
-				Logger.LogError(Ex, "Unhandled exception while executing lease {LeaseId}", LeaseInfo.Lease.Id);
+				_logger.LogError(ex, "Unhandled exception while executing lease {LeaseId}", leaseInfo.Lease.Id);
 			}
 
 			// Update the state of the lease
-			lock (LockObject)
+			lock (_lockObject)
 			{
-				if (LeaseInfo.CancellationTokenSource.IsCancellationRequested)
+				if (leaseInfo.CancellationTokenSource.IsCancellationRequested)
 				{
-					LeaseInfo.Lease.State = LeaseState.Cancelled;
-					LeaseInfo.Lease.Outcome = LeaseOutcome.Failed;
-					LeaseInfo.Lease.Output = ByteString.Empty;
+					leaseInfo.Lease.State = LeaseState.Cancelled;
+					leaseInfo.Lease.Outcome = LeaseOutcome.Failed;
+					leaseInfo.Lease.Output = ByteString.Empty;
 				}
 				else
 				{
-					LeaseInfo.Lease.State = (Result.Outcome == LeaseOutcome.Cancelled) ? LeaseState.Cancelled : LeaseState.Completed;
-					LeaseInfo.Lease.Outcome = Result.Outcome;
-					LeaseInfo.Lease.Output = (Result.Output != null)? ByteString.CopyFrom(Result.Output) : ByteString.Empty;
+					leaseInfo.Lease.State = (result.Outcome == LeaseOutcome.Cancelled) ? LeaseState.Cancelled : LeaseState.Completed;
+					leaseInfo.Lease.Outcome = result.Outcome;
+					leaseInfo.Lease.Output = (result.Output != null)? ByteString.CopyFrom(result.Output) : ByteString.Empty;
 				}
-				Logger.LogInformation("Transitioning lease {LeaseId} to {State}, outcome={Outcome}", LeaseInfo.Lease.Id, LeaseInfo.Lease.State, LeaseInfo.Lease.Outcome);
+				_logger.LogInformation("Transitioning lease {LeaseId} to {State}, outcome={Outcome}", leaseInfo.Lease.Id, leaseInfo.Lease.State, leaseInfo.Lease.Outcome);
 			}
-			UpdateLeasesEvent.Set();
+			_updateLeasesEvent.Set();
 		}
 
 		/// <summary>
 		/// Dispatch a lease payload to the appropriate handler
 		/// </summary>
-		/// <param name="RpcConnection">The RPC connection to the server</param>
-		/// <param name="AgentId">The current agent id</param>
-		/// <param name="LeaseInfo">Information about the lease</param>
+		/// <param name="rpcConnection">The RPC connection to the server</param>
+		/// <param name="agentId">The current agent id</param>
+		/// <param name="leaseInfo">Information about the lease</param>
 		/// <returns>Outcome from the lease</returns>
-		internal async Task<LeaseResult> HandleLeasePayloadAsync(IRpcConnection RpcConnection, string AgentId, LeaseInfo LeaseInfo)
+		internal async Task<LeaseResult> HandleLeasePayloadAsync(IRpcConnection rpcConnection, string agentId, LeaseInfo leaseInfo)
 		{
-			Any Payload = LeaseInfo.Lease.Payload;
+			Any payload = leaseInfo.Lease.Payload;
 
-			ComputeTaskMessage ComputeTask;
-			if (LeaseInfo.Lease.Payload.TryUnpack(out ComputeTask))
+			ComputeTaskMessage computeTask;
+			if (leaseInfo.Lease.Payload.TryUnpack(out computeTask))
 			{
 				GlobalTracer.Instance.ActiveSpan?.SetTag("task", "Compute");
-				return await ComputeAsync(RpcConnection, LeaseInfo.Lease.Id, ComputeTask, Logger, LeaseInfo.CancellationTokenSource.Token);
+				return await ComputeAsync(rpcConnection, leaseInfo.Lease.Id, computeTask, _logger, leaseInfo.CancellationTokenSource.Token);
 			}
 
-			ConformTask ConformTask;
-			if (LeaseInfo.Lease.Payload.TryUnpack(out ConformTask))
+			ConformTask conformTask;
+			if (leaseInfo.Lease.Payload.TryUnpack(out conformTask))
 			{
 				GlobalTracer.Instance.ActiveSpan?.SetTag("task", "Conform");
-				Func<ILogger, Task<LeaseResult>> Handler = NewLogger => ConformAsync(RpcConnection, AgentId, LeaseInfo.Lease.Id, ConformTask, NewLogger, LeaseInfo.CancellationTokenSource.Token);
-				return await HandleLeasePayloadWithLogAsync(RpcConnection, ConformTask.LogId, AgentId, null, null, Handler, LeaseInfo.CancellationTokenSource.Token);
+				Task<LeaseResult> Handler(ILogger newLogger) => ConformAsync(rpcConnection, agentId, leaseInfo.Lease.Id, conformTask, newLogger, leaseInfo.CancellationTokenSource.Token);
+				return await HandleLeasePayloadWithLogAsync(rpcConnection, conformTask.LogId, null, null, Handler, leaseInfo.CancellationTokenSource.Token);
 			}
 
-			ExecuteJobTask JobTask;
-			if (LeaseInfo.Lease.Payload.TryUnpack(out JobTask))
+			ExecuteJobTask jobTask;
+			if (leaseInfo.Lease.Payload.TryUnpack(out jobTask))
 			{
 				GlobalTracer.Instance.ActiveSpan?.SetTag("task", "Job");
-				Func<ILogger, Task<LeaseResult>> Handler = NewLogger => ExecuteJobAsync(RpcConnection, AgentId, LeaseInfo.Lease.Id, JobTask, NewLogger, LeaseInfo.CancellationTokenSource.Token);
-				return await HandleLeasePayloadWithLogAsync(RpcConnection, JobTask.LogId, AgentId, JobTask.JobId, JobTask.BatchId, Handler, LeaseInfo.CancellationTokenSource.Token);
+				Task<LeaseResult> Handler(ILogger newLogger) => ExecuteJobAsync(rpcConnection, agentId, leaseInfo.Lease.Id, jobTask, newLogger, leaseInfo.CancellationTokenSource.Token);
+				return await HandleLeasePayloadWithLogAsync(rpcConnection, jobTask.LogId, jobTask.JobId, jobTask.BatchId, Handler, leaseInfo.CancellationTokenSource.Token);
 			}
 
-			UpgradeTask UpgradeTask;
-			if (LeaseInfo.Lease.Payload.TryUnpack(out UpgradeTask) || TryUnpackUpgradeTask(LeaseInfo.Lease.Payload, out UpgradeTask))
+			UpgradeTask upgradeTask;
+			if (leaseInfo.Lease.Payload.TryUnpack(out upgradeTask) || TryUnpackUpgradeTask(leaseInfo.Lease.Payload, out upgradeTask))
 			{
 				GlobalTracer.Instance.ActiveSpan?.SetTag("task", "Upgrade");
-				Func<ILogger, Task<LeaseResult>> Handler = NewLogger => UpgradeAsync(RpcConnection, AgentId, UpgradeTask, NewLogger, LeaseInfo.CancellationTokenSource.Token);
-				return await HandleLeasePayloadWithLogAsync(RpcConnection, UpgradeTask.LogId, AgentId, null, null, Handler, LeaseInfo.CancellationTokenSource.Token);
+				Task<LeaseResult> Handler(ILogger newLogger) => UpgradeAsync(rpcConnection, upgradeTask, newLogger, leaseInfo.CancellationTokenSource.Token);
+				return await HandleLeasePayloadWithLogAsync(rpcConnection, upgradeTask.LogId, null, null, Handler, leaseInfo.CancellationTokenSource.Token);
 			}
 
-			ShutdownTask ShutdownTask;
-			if (LeaseInfo.Lease.Payload.TryUnpack(out ShutdownTask))
+			ShutdownTask shutdownTask;
+			if (leaseInfo.Lease.Payload.TryUnpack(out shutdownTask))
 			{
 				GlobalTracer.Instance.ActiveSpan?.SetTag("task", "Shutdown");
-				Func<ILogger, Task<LeaseResult>> Handler = NewLogger => ShutdownAsync(RpcConnection, AgentId, ShutdownTask, NewLogger, LeaseInfo.CancellationTokenSource.Token);
-				return await HandleLeasePayloadWithLogAsync(RpcConnection, ShutdownTask.LogId, AgentId, null, null, Handler, LeaseInfo.CancellationTokenSource.Token);
+				Task<LeaseResult> Handler(ILogger newLogger) => ShutdownAsync(rpcConnection, agentId, shutdownTask, newLogger, leaseInfo.CancellationTokenSource.Token);
+				return await HandleLeasePayloadWithLogAsync(rpcConnection, shutdownTask.LogId, null, null, Handler, leaseInfo.CancellationTokenSource.Token);
 			}
 
-			RestartTask RestartTask;
-			if (LeaseInfo.Lease.Payload.TryUnpack(out RestartTask))
+			RestartTask restartTask;
+			if (leaseInfo.Lease.Payload.TryUnpack(out restartTask))
 			{
 				GlobalTracer.Instance.ActiveSpan?.SetTag("task", "Restart");
-				Func<ILogger, Task<LeaseResult>> Handler = NewLogger => RestartAsync(RpcConnection, AgentId, RestartTask, NewLogger, LeaseInfo.CancellationTokenSource.Token);
-				return await HandleLeasePayloadWithLogAsync(RpcConnection, RestartTask.LogId, AgentId, null, null, Handler, LeaseInfo.CancellationTokenSource.Token);
+				Task<LeaseResult> Handler(ILogger newLogger) => RestartAsync(rpcConnection, agentId, restartTask, newLogger, leaseInfo.CancellationTokenSource.Token);
+				return await HandleLeasePayloadWithLogAsync(rpcConnection, restartTask.LogId, null, null, Handler, leaseInfo.CancellationTokenSource.Token);
 			}
 
-			Logger.LogError("Invalid lease payload type ({PayloadType})", Payload.TypeUrl);
+			_logger.LogError("Invalid lease payload type ({PayloadType})", payload.TypeUrl);
 			return LeaseResult.Failed;
 		}
 
 		/// <summary>
 		/// Attempts to parse an upgrade task, given either name
 		/// </summary>
-		/// <param name="Payload"></param>
-		/// <param name="UpgradeTask"></param>
+		/// <param name="payload"></param>
+		/// <param name="upgradeTask"></param>
 		/// <returns></returns>
-		static bool TryUnpackUpgradeTask(Any Payload, out UpgradeTask UpgradeTask)
+		static bool TryUnpackUpgradeTask(Any payload, out UpgradeTask upgradeTask)
 		{
-			if (Payload.TypeUrl == "type.googleapis.com/UpgradeTask" || Payload.TypeUrl == "type.googleapis.com/Horde.UpgradeTask")
+			if (payload.TypeUrl == "type.googleapis.com/UpgradeTask" || payload.TypeUrl == "type.googleapis.com/Horde.UpgradeTask")
 			{
-				UpgradeTask = UpgradeTask.Parser.ParseFrom(Payload.Value);
+				upgradeTask = UpgradeTask.Parser.ParseFrom(payload.Value);
 				return true;
 			}
 			else
 			{
-				UpgradeTask = null!;
+				upgradeTask = null!;
 				return false;
 			}
 		}
@@ -832,28 +832,28 @@ namespace Horde.Agent.Services
 		/// <summary>
 		/// Terminates any processes which are still running under the given directory
 		/// </summary>
-		/// <param name="Logger">Logger device</param>
-		/// <param name="CancellationToken">Cancellation token</param>
-		void TerminateProcesses(ILogger Logger, CancellationToken CancellationToken)
+		/// <param name="logger">Logger device</param>
+		/// <param name="cancellationToken">Cancellation token</param>
+		void TerminateProcesses(ILogger logger, CancellationToken cancellationToken)
 		{
 			// Terminate child processes from any previous runs
-			ProcessUtils.TerminateProcesses(ShouldTerminateProcess, Logger, CancellationToken);
+			ProcessUtils.TerminateProcesses(ShouldTerminateProcess, logger, cancellationToken);
 		}
 
 		/// <summary>
 		/// Callback for determining whether a process should be terminated
 		/// </summary>
-		/// <param name="ImageFile">The file to terminate</param>
+		/// <param name="imageFile">The file to terminate</param>
 		/// <returns>True if the process should be terminated</returns>
-		bool ShouldTerminateProcess(FileReference ImageFile)
+		bool ShouldTerminateProcess(FileReference imageFile)
 		{
-			if (ImageFile.IsUnderDirectory(WorkingDir))
+			if (imageFile.IsUnderDirectory(_workingDir))
 			{
 				return true;
 			}
 
-			string FileName = ImageFile.GetFileName();
-			if (ProcessNamesToTerminate.Contains(FileName))
+			string fileName = imageFile.GetFileName();
+			if (_processNamesToTerminate.Contains(fileName))
 			{
 				return true;
 			}
@@ -864,34 +864,33 @@ namespace Horde.Agent.Services
 		/// <summary>
 		/// Dispatch a lease payload to the appropriate handler
 		/// </summary>
-		/// <param name="RpcConnection">The RPC connection to the server</param>
-		/// <param name="LogId">Unique id for the log to create</param>
-		/// <param name="AgentId">The current agent id</param>
-		/// <param name="JobId">The job being executed</param>
-		/// <param name="JobBatchId">The batch being executed</param>
-		/// <param name="ExecuteLease">Action to perform to execute the lease</param>
-		/// <param name="CancellationToken">The cancellation token for this lease</param>
+		/// <param name="rpcConnection">The RPC connection to the server</param>
+		/// <param name="logId">Unique id for the log to create</param>
+		/// <param name="jobId">The job being executed</param>
+		/// <param name="jobBatchId">The batch being executed</param>
+		/// <param name="executeLease">Action to perform to execute the lease</param>
+		/// <param name="cancellationToken">The cancellation token for this lease</param>
 		/// <returns>Outcome from the lease</returns>
-		async Task<LeaseResult> HandleLeasePayloadWithLogAsync(IRpcConnection RpcConnection, string LogId, string AgentId, string? JobId, string? JobBatchId, Func<ILogger, Task<LeaseResult>> ExecuteLease, CancellationToken CancellationToken)
+		async Task<LeaseResult> HandleLeasePayloadWithLogAsync(IRpcConnection rpcConnection, string logId, string? jobId, string? jobBatchId, Func<ILogger, Task<LeaseResult>> executeLease, CancellationToken cancellationToken)
 		{
-			await using (JsonRpcLogger? ConformLogger = String.IsNullOrEmpty(LogId) ? null : new JsonRpcLogger(RpcConnection, LogId, JobId, JobBatchId, null, null, Logger))
+			await using (JsonRpcLogger? conformLogger = String.IsNullOrEmpty(logId) ? null : new JsonRpcLogger(rpcConnection, logId, jobId, jobBatchId, null, null, _logger))
 			{
-				ILogger LeaseLogger = (ConformLogger == null) ? (ILogger)Logger : new ForwardingLogger(Logger, new DefaultLoggerIndentHandler(ConformLogger));
+				ILogger leaseLogger = (conformLogger == null) ? (ILogger)_logger : new ForwardingLogger(_logger, new DefaultLoggerIndentHandler(conformLogger));
 				try
 				{
-					LeaseResult Result = await ExecuteLease(LeaseLogger);
-					return Result;
+					LeaseResult result = await executeLease(leaseLogger);
+					return result;
 				}
-				catch (Exception Ex)
+				catch (Exception ex)
 				{
-					if (CancellationToken.IsCancellationRequested && IsCancellationException(Ex))
+					if (cancellationToken.IsCancellationRequested && IsCancellationException(ex))
 					{
-						LeaseLogger.LogInformation("Lease was cancelled");
+						leaseLogger.LogInformation("Lease was cancelled");
 						return LeaseResult.Cancelled;
 					}
 					else
 					{
-						LeaseLogger.LogError(Ex, "Caught unhandled exception while attempting to execute lease:\n{Exception}", Ex);
+						leaseLogger.LogError(ex, "Caught unhandled exception while attempting to execute lease:\n{Exception}", ex);
 						return LeaseResult.Failed;
 					}
 				}
@@ -901,35 +900,35 @@ namespace Horde.Agent.Services
 		/// <summary>
 		/// Execute a remote execution aciton
 		/// </summary>
-		/// <param name="RpcConnection">RPC client for communicating with the server</param>
-		/// <param name="LeaseId">The lease id</param>
-		/// <param name="ComputeTask">The action task parameters</param>
-		/// <param name="Logger">Logger for the task</param>
-		/// <param name="CancellationToken">Token used to cancel the operation</param>
+		/// <param name="rpcConnection">RPC client for communicating with the server</param>
+		/// <param name="leaseId">The lease id</param>
+		/// <param name="computeTask">The action task parameters</param>
+		/// <param name="logger">Logger for the task</param>
+		/// <param name="cancellationToken">Token used to cancel the operation</param>
 		/// <returns></returns>
-		internal async Task<LeaseResult> ComputeAsync(IRpcConnection RpcConnection, string LeaseId, ComputeTaskMessage ComputeTask, ILogger Logger, CancellationToken CancellationToken)
+		internal async Task<LeaseResult> ComputeAsync(IRpcConnection rpcConnection, string leaseId, ComputeTaskMessage computeTask, ILogger logger, CancellationToken cancellationToken)
 		{
-			Logger.LogInformation("Starting compute task (lease {LeaseId})", LeaseId);
+			logger.LogInformation("Starting compute task (lease {LeaseId})", leaseId);
 
-			ComputeTaskResultMessage Result;
+			ComputeTaskResultMessage result;
 			try
 			{
-				DateTimeOffset ActionTaskStartTime = DateTimeOffset.UtcNow;
-				using (IRpcClientRef Client = await RpcConnection.GetClientRef(new RpcContext(), CancellationToken))
+				DateTimeOffset actionTaskStartTime = DateTimeOffset.UtcNow;
+				using (IRpcClientRef client = await rpcConnection.GetClientRef(new RpcContext(), cancellationToken))
 				{
-					DirectoryReference LeaseDir = DirectoryReference.Combine(WorkingDir, "Compute", LeaseId);
-					DirectoryReference.CreateDirectory(LeaseDir);
+					DirectoryReference leaseDir = DirectoryReference.Combine(_workingDir, "Compute", leaseId);
+					DirectoryReference.CreateDirectory(leaseDir);
 
-					ComputeTaskExecutor Executor = new ComputeTaskExecutor(StorageClient, Logger);
+					ComputeTaskExecutor executor = new ComputeTaskExecutor(_storageClient, logger);
 					try
 					{
-						Result = await Executor.ExecuteAsync(LeaseId, ComputeTask, LeaseDir, CancellationToken);
+						result = await executor.ExecuteAsync(leaseId, computeTask, leaseDir, cancellationToken);
 					}
 					finally
 					{
 						try
 						{
-							DirectoryReference.Delete(LeaseDir, true);
+							DirectoryReference.Delete(leaseDir, true);
 						}
 						catch
 						{
@@ -937,65 +936,65 @@ namespace Horde.Agent.Services
 					}
 				}
 			}
-			catch (BlobNotFoundException Ex)
+			catch (BlobNotFoundException ex)
 			{
-				Logger.LogError(Ex, "Blob not found: {Hash}", Ex.Hash);
-				Result = new ComputeTaskResultMessage(ComputeTaskOutcome.BlobNotFound, Ex.Hash.ToString());
+				logger.LogError(ex, "Blob not found: {Hash}", ex.Hash);
+				result = new ComputeTaskResultMessage(ComputeTaskOutcome.BlobNotFound, ex.Hash.ToString());
 			}
-			catch (Exception Ex)
+			catch (Exception ex)
 			{
-				Logger.LogError(Ex, "Exception while executing compute task");
-				Result = new ComputeTaskResultMessage(ComputeTaskOutcome.Exception, Ex.ToString());
+				logger.LogError(ex, "Exception while executing compute task");
+				result = new ComputeTaskResultMessage(ComputeTaskOutcome.Exception, ex.ToString());
 			}
 
-			return new LeaseResult(Result.ToByteArray());
+			return new LeaseResult(result.ToByteArray());
 		}
 
 		/// <summary>
 		/// Conform a machine
 		/// </summary>
-		/// <param name="RpcConnection">RPC client for communicating with the server</param>
-		/// <param name="AgentId">The current agent id</param>
-		/// <param name="LeaseId">The current lease id</param>
-		/// <param name="ConformTask">The conform task parameters</param>
-		/// <param name="ConformLogger">Logger for the task</param>
-		/// <param name="CancellationToken">Token used to cancel the operation</param>
+		/// <param name="rpcConnection">RPC client for communicating with the server</param>
+		/// <param name="agentId">The current agent id</param>
+		/// <param name="leaseId">The current lease id</param>
+		/// <param name="conformTask">The conform task parameters</param>
+		/// <param name="conformLogger">Logger for the task</param>
+		/// <param name="cancellationToken">Token used to cancel the operation</param>
 		/// <returns>Async task</returns>
-		async Task<LeaseResult> ConformAsync(IRpcConnection RpcConnection, string AgentId, string LeaseId, ConformTask ConformTask, ILogger ConformLogger, CancellationToken CancellationToken)
+		async Task<LeaseResult> ConformAsync(IRpcConnection rpcConnection, string agentId, string leaseId, ConformTask conformTask, ILogger conformLogger, CancellationToken cancellationToken)
 		{
-			ConformLogger.LogInformation("Conforming, lease {LeaseId}", LeaseId);
-			TerminateProcesses(ConformLogger, CancellationToken);
+			conformLogger.LogInformation("Conforming, lease {LeaseId}", leaseId);
+			TerminateProcesses(conformLogger, cancellationToken);
 
-			bool RemoveUntrackedFiles = ConformTask.RemoveUntrackedFiles;
-			IList<AgentWorkspace> PendingWorkspaces = ConformTask.Workspaces;
+			bool removeUntrackedFiles = conformTask.RemoveUntrackedFiles;
+			IList<AgentWorkspace> pendingWorkspaces = conformTask.Workspaces;
 			for (; ;)
 			{
 				// Run the conform task
-				if (Settings.Executor == ExecutorType.Perforce && Settings.PerforceExecutor.RunConform)
+				if (_settings.Executor == ExecutorType.Perforce && _settings.PerforceExecutor.RunConform)
 				{
-					await PerforceExecutor.ConformAsync(WorkingDir, PendingWorkspaces, RemoveUntrackedFiles, ConformLogger, CancellationToken);
+					await PerforceExecutor.ConformAsync(_workingDir, pendingWorkspaces, removeUntrackedFiles, conformLogger, cancellationToken);
 				}
 				else
 				{
-					ConformLogger.LogInformation("Skipping due to Settings.RunConform flag");
+					conformLogger.LogInformation("Skipping due to Settings.RunConform flag");
 				}
 
 				// Update the new set of workspaces
-				UpdateAgentWorkspacesRequest Request = new UpdateAgentWorkspacesRequest();
-				Request.AgentId = AgentId;
-				Request.Workspaces.AddRange(PendingWorkspaces);
-				Request.RemoveUntrackedFiles = RemoveUntrackedFiles;
+				UpdateAgentWorkspacesRequest request = new UpdateAgentWorkspacesRequest();
+				request.AgentId = agentId;
+				request.Workspaces.AddRange(pendingWorkspaces);
+				request.RemoveUntrackedFiles = removeUntrackedFiles;
 
-				UpdateAgentWorkspacesResponse Response = await RpcConnection.InvokeAsync(x => x.UpdateAgentWorkspacesAsync(Request, null, null, CancellationToken), new RpcContext(), CancellationToken);
-				if (!Response.Retry)
+				UpdateAgentWorkspacesResponse response = await rpcConnection.InvokeAsync(x => x.UpdateAgentWorkspacesAsync(request, null, null, cancellationToken), new RpcContext(), cancellationToken);
+				if (!response.Retry)
 				{
-					ConformLogger.LogInformation("Conform finished");
+					conformLogger.LogInformation("Conform finished");
 					break;
 				}
 
-				ConformLogger.LogInformation("Pending workspaces have changed - running conform again...");
-				PendingWorkspaces = Response.PendingWorkspaces;
-				RemoveUntrackedFiles = Response.RemoveUntrackedFiles;
+				conformLogger.LogInformation("Pending workspaces have changed - running conform again...");
+				pendingWorkspaces = response.PendingWorkspaces;
+				removeUntrackedFiles = response.RemoveUntrackedFiles;
 			}
 
 			return LeaseResult.Success;
@@ -1004,50 +1003,50 @@ namespace Horde.Agent.Services
 		/// <summary>
 		/// Execute part of a job
 		/// </summary>
-		/// <param name="RpcClient">RPC client for communicating with the server</param>
-		/// <param name="AgentId">The current agent id</param>
-		/// <param name="LeaseId">The current lease id</param>
-		/// <param name="ExecuteTask">The task to execute</param>
-		/// <param name="Logger">The logger to use for this lease</param>
-		/// <param name="CancellationToken">Token used to cancel the operation</param>
+		/// <param name="rpcClient">RPC client for communicating with the server</param>
+		/// <param name="agentId">The current agent id</param>
+		/// <param name="leaseId">The current lease id</param>
+		/// <param name="executeTask">The task to execute</param>
+		/// <param name="logger">The logger to use for this lease</param>
+		/// <param name="cancellationToken">Token used to cancel the operation</param>
 		/// <returns>Async task</returns>
-		internal async Task<LeaseResult> ExecuteJobAsync(IRpcConnection RpcClient, string AgentId, string LeaseId, ExecuteJobTask ExecuteTask, ILogger Logger, CancellationToken CancellationToken)
+		internal async Task<LeaseResult> ExecuteJobAsync(IRpcConnection rpcClient, string agentId, string leaseId, ExecuteJobTask executeTask, ILogger logger, CancellationToken cancellationToken)
 		{
-			Logger.LogInformation("Executing job \"{JobName}\", jobId {JobId}, batchId {BatchId}, leaseId {LeaseId}", ExecuteTask.JobName, ExecuteTask.JobId, ExecuteTask.BatchId, LeaseId);
-			GlobalTracer.Instance.ActiveSpan?.SetTag("jobId", ExecuteTask.JobId.ToString());
-			GlobalTracer.Instance.ActiveSpan?.SetTag("jobName", ExecuteTask.JobName.ToString());
-			GlobalTracer.Instance.ActiveSpan?.SetTag("batchId", ExecuteTask.BatchId.ToString());	
+			logger.LogInformation("Executing job \"{JobName}\", jobId {JobId}, batchId {BatchId}, leaseId {LeaseId}", executeTask.JobName, executeTask.JobId, executeTask.BatchId, leaseId);
+			GlobalTracer.Instance.ActiveSpan?.SetTag("jobId", executeTask.JobId.ToString());
+			GlobalTracer.Instance.ActiveSpan?.SetTag("jobName", executeTask.JobName.ToString());
+			GlobalTracer.Instance.ActiveSpan?.SetTag("batchId", executeTask.BatchId.ToString());	
 
 			// Start executing the current batch
-			BeginBatchResponse Batch = await RpcClient.InvokeAsync(x => x.BeginBatchAsync(new BeginBatchRequest(ExecuteTask.JobId, ExecuteTask.BatchId, LeaseId), null, null, CancellationToken), new RpcContext(), CancellationToken);
+			BeginBatchResponse batch = await rpcClient.InvokeAsync(x => x.BeginBatchAsync(new BeginBatchRequest(executeTask.JobId, executeTask.BatchId, leaseId), null, null, cancellationToken), new RpcContext(), cancellationToken);
 
 			// Execute the batch
 			try
 			{
-				await ExecuteBatchAsync(RpcClient, AgentId, LeaseId, ExecuteTask, Batch, Logger, CancellationToken);
+				await ExecuteBatchAsync(rpcClient, agentId, leaseId, executeTask, batch, logger, cancellationToken);
 			}
-			catch (Exception Ex)
+			catch (Exception ex)
 			{
-				if (CancellationToken.IsCancellationRequested && IsCancellationException(Ex))
+				if (cancellationToken.IsCancellationRequested && IsCancellationException(ex))
 				{
-					Logger.LogError("Step was aborted");
+					logger.LogError("Step was aborted");
 					throw;
 				}
 				else
 				{
-					Logger.LogError(Ex, "Exception while executing batch: {Ex}", Ex);
+					logger.LogError(ex, "Exception while executing batch: {ex}", ex);
 				}
 			}
 
 			// If this lease was cancelled, don't bother updating the job state.
-			if (CancellationToken.IsCancellationRequested)
+			if (cancellationToken.IsCancellationRequested)
 			{
 				return LeaseResult.Cancelled;
 			}
 
 			// Mark the batch as complete
-			await RpcClient.InvokeAsync(x => x.FinishBatchAsync(new FinishBatchRequest(ExecuteTask.JobId, ExecuteTask.BatchId, LeaseId), null, null, CancellationToken), new RpcContext(), CancellationToken);
-			Logger.LogInformation("Done.");
+			await rpcClient.InvokeAsync(x => x.FinishBatchAsync(new FinishBatchRequest(executeTask.JobId, executeTask.BatchId, leaseId), null, null, cancellationToken), new RpcContext(), cancellationToken);
+			logger.LogInformation("Done.");
 
 			return LeaseResult.Success;
 		}
@@ -1055,31 +1054,31 @@ namespace Horde.Agent.Services
 		/// <summary>
 		/// Executes a batch
 		/// </summary>
-		/// <param name="RpcClient">RPC client for communicating with the server</param>
-		/// <param name="AgentId">The current agent id</param>
-		/// <param name="LeaseId">The current lease id</param>
-		/// <param name="ExecuteTask">The task to execute</param>
-		/// <param name="Batch">The batch to execute</param>
-		/// <param name="BatchLogger">Output log for the batch</param>
-		/// <param name="CancellationToken">Cancellation token to abort the batch</param>
+		/// <param name="rpcClient">RPC client for communicating with the server</param>
+		/// <param name="agentId">The current agent id</param>
+		/// <param name="leaseId">The current lease id</param>
+		/// <param name="executeTask">The task to execute</param>
+		/// <param name="batch">The batch to execute</param>
+		/// <param name="batchLogger">Output log for the batch</param>
+		/// <param name="cancellationToken">Cancellation token to abort the batch</param>
 		/// <returns>Async task</returns>
-		async Task ExecuteBatchAsync(IRpcConnection RpcClient, string AgentId, string LeaseId, ExecuteJobTask ExecuteTask, BeginBatchResponse Batch, ILogger BatchLogger, CancellationToken CancellationToken)
+		async Task ExecuteBatchAsync(IRpcConnection rpcClient, string agentId, string leaseId, ExecuteJobTask executeTask, BeginBatchResponse batch, ILogger batchLogger, CancellationToken cancellationToken)
 		{
-			BatchLogger.LogInformation("Executing batch {BatchId} using {Executor} executor", ExecuteTask.BatchId, Settings.Executor.ToString());
-			TerminateProcesses(BatchLogger, CancellationToken);
+			batchLogger.LogInformation("Executing batch {BatchId} using {Executor} executor", executeTask.BatchId, _settings.Executor.ToString());
+			TerminateProcesses(batchLogger, cancellationToken);
 
 			// Get the working directory for this lease
-			DirectoryReference ScratchDir = DirectoryReference.Combine(WorkingDir, "Scratch");
+			DirectoryReference scratchDir = DirectoryReference.Combine(_workingDir, "Scratch");
 
 			// Create an executor for this job
-			IExecutor Executor = CreateExecutor(RpcClient, ExecuteTask, Batch);
+			IExecutor executor = _createExecutor(rpcClient, executeTask, batch);
 
 			// Try to initialize the executor
-			BatchLogger.LogInformation("Initializing...");
-			using (BatchLogger.BeginIndentScope("  "))
+			batchLogger.LogInformation("Initializing...");
+			using (batchLogger.BeginIndentScope("  "))
 			{
-				using IScope Scope = GlobalTracer.Instance.BuildSpan("Initialize").StartActive();
-				await Executor.InitializeAsync(BatchLogger, CancellationToken);
+				using IScope scope = GlobalTracer.Instance.BuildSpan("Initialize").StartActive();
+				await executor.InitializeAsync(batchLogger, cancellationToken);
 			}
 
 			try
@@ -1088,218 +1087,218 @@ namespace Horde.Agent.Services
 				for (; ; )
 				{
 					// Get the next step to execute
-					BeginStepResponse Step = await RpcClient.InvokeAsync(x => x.BeginStepAsync(new BeginStepRequest(ExecuteTask.JobId, ExecuteTask.BatchId, LeaseId), null, null, CancellationToken), new RpcContext(), CancellationToken);
-					if (Step.State == BeginStepResponse.Types.Result.Waiting)
+					BeginStepResponse step = await rpcClient.InvokeAsync(x => x.BeginStepAsync(new BeginStepRequest(executeTask.JobId, executeTask.BatchId, leaseId), null, null, cancellationToken), new RpcContext(), cancellationToken);
+					if (step.State == BeginStepResponse.Types.Result.Waiting)
 					{
-						BatchLogger.LogInformation("Waiting for dependency to be ready");
-						await Task.Delay(TimeSpan.FromSeconds(20.0), CancellationToken);
+						batchLogger.LogInformation("Waiting for dependency to be ready");
+						await Task.Delay(TimeSpan.FromSeconds(20.0), cancellationToken);
 						continue;
 					}
-					else if (Step.State == BeginStepResponse.Types.Result.Complete)
+					else if (step.State == BeginStepResponse.Types.Result.Complete)
 					{
 						break;
 					}
-					else if (Step.State != BeginStepResponse.Types.Result.Ready)
+					else if (step.State != BeginStepResponse.Types.Result.Ready)
 					{
-						BatchLogger.LogError("Unexpected step state: {StepState}", Step.State);
+						batchLogger.LogError("Unexpected step state: {StepState}", step.State);
 						break;
 					}
 
 					// Get current disk space available. This will allow us to more easily spot steps that eat up a lot of disk space.
-					string? DriveName;
+					string? driveName;
 					if (RuntimeInformation.IsOSPlatform(OSPlatform.Windows))
 					{
-						DriveName = Path.GetPathRoot(WorkingDir.FullName);
+						driveName = Path.GetPathRoot(_workingDir.FullName);
 					}
 					else
 					{
-						DriveName = WorkingDir.FullName;
+						driveName = _workingDir.FullName;
 					}
 
-					float AvailableFreeSpace = 0;
-					if (DriveName != null)
+					float availableFreeSpace = 0;
+					if (driveName != null)
 					{
 						try
 						{
-							DriveInfo Info = new DriveInfo(DriveName);
-							AvailableFreeSpace = (1.0f * Info.AvailableFreeSpace) / 1024 / 1024 / 1024;
+							DriveInfo info = new DriveInfo(driveName);
+							availableFreeSpace = (1.0f * info.AvailableFreeSpace) / 1024 / 1024 / 1024;
 						}
-						catch (Exception Ex)
+						catch (Exception ex)
 						{
-							BatchLogger.LogWarning(Ex, "Unable to query disk info for path '{DriveName}'", DriveName);
+							batchLogger.LogWarning(ex, "Unable to query disk info for path '{DriveName}'", driveName);
 						}
 					}
 
 					// Print the new state
-					Stopwatch StepTimer = Stopwatch.StartNew();
+					Stopwatch stepTimer = Stopwatch.StartNew();
 
-					BatchLogger.LogInformation("Starting job {JobId}, batch {BatchId}, step {StepId} (Drive Space Left: {DriveSpaceRemaining} GB)", ExecuteTask.JobId, ExecuteTask.BatchId, Step.StepId, AvailableFreeSpace.ToString("F1"));
+					batchLogger.LogInformation("Starting job {JobId}, batch {BatchId}, step {StepId} (Drive Space Left: {DriveSpaceRemaining} GB)", executeTask.JobId, executeTask.BatchId, step.StepId, availableFreeSpace.ToString("F1"));
 
 					// Create a trace span
-					using IScope Scope = GlobalTracer.Instance.BuildSpan("Execute").WithResourceName(Step.Name).StartActive();
-					Scope.Span.SetTag("stepId", Step.StepId);
-					Scope.Span.SetTag("logId", Step.LogId);
+					using IScope scope = GlobalTracer.Instance.BuildSpan("Execute").WithResourceName(step.Name).StartActive();
+					scope.Span.SetTag("stepId", step.StepId);
+					scope.Span.SetTag("logId", step.LogId);
 					//				using IDisposable TraceProperty = LogContext.PushProperty("dd.trace_id", CorrelationIdentifier.TraceId.ToString());
 					//				using IDisposable SpanProperty = LogContext.PushProperty("dd.span_id", CorrelationIdentifier.SpanId.ToString());
 
 					// Update the context to include information about this step
-					JobStepOutcome StepOutcome;
-					JobStepState StepState;
-					using (BatchLogger.BeginIndentScope("  "))
+					JobStepOutcome stepOutcome;
+					JobStepState stepState;
+					using (batchLogger.BeginIndentScope("  "))
 					{
 						// Start writing to the log file
-						await using (JsonRpcLogger StepLogger = new JsonRpcLogger(RpcClient, Step.LogId, ExecuteTask.JobId, ExecuteTask.BatchId, Step.StepId, Step.Warnings, Logger))
+						await using (JsonRpcLogger stepLogger = new JsonRpcLogger(rpcClient, step.LogId, executeTask.JobId, executeTask.BatchId, step.StepId, step.Warnings, _logger))
 						{
 							// Execute the task
-							ILogger ForwardingLogger = new DefaultLoggerIndentHandler(StepLogger);
-							if (Settings.WriteStepOutputToLogger)
+							ILogger forwardingLogger = new DefaultLoggerIndentHandler(stepLogger);
+							if (_settings.WriteStepOutputToLogger)
 							{
-								ForwardingLogger = new ForwardingLogger(Logger, ForwardingLogger);
+								forwardingLogger = new ForwardingLogger(_logger, forwardingLogger);
 							}
 
-							using CancellationTokenSource StepPollCancelSource = new CancellationTokenSource();
-							using CancellationTokenSource StepAbortSource = new CancellationTokenSource();
-							TaskCompletionSource<bool> StepFinishedSource = new TaskCompletionSource<bool>();
-							Task StepPollTask = Task.Run(() => PollForStepAbort(RpcClient, ExecuteTask.JobId, ExecuteTask.BatchId, Step.StepId, StepPollCancelSource.Token, StepAbortSource, StepFinishedSource.Task));
+							using CancellationTokenSource stepPollCancelSource = new CancellationTokenSource();
+							using CancellationTokenSource stepAbortSource = new CancellationTokenSource();
+							TaskCompletionSource<bool> stepFinishedSource = new TaskCompletionSource<bool>();
+							Task stepPollTask = Task.Run(() => PollForStepAbort(rpcClient, executeTask.JobId, executeTask.BatchId, step.StepId, stepPollCancelSource.Token, stepAbortSource, stepFinishedSource.Task));
 
 							try
 							{
-								(StepOutcome, StepState) = await ExecuteStepAsync(Executor, Step, ForwardingLogger, CancellationToken, StepAbortSource.Token);
+								(stepOutcome, stepState) = await ExecuteStepAsync(executor, step, forwardingLogger, cancellationToken, stepAbortSource.Token);
 							}
 							finally
 							{
 								// Will get called even when cancellation token for the lease/batch fires
-								StepFinishedSource.SetResult(true); // Tell background poll task to stop
-								await StepPollTask;
+								stepFinishedSource.SetResult(true); // Tell background poll task to stop
+								await stepPollTask;
 							}
 
 							// Kill any processes spawned by the step
-							TerminateProcesses(StepLogger, CancellationToken);
+							TerminateProcesses(stepLogger, cancellationToken);
 
 							// Wait for the logger to finish
-							await StepLogger.StopAsync();
+							await stepLogger.StopAsync();
 
 							// Reflect the warnings/errors in the step outcome
-							if (StepOutcome > StepLogger.Outcome)
+							if (stepOutcome > stepLogger.Outcome)
 							{
-								StepOutcome = StepLogger.Outcome;
+								stepOutcome = stepLogger.Outcome;
 							}
 						}
 
 						// Update the server with the outcome from the step
-						BatchLogger.LogInformation("Marking step as complete (Outcome={Outcome}, State={StepState})", StepOutcome, StepState);
-						await RpcClient.InvokeAsync(x => x.UpdateStepAsync(new UpdateStepRequest(ExecuteTask.JobId, ExecuteTask.BatchId, Step.StepId, StepState, StepOutcome), null, null, CancellationToken), new RpcContext(), CancellationToken);
+						batchLogger.LogInformation("Marking step as complete (Outcome={Outcome}, State={StepState})", stepOutcome, stepState);
+						await rpcClient.InvokeAsync(x => x.UpdateStepAsync(new UpdateStepRequest(executeTask.JobId, executeTask.BatchId, step.StepId, stepState, stepOutcome), null, null, cancellationToken), new RpcContext(), cancellationToken);
 					}
 
 					// Print the finishing state
-					StepTimer.Stop();
-					BatchLogger.LogInformation("Completed in {Time}", StepTimer.Elapsed);
+					stepTimer.Stop();
+					batchLogger.LogInformation("Completed in {Time}", stepTimer.Elapsed);
 				}
 			}
-			catch (Exception Ex)
+			catch (Exception ex)
 			{
-				if (CancellationToken.IsCancellationRequested && IsCancellationException(Ex))
+				if (cancellationToken.IsCancellationRequested && IsCancellationException(ex))
 				{
-					Logger.LogError("Step was aborted");
+					_logger.LogError("Step was aborted");
 				}
 				else
 				{
-					Logger.LogError(Ex, "Exception while executing batch: {Ex}", Ex);
+					_logger.LogError(ex, "Exception while executing batch: {ex}", ex);
 				}
 			}
 
 			// Clean the environment
-			BatchLogger.LogInformation("Finalizing...");
-			using (BatchLogger.BeginIndentScope("  "))
+			batchLogger.LogInformation("Finalizing...");
+			using (batchLogger.BeginIndentScope("  "))
 			{
-				using IScope Scope = GlobalTracer.Instance.BuildSpan("Finalize").StartActive();
-				await Executor.FinalizeAsync(BatchLogger, CancellationToken.None);
+				using IScope scope = GlobalTracer.Instance.BuildSpan("Finalize").StartActive();
+				await executor.FinalizeAsync(batchLogger, CancellationToken.None);
 			}
 		}
 
 		/// <summary>
 		/// Executes a step
 		/// </summary>
-		/// <param name="Executor">The executor to run this step</param>
-		/// <param name="Step">Step to execute</param>
-		/// <param name="StepLogger">Logger for the step</param>
-		/// <param name="CancellationToken">Cancellation token to abort the batch</param>
-		/// <param name="StepCancellationToken">Cancellation token to abort only this individual step</param>
+		/// <param name="executor">The executor to run this step</param>
+		/// <param name="step">Step to execute</param>
+		/// <param name="stepLogger">Logger for the step</param>
+		/// <param name="cancellationToken">Cancellation token to abort the batch</param>
+		/// <param name="stepCancellationToken">Cancellation token to abort only this individual step</param>
 		/// <returns>Async task</returns>
-		internal async Task<(JobStepOutcome, JobStepState)> ExecuteStepAsync(IExecutor Executor, BeginStepResponse Step, ILogger StepLogger, CancellationToken CancellationToken, CancellationToken StepCancellationToken)
+		internal async Task<(JobStepOutcome, JobStepState)> ExecuteStepAsync(IExecutor executor, BeginStepResponse step, ILogger stepLogger, CancellationToken cancellationToken, CancellationToken stepCancellationToken)
 		{
-			using CancellationTokenSource Combined = CancellationTokenSource.CreateLinkedTokenSource(CancellationToken, StepCancellationToken);
+			using CancellationTokenSource combined = CancellationTokenSource.CreateLinkedTokenSource(cancellationToken, stepCancellationToken);
 			try
 			{
-				JobStepOutcome StepOutcome = await Executor.RunAsync(Step, StepLogger, Combined.Token);
-				return (StepOutcome, JobStepState.Completed);
+				JobStepOutcome stepOutcome = await executor.RunAsync(step, stepLogger, combined.Token);
+				return (stepOutcome, JobStepState.Completed);
 			}
-			catch (Exception Ex)
+			catch (Exception ex)
 			{
-				if (CancellationToken.IsCancellationRequested && IsCancellationException(Ex))
+				if (cancellationToken.IsCancellationRequested && IsCancellationException(ex))
 				{
-					StepLogger.LogError("The step was cancelled by batch/lease");
+					stepLogger.LogError("The step was cancelled by batch/lease");
 					throw;
 				}
 				
-				if (StepCancellationToken.IsCancellationRequested && IsCancellationException(Ex))
+				if (stepCancellationToken.IsCancellationRequested && IsCancellationException(ex))
 				{
-					StepLogger.LogError("The step was intentionally cancelled");
+					stepLogger.LogError("The step was intentionally cancelled");
 					return (JobStepOutcome.Failure, JobStepState.Aborted);
 				}
 				
-				StepLogger.LogError(Ex, "Exception while executing step: {Ex}", Ex);
+				stepLogger.LogError(ex, "Exception while executing step: {ex}", ex);
 				return (JobStepOutcome.Failure, JobStepState.Completed);
 			}
 		}
 
-		internal async Task PollForStepAbort(IRpcConnection RpcClient, string JobId, string BatchId, string StepId, CancellationToken CancellationToken, CancellationTokenSource StepCancelSource, Task FinishedTask)
+		internal async Task PollForStepAbort(IRpcConnection epcClient, string jobId, string batchId, string stepId, CancellationToken cancellationToken, CancellationTokenSource stepCancelSource, Task finishedTask)
 		{
-			Stopwatch Timer = Stopwatch.StartNew();
-			while (!FinishedTask.IsCompleted)
+			Stopwatch timer = Stopwatch.StartNew();
+			while (!finishedTask.IsCompleted)
 			{
 				try
 				{
-					if (Timer.Elapsed > TimeSpan.FromHours(24.0))
+					if (timer.Elapsed > TimeSpan.FromHours(24.0))
 					{
-						Logger.LogDebug("Step was aborted after running for {NumHours} (JobId={JobId} BatchId={BatchId} StepId={StepId})", Timer.Elapsed, JobId, BatchId, StepId);
-						StepCancelSource.Cancel();
+						_logger.LogDebug("Step was aborted after running for {NumHours} (JobId={JobId} BatchId={BatchId} StepId={StepId})", timer.Elapsed, jobId, batchId, stepId);
+						stepCancelSource.Cancel();
 						break;
 					}
 
-					GetStepResponse Res = await RpcClient.InvokeAsync(x => x.GetStepAsync(new GetStepRequest(JobId, BatchId, StepId), null, null, CancellationToken), new RpcContext(), CancellationToken);
-					if (Res.AbortRequested)
+					GetStepResponse res = await epcClient.InvokeAsync(x => x.GetStepAsync(new GetStepRequest(jobId, batchId, stepId), null, null, cancellationToken), new RpcContext(), cancellationToken);
+					if (res.AbortRequested)
 					{
-						Logger.LogDebug("Step was aborted by server (JobId={JobId} BatchId={BatchId} StepId={StepId})", JobId, BatchId, StepId);
-						StepCancelSource.Cancel();
+						_logger.LogDebug("Step was aborted by server (JobId={JobId} BatchId={BatchId} StepId={StepId})", jobId, batchId, stepId);
+						stepCancelSource.Cancel();
 						break;
 					}
 				}
-				catch (RpcException Ex)
+				catch (RpcException ex)
 				{
-					Logger.LogError(Ex, "Poll for step abort has failed. Aborting (JobId={JobId} BatchId={BatchId} StepId={StepId})", JobId, BatchId, StepId);
-					StepCancelSource.Cancel();
+					_logger.LogError(ex, "Poll for step abort has failed. Aborting (JobId={JobId} BatchId={BatchId} StepId={StepId})", jobId, batchId, stepId);
+					stepCancelSource.Cancel();
 					break;
 				}
 
-				await Task.WhenAny(Task.Delay(StepAbortPollInterval), FinishedTask);
+				await Task.WhenAny(Task.Delay(_stepAbortPollInterval), finishedTask);
 			}
 		}
 
 		/// <summary>
 		/// Determine if the given exception was triggered due to a cancellation event
 		/// </summary>
-		/// <param name="Ex">The exception to check</param>
+		/// <param name="ex">The exception to check</param>
 		/// <returns>True if the exception is a cancellation exception</returns>
-		static bool IsCancellationException(Exception Ex)
+		static bool IsCancellationException(Exception ex)
 		{
-			if(Ex is OperationCanceledException)
+			if(ex is OperationCanceledException)
 			{
 				return true;
 			}
 
-			RpcException? RpcException = Ex as RpcException;
-			if(RpcException != null && RpcException.StatusCode == StatusCode.Cancelled)
+			RpcException? rpcException = ex as RpcException;
+			if(rpcException != null && rpcException.StatusCode == StatusCode.Cancelled)
 			{
 				return true;
 			}
@@ -1310,44 +1309,43 @@ namespace Horde.Agent.Services
 		/// <summary>
 		/// Check for an update of the agent software
 		/// </summary>
-		/// <param name="RpcClient">RPC client for communicating with the server</param>
-		/// <param name="AgentId">The current agent id</param>
-		/// <param name="UpgradeTask">The upgrade task</param>
-		/// <param name="Logger">Logging device</param>
-		/// <param name="CancellationToken">Token used to cancel the operation</param>
+		/// <param name="rpcClient">RPC client for communicating with the server</param>
+		/// <param name="upgradeTask">The upgrade task</param>
+		/// <param name="logger">Logging device</param>
+		/// <param name="cancellationToken">Token used to cancel the operation</param>
 		/// <returns>Outcome of this operation</returns>
-		async Task<LeaseResult> UpgradeAsync(IRpcConnection RpcClient, string AgentId, UpgradeTask UpgradeTask, ILogger Logger, CancellationToken CancellationToken)
+		async Task<LeaseResult> UpgradeAsync(IRpcConnection rpcClient, UpgradeTask upgradeTask, ILogger logger, CancellationToken cancellationToken)
 		{
-			string RequiredVersion = UpgradeTask.SoftwareId;
+			string requiredVersion = upgradeTask.SoftwareId;
 
 			// Check if we're running the right version
-			if (RequiredVersion != null && RequiredVersion != Program.Version)
+			if (requiredVersion != null && requiredVersion != Program.Version)
 			{
-				Logger.LogInformation("Upgrading from {CurrentVersion} to {TargetVersion}", Program.Version, RequiredVersion);
+				logger.LogInformation("Upgrading from {CurrentVersion} to {TargetVersion}", Program.Version, requiredVersion);
 
 				// Clear out the working directory
-				DirectoryReference UpgradeDir = DirectoryReference.Combine(WorkingDir, "Upgrade");
-				DirectoryReference.CreateDirectory(UpgradeDir);
-				await DeleteDirectoryContentsAsync(new DirectoryInfo(UpgradeDir.FullName));
+				DirectoryReference upgradeDir = DirectoryReference.Combine(_workingDir, "Upgrade");
+				DirectoryReference.CreateDirectory(upgradeDir);
+				await DeleteDirectoryContentsAsync(new DirectoryInfo(upgradeDir.FullName));
 
 				// Download the new software
-				FileInfo OutputFile = new FileInfo(Path.Combine(UpgradeDir.FullName, "Agent.zip"));
-				using (IRpcClientRef RpcClientRef = await RpcClient.GetClientRef(new RpcContext(), CancellationToken))
-				using (AsyncServerStreamingCall<DownloadSoftwareResponse> Cursor = RpcClientRef.Client.DownloadSoftware(new DownloadSoftwareRequest(RequiredVersion), null, null, CancellationToken))
+				FileInfo outputFile = new FileInfo(Path.Combine(upgradeDir.FullName, "Agent.zip"));
+				using (IRpcClientRef rpcClientRef = await rpcClient.GetClientRef(new RpcContext(), cancellationToken))
+				using (AsyncServerStreamingCall<DownloadSoftwareResponse> cursor = rpcClientRef.Client.DownloadSoftware(new DownloadSoftwareRequest(requiredVersion), null, null, cancellationToken))
 				{
-					using (Stream OutputStream = OutputFile.Open(FileMode.Create))
+					using (Stream outputStream = outputFile.Open(FileMode.Create))
 					{
-						while (await Cursor.ResponseStream.MoveNext(CancellationToken))
+						while (await cursor.ResponseStream.MoveNext(cancellationToken))
 						{
-							OutputStream.Write(Cursor.ResponseStream.Current.Data.Span);
+							outputStream.Write(cursor.ResponseStream.Current.Data.Span);
 						}
 					}
 				}
 
 				// Extract it to a temporary directory
-				DirectoryReference ExtractedDir = DirectoryReference.Combine(UpgradeDir, "Extracted");
-				DirectoryReference.CreateDirectory(ExtractedDir);
-				ZipFile.ExtractToDirectory(OutputFile.FullName, ExtractedDir.FullName);
+				DirectoryReference extractedDir = DirectoryReference.Combine(upgradeDir, "Extracted");
+				DirectoryReference.CreateDirectory(extractedDir);
+				ZipFile.ExtractToDirectory(outputFile.FullName, extractedDir.FullName);
 
 				//				// Debug code for updating an agent with the local version
 				//				foreach (FileInfo SourceFile in new FileInfo(Assembly.GetExecutingAssembly().Location).Directory.EnumerateFiles())
@@ -1355,58 +1353,57 @@ namespace Horde.Agent.Services
 				//					SourceFile.CopyTo(Path.Combine(ExtractedDir.FullName, SourceFile.Name), true);
 				//				}
 
-
 				// Get the current process and assembly. This may be different if running through dotnet.exe rather than a native PE image.
-				FileReference AssemblyFileName = new FileReference(Assembly.GetExecutingAssembly().Location);
+				FileReference assemblyFileName = new FileReference(Assembly.GetExecutingAssembly().Location);
 
 				// Spawn the other process
-				using (Process Process = new Process())
+				using (Process process = new Process())
 				{
-					StringBuilder Arguments = new StringBuilder();
+					StringBuilder arguments = new StringBuilder();
 
-					DirectoryReference TargetDir = AssemblyFileName.Directory;
+					DirectoryReference targetDir = assemblyFileName.Directory;
 
 					// We were launched via an external application (presumably dotnet.exe). Do the same thing again.
-					FileReference NewAssemblyFileName = FileReference.Combine(ExtractedDir, AssemblyFileName.MakeRelativeTo(TargetDir));
-					if (!FileReference.Exists(NewAssemblyFileName))
+					FileReference newAssemblyFileName = FileReference.Combine(extractedDir, assemblyFileName.MakeRelativeTo(targetDir));
+					if (!FileReference.Exists(newAssemblyFileName))
 					{
-						Logger.LogError("Unable to find {AgentExe} in extracted archive", NewAssemblyFileName);
+						logger.LogError("Unable to find {AgentExe} in extracted archive", newAssemblyFileName);
 						return LeaseResult.Failed;
 					}
 
-					Process.StartInfo.FileName = "dotnet";
+					process.StartInfo.FileName = "dotnet";
 
-					StringBuilder CurrentArguments = new StringBuilder();
-					foreach (string Arg in Program.Args)
+					StringBuilder currentArguments = new StringBuilder();
+					foreach (string arg in Program.Args)
 					{
-						CurrentArguments.AppendArgument(Arg);
+						currentArguments.AppendArgument(arg);
 					}
 
-					Arguments.AppendArgument(NewAssemblyFileName.FullName);
-					Arguments.AppendArgument("Service");
-					Arguments.AppendArgument("Upgrade");
-					Arguments.AppendArgument("-ProcessId=", Process.GetCurrentProcess().Id.ToString());
-					Arguments.AppendArgument("-TargetDir=", TargetDir.FullName);
-					Arguments.AppendArgument("-Arguments=", CurrentArguments.ToString());
+					arguments.AppendArgument(newAssemblyFileName.FullName);
+					arguments.AppendArgument("Service");
+					arguments.AppendArgument("Upgrade");
+					arguments.AppendArgument("-ProcessId=", Process.GetCurrentProcess().Id.ToString());
+					arguments.AppendArgument("-TargetDir=", targetDir.FullName);
+					arguments.AppendArgument("-Arguments=", currentArguments.ToString());
 
-					Process.StartInfo.Arguments = Arguments.ToString();
-					Process.StartInfo.UseShellExecute = false;
-					Process.EnableRaisingEvents = true;
+					process.StartInfo.Arguments = arguments.ToString();
+					process.StartInfo.UseShellExecute = false;
+					process.EnableRaisingEvents = true;
 
-					StringBuilder LaunchCommand = new StringBuilder();
-					LaunchCommand.AppendArgument(Process.StartInfo.FileName);
-					LaunchCommand.Append(' ');
-					LaunchCommand.Append(Arguments);
-					Logger.LogInformation("Launching: {Launch}", LaunchCommand.ToString());
+					StringBuilder launchCommand = new StringBuilder();
+					launchCommand.AppendArgument(process.StartInfo.FileName);
+					launchCommand.Append(' ');
+					launchCommand.Append(arguments);
+					logger.LogInformation("Launching: {Launch}", launchCommand.ToString());
 
-					TaskCompletionSource<int> ExitCodeSource = new TaskCompletionSource<int>();
-					Process.Exited += (Sender, Args) => { ExitCodeSource.SetResult(Process.ExitCode); };
+					TaskCompletionSource<int> exitCodeSource = new TaskCompletionSource<int>();
+					process.Exited += (sender, args) => { exitCodeSource.SetResult(process.ExitCode); };
 
-					Process.Start();
+					process.Start();
 
-					using (CancellationToken.Register(() => { ExitCodeSource.SetResult(0); }))
+					using (cancellationToken.Register(() => { exitCodeSource.SetResult(0); }))
 					{
-						await ExitCodeSource.Task;
+						await exitCodeSource.Task;
 					}
 				}
 			}
@@ -1417,65 +1414,65 @@ namespace Horde.Agent.Services
 		/// <summary>
 		/// Delete the contents of a directory without deleting it itself
 		/// </summary>
-		/// <param name="BaseDir">Directory to clean</param>
+		/// <param name="baseDir">Directory to clean</param>
 		/// <returns>Async task</returns>
-		static async Task DeleteDirectoryContentsAsync(DirectoryInfo BaseDir)
+		static async Task DeleteDirectoryContentsAsync(DirectoryInfo baseDir)
 		{
-			List<Task> ChildTasks = new List<Task>();
-			foreach (DirectoryInfo SubDir in BaseDir.EnumerateDirectories())
+			List<Task> childTasks = new List<Task>();
+			foreach (DirectoryInfo subDir in baseDir.EnumerateDirectories())
 			{
-				ChildTasks.Add(Task.Run(() => DeleteDirectory(SubDir)));
+				childTasks.Add(Task.Run(() => DeleteDirectory(subDir)));
 			}
-			foreach (FileInfo File in BaseDir.EnumerateFiles())
+			foreach (FileInfo file in baseDir.EnumerateFiles())
 			{
-				File.Attributes = FileAttributes.Normal;
-				File.Delete();
+				file.Attributes = FileAttributes.Normal;
+				file.Delete();
 			}
-			await Task.WhenAll(ChildTasks);
+			await Task.WhenAll(childTasks);
 		}
 
 		/// <summary>
 		/// Deletes a directory and its contents
 		/// </summary>
-		/// <param name="BaseDir">Directory to delete</param>
+		/// <param name="baseDir">Directory to delete</param>
 		/// <returns>Async task</returns>
-		static async Task DeleteDirectory(DirectoryInfo BaseDir)
+		static async Task DeleteDirectory(DirectoryInfo baseDir)
 		{
-			await DeleteDirectoryContentsAsync(BaseDir);
-			BaseDir.Delete();
+			await DeleteDirectoryContentsAsync(baseDir);
+			baseDir.Delete();
 		}
 
 		/// <summary>
 		/// Check for an update of the agent software
 		/// </summary>
-		/// <param name="RpcClient">RPC client for communicating with the server</param>
-		/// <param name="AgentId">The current agent id</param>
-		/// <param name="RestartTask">The restart task parameters</param>
-		/// <param name="Logger">Logging device</param>
-		/// <param name="CancellationToken">Token used to cancel the operation</param>
+		/// <param name="rpcClient">RPC client for communicating with the server</param>
+		/// <param name="agentId">The current agent id</param>
+		/// <param name="restartTask">The restart task parameters</param>
+		/// <param name="logger">Logging device</param>
+		/// <param name="cancellationToken">Token used to cancel the operation</param>
 		/// <returns>Outcome of this operation</returns>
-		Task<LeaseResult> RestartAsync(IRpcConnection RpcClient, string AgentId, RestartTask RestartTask, ILogger Logger, CancellationToken CancellationToken)
+		Task<LeaseResult> RestartAsync(IRpcConnection rpcClient, string agentId, RestartTask restartTask, ILogger logger, CancellationToken cancellationToken)
 		{
-			Logger.LogInformation("Setting restart flag");
-			RequestShutdown = true;
-			RestartAfterShutdown = true;
+			logger.LogInformation("Setting restart flag");
+			_requestShutdown = true;
+			_restartAfterShutdown = true;
 			return Task.FromResult(LeaseResult.Success);
 		}
 
 		/// <summary>
 		/// Shutdown the agent
 		/// </summary>
-		/// <param name="RpcClient">RPC client for communicating with the server</param>
-		/// <param name="AgentId">The current agent id</param>
-		/// <param name="ShutdownTask">The restart task parameters</param>
-		/// <param name="Logger">Logging device</param>
-		/// <param name="CancellationToken">Token used to cancel the operation</param>
+		/// <param name="rpcClient">RPC client for communicating with the server</param>
+		/// <param name="agentId">The current agent id</param>
+		/// <param name="shutdownTask">The restart task parameters</param>
+		/// <param name="logger">Logging device</param>
+		/// <param name="cancellationToken">Token used to cancel the operation</param>
 		/// <returns>Outcome of this operation</returns>
-		Task<LeaseResult> ShutdownAsync(IRpcConnection RpcClient, string AgentId, ShutdownTask ShutdownTask, ILogger Logger, CancellationToken CancellationToken)
+		Task<LeaseResult> ShutdownAsync(IRpcConnection rpcClient, string agentId, ShutdownTask shutdownTask, ILogger logger, CancellationToken cancellationToken)
 		{
-			Logger.LogInformation("Setting shutdown flag");
-			RequestShutdown = true;
-			RestartAfterShutdown = false;
+			logger.LogInformation("Setting shutdown flag");
+			_requestShutdown = true;
+			_restartAfterShutdown = false;
 			return Task.FromResult(LeaseResult.Success);
 		}
 
@@ -1483,278 +1480,278 @@ namespace Horde.Agent.Services
 		/// Gets the hardware capabilities of this worker
 		/// </summary>
 		/// <returns>Worker object for advertising to the server</returns>
-		public static async Task<AgentCapabilities> GetAgentCapabilities(DirectoryReference WorkingDir, ILogger Logger)
+		public static async Task<AgentCapabilities> GetAgentCapabilities(DirectoryReference workingDir, ILogger logger)
 		{
 			// Create the primary device
-			DeviceCapabilities PrimaryDevice = new DeviceCapabilities();
-			PrimaryDevice.Handle = "Primary";
+			DeviceCapabilities primaryDevice = new DeviceCapabilities();
+			primaryDevice.Handle = "Primary";
 
-			List<DeviceCapabilities> OtherDevices = new List<DeviceCapabilities>();
+			List<DeviceCapabilities> otherDevices = new List<DeviceCapabilities>();
 			if (RuntimeInformation.IsOSPlatform(OSPlatform.Windows))
 			{
-				PrimaryDevice.Properties.Add("Platform=Win64");
-				PrimaryDevice.Properties.Add("PlatformGroup=Windows");
-				PrimaryDevice.Properties.Add("PlatformGroup=Microsoft");
-				PrimaryDevice.Properties.Add("PlatformGroup=Desktop");
+				primaryDevice.Properties.Add("Platform=Win64");
+				primaryDevice.Properties.Add("PlatformGroup=Windows");
+				primaryDevice.Properties.Add("PlatformGroup=Microsoft");
+				primaryDevice.Properties.Add("PlatformGroup=Desktop");
 
-				PrimaryDevice.Properties.Add("OSFamily=Windows");
+				primaryDevice.Properties.Add("OSFamily=Windows");
 
 				// Add OS info
-				using (ManagementObjectSearcher Searcher = new ManagementObjectSearcher("select * from Win32_OperatingSystem"))
+				using (ManagementObjectSearcher searcher = new ManagementObjectSearcher("select * from Win32_OperatingSystem"))
 				{
-					foreach (ManagementObject Row in Searcher.Get())
+					foreach (ManagementObject row in searcher.Get())
 					{
-						using (Row)
+						using (row)
 						{
-							Dictionary<string, object> Properties = GetWmiProperties(Row);
+							Dictionary<string, object> properties = GetWmiProperties(row);
 
-							object? Name;
-							if (Properties.TryGetValue("Caption", out Name))
+							object? name;
+							if (properties.TryGetValue("Caption", out name))
 							{
-								PrimaryDevice.Properties.Add($"OSDistribution={Name}");
+								primaryDevice.Properties.Add($"OSDistribution={name}");
 							}
 
-							object? Version;
-							if (Properties.TryGetValue("Version", out Version))
+							object? version;
+							if (properties.TryGetValue("Version", out version))
 							{
-								PrimaryDevice.Properties.Add($"OSKernelVersion={Version}");
+								primaryDevice.Properties.Add($"OSKernelVersion={version}");
 							}
 						}
 					}
 				}
 
 				// Add CPU info
-				using (ManagementObjectSearcher Searcher = new ManagementObjectSearcher("select * from Win32_Processor"))
+				using (ManagementObjectSearcher searcher = new ManagementObjectSearcher("select * from Win32_Processor"))
 				{
-					Dictionary<string, int> NameToCount = new Dictionary<string, int>();
-					int TotalPhysicalCores = 0;
-					int TotalLogicalCores = 0;
+					Dictionary<string, int> nameToCount = new Dictionary<string, int>();
+					int totalPhysicalCores = 0;
+					int totalLogicalCores = 0;
 
-					foreach (ManagementObject Row in Searcher.Get())
+					foreach (ManagementObject row in searcher.Get())
 					{
-						using (Row)
+						using (row)
 						{
-							Dictionary<string, object> Properties = GetWmiProperties(Row);
+							Dictionary<string, object> properties = GetWmiProperties(row);
 
-							object? NameObject;
-							if (Properties.TryGetValue("Name", out NameObject))
+							object? nameObject;
+							if (properties.TryGetValue("Name", out nameObject))
 							{
-								string Name = NameObject.ToString() ?? String.Empty;
-								int Count;
-								NameToCount.TryGetValue(Name, out Count);
-								NameToCount[Name] = Count + 1;
+								string name = nameObject.ToString() ?? String.Empty;
+								int count;
+								nameToCount.TryGetValue(name, out count);
+								nameToCount[name] = count + 1;
 							}
 
-							object? NumPhysicalCores;
-							if ((Properties.TryGetValue("NumberOfEnabledCore", out NumPhysicalCores) && NumPhysicalCores is uint) || (Properties.TryGetValue("NumberOfCores", out NumPhysicalCores) && NumPhysicalCores is uint))
+							object? numPhysicalCores;
+							if ((properties.TryGetValue("NumberOfEnabledCore", out numPhysicalCores) && numPhysicalCores is uint) || (properties.TryGetValue("NumberOfCores", out numPhysicalCores) && numPhysicalCores is uint))
 							{
-								TotalPhysicalCores += (int)(uint)NumPhysicalCores;
+								totalPhysicalCores += (int)(uint)numPhysicalCores;
 							}
 
-							object? NumLogicalCores;
-							if (Properties.TryGetValue("NumberOfLogicalProcessors", out NumLogicalCores) && NumLogicalCores is uint)
+							object? numLogicalCores;
+							if (properties.TryGetValue("NumberOfLogicalProcessors", out numLogicalCores) && numLogicalCores is uint numLogicalCoresUint)
 							{
-								TotalLogicalCores += (int)(uint)NumLogicalCores;
+								totalLogicalCores += (int)numLogicalCoresUint;
 							}
 						}
 					}
 
-					AddCpuInfo(PrimaryDevice, NameToCount, TotalLogicalCores, TotalPhysicalCores);
+					AddCpuInfo(primaryDevice, nameToCount, totalLogicalCores, totalPhysicalCores);
 				}
 
 				// Add RAM info
-				using (ManagementObjectSearcher Searcher = new ManagementObjectSearcher("select Capacity from Win32_PhysicalMemory"))
+				using (ManagementObjectSearcher searcher = new ManagementObjectSearcher("select Capacity from Win32_PhysicalMemory"))
 				{
-					ulong TotalCapacity = 0;
-					foreach (ManagementObject Row in Searcher.Get())
+					ulong totalCapacity = 0;
+					foreach (ManagementObject row in searcher.Get())
 					{
-						using (Row)
+						using (row)
 						{
-							object? Capacity = Row.GetPropertyValue("Capacity");
-							if (Capacity is ulong)
+							object? capacity = row.GetPropertyValue("Capacity");
+							if (capacity is ulong ulongCapacity)
 							{
-								TotalCapacity += (ulong)Capacity;
+								totalCapacity += ulongCapacity;
 							}
 						}
 					}
 
-					if (TotalCapacity > 0)
+					if (totalCapacity > 0)
 					{
-						PrimaryDevice.Properties.Add($"RAM={TotalCapacity / (1024 * 1024 * 1024)}");
+						primaryDevice.Properties.Add($"RAM={totalCapacity / (1024 * 1024 * 1024)}");
 					}
 				}
 
 				// Add GPU info
-				using (ManagementObjectSearcher Searcher = new ManagementObjectSearcher("select Name, DriverVersion, AdapterRAM from Win32_VideoController"))
+				using (ManagementObjectSearcher searcher = new ManagementObjectSearcher("select Name, DriverVersion, AdapterRAM from Win32_VideoController"))
 				{
-					int Index = 0;
-					foreach (ManagementObject Row in Searcher.Get())
+					int index = 0;
+					foreach (ManagementObject row in searcher.Get())
 					{
-						using (Row)
+						using (row)
 						{
-							WmiProperties Properties = new WmiProperties(Row);
-							if (Properties.TryGetValue("Name", out string? Name) && Properties.TryGetValue("DriverVersion", out string? DriverVersion))
+							WmiProperties properties = new WmiProperties(row);
+							if (properties.TryGetValue("Name", out string? name) && properties.TryGetValue("DriverVersion", out string? driverVersion))
 							{
-								string Prefix = $"GPU-{++Index}";
-								PrimaryDevice.Properties.Add($"{Prefix}-Name={Name}");
-								PrimaryDevice.Properties.Add($"{Prefix}-DriverVersion={DriverVersion}");								
+								string prefix = $"GPU-{++index}";
+								primaryDevice.Properties.Add($"{prefix}-Name={name}");
+								primaryDevice.Properties.Add($"{prefix}-DriverVersion={driverVersion}");								
 							}
 						}
 					}
 				}
 
 				// Add EC2 properties if needed
-				await AddAwsProperties(PrimaryDevice.Properties, Logger);
+				await AddAwsProperties(primaryDevice.Properties, logger);
 
 				// Add session information
-				PrimaryDevice.Properties.Add($"User={Environment.UserName}");
-				PrimaryDevice.Properties.Add($"Domain={Environment.UserDomainName}");
-				PrimaryDevice.Properties.Add($"Interactive={Environment.UserInteractive}");
-				PrimaryDevice.Properties.Add($"Elevated={BuildGraphExecutor.IsUserAdministrator()}");
+				primaryDevice.Properties.Add($"User={Environment.UserName}");
+				primaryDevice.Properties.Add($"Domain={Environment.UserDomainName}");
+				primaryDevice.Properties.Add($"Interactive={Environment.UserInteractive}");
+				primaryDevice.Properties.Add($"Elevated={BuildGraphExecutor.IsUserAdministrator()}");
 			}
 			else if (RuntimeInformation.IsOSPlatform(OSPlatform.Linux))
 			{
-				PrimaryDevice.Properties.Add("Platform=Linux");
-				PrimaryDevice.Properties.Add("PlatformGroup=Linux");
-				PrimaryDevice.Properties.Add("PlatformGroup=Unix");
-				PrimaryDevice.Properties.Add("PlatformGroup=Desktop");
+				primaryDevice.Properties.Add("Platform=Linux");
+				primaryDevice.Properties.Add("PlatformGroup=Linux");
+				primaryDevice.Properties.Add("PlatformGroup=Unix");
+				primaryDevice.Properties.Add("PlatformGroup=Desktop");
 
-				PrimaryDevice.Properties.Add("OSFamily=Linux");
-				PrimaryDevice.Properties.Add("OSVersion=Linux");
+				primaryDevice.Properties.Add("OSFamily=Linux");
+				primaryDevice.Properties.Add("OSVersion=Linux");
 
 				// Add EC2 properties if needed
-				await AddAwsProperties(PrimaryDevice.Properties, Logger);
+				await AddAwsProperties(primaryDevice.Properties, logger);
 
 				// Parse the CPU info
-				List<Dictionary<string, string>>? CpuRecords = await ReadLinuxHwPropsAsync("/proc/cpuinfo", Logger);
-				if (CpuRecords != null)
+				List<Dictionary<string, string>>? cpuRecords = await ReadLinuxHwPropsAsync("/proc/cpuinfo", logger);
+				if (cpuRecords != null)
 				{
-					Dictionary<string, string> CpuNames = new Dictionary<string, string>(StringComparer.Ordinal);
-					foreach (Dictionary<string, string> CpuRecord in CpuRecords)
+					Dictionary<string, string> cpuNames = new Dictionary<string, string>(StringComparer.Ordinal);
+					foreach (Dictionary<string, string> cpuRecord in cpuRecords)
 					{
-						if (CpuRecord.TryGetValue("physical id", out string? PhysicalId) && CpuRecord.TryGetValue("model name", out string? ModelName))
+						if (cpuRecord.TryGetValue("physical id", out string? physicalId) && cpuRecord.TryGetValue("model name", out string? modelName))
 						{
-							CpuNames[PhysicalId] = ModelName;
+							cpuNames[physicalId] = modelName;
 						}
 					}
 
-					Dictionary<string, int> NameToCount = new Dictionary<string, int>(StringComparer.Ordinal);
-					foreach (string CpuName in CpuNames.Values)
+					Dictionary<string, int> nameToCount = new Dictionary<string, int>(StringComparer.Ordinal);
+					foreach (string cpuName in cpuNames.Values)
 					{
-						NameToCount.TryGetValue(CpuName, out int Count);
-						NameToCount[CpuName] = Count + 1;
+						nameToCount.TryGetValue(cpuName, out int count);
+						nameToCount[cpuName] = count + 1;
 					}
 
-					HashSet<string> LogicalCores = new HashSet<string>();
-					HashSet<string> PhysicalCores = new HashSet<string>();
-					foreach (Dictionary<string, string> CpuRecord in CpuRecords)
+					HashSet<string> logicalCores = new HashSet<string>();
+					HashSet<string> physicalCores = new HashSet<string>();
+					foreach (Dictionary<string, string> cpuRecord in cpuRecords)
 					{
-						if (CpuRecord.TryGetValue("processor", out string? LogicalCoreId))
+						if (cpuRecord.TryGetValue("processor", out string? logicalCoreId))
 						{
-							LogicalCores.Add(LogicalCoreId);
+							logicalCores.Add(logicalCoreId);
 						}
-						if (CpuRecord.TryGetValue("core id", out string? PhysicalCoreId))
+						if (cpuRecord.TryGetValue("core id", out string? physicalCoreId))
 						{
-							PhysicalCores.Add(PhysicalCoreId);
+							physicalCores.Add(physicalCoreId);
 						}
 					}
 
-					AddCpuInfo(PrimaryDevice, NameToCount, LogicalCores.Count, PhysicalCores.Count);
+					AddCpuInfo(primaryDevice, nameToCount, logicalCores.Count, physicalCores.Count);
 				}
 
 				// Parse the RAM info
-				List<Dictionary<string, string>>? MemRecords = await ReadLinuxHwPropsAsync("/proc/meminfo", Logger);
-				if (MemRecords != null && MemRecords.Count > 0 && MemRecords[0].TryGetValue("MemTotal", out string? MemTotal))
+				List<Dictionary<string, string>>? memRecords = await ReadLinuxHwPropsAsync("/proc/meminfo", logger);
+				if (memRecords != null && memRecords.Count > 0 && memRecords[0].TryGetValue("MemTotal", out string? memTotal))
 				{
-					Match Match = Regex.Match(MemTotal, @"(\d+)\s+kB");
-					if (Match.Success)
+					Match match = Regex.Match(memTotal, @"(\d+)\s+kB");
+					if (match.Success)
 					{
-						long TotalCapacity = long.Parse(Match.Groups[1].Value) * 1024;
-						PrimaryDevice.Properties.Add($"RAM={TotalCapacity / (1024 * 1024 * 1024)}");
+						long totalCapacity = Int64.Parse(match.Groups[1].Value) * 1024;
+						primaryDevice.Properties.Add($"RAM={totalCapacity / (1024 * 1024 * 1024)}");
 					}
 				}
 
 				// Add session information
-				PrimaryDevice.Properties.Add($"User={Environment.UserName}");
-				PrimaryDevice.Properties.Add($"Domain={Environment.UserDomainName}");
-				PrimaryDevice.Properties.Add($"Interactive={Environment.UserInteractive}");
-				PrimaryDevice.Properties.Add($"Elevated={BuildGraphExecutor.IsUserAdministrator()}");
+				primaryDevice.Properties.Add($"User={Environment.UserName}");
+				primaryDevice.Properties.Add($"Domain={Environment.UserDomainName}");
+				primaryDevice.Properties.Add($"Interactive={Environment.UserInteractive}");
+				primaryDevice.Properties.Add($"Elevated={BuildGraphExecutor.IsUserAdministrator()}");
 			}
 			else if(RuntimeInformation.IsOSPlatform(OSPlatform.OSX))
 			{
-				PrimaryDevice.Properties.Add("Platform=Mac");
-				PrimaryDevice.Properties.Add("PlatformGroup=Apple");
-				PrimaryDevice.Properties.Add("PlatformGroup=Desktop");
+				primaryDevice.Properties.Add("Platform=Mac");
+				primaryDevice.Properties.Add("PlatformGroup=Apple");
+				primaryDevice.Properties.Add("PlatformGroup=Desktop");
 
-				PrimaryDevice.Properties.Add("OSFamily=MacOS");
-				PrimaryDevice.Properties.Add("OSVersion=MacOS");
+				primaryDevice.Properties.Add("OSFamily=MacOS");
+				primaryDevice.Properties.Add("OSVersion=MacOS");
 
-				string Output;
-				using(Process Process = new Process())
+				string output;
+				using(Process process = new Process())
 				{
-					Process.StartInfo.FileName = "system_profiler";
-					Process.StartInfo.Arguments = "SPHardwareDataType SPSoftwareDataType -xml";
-					Process.StartInfo.CreateNoWindow = true;
-					Process.StartInfo.RedirectStandardOutput = true;
-					Process.StartInfo.RedirectStandardInput = false;
-					Process.StartInfo.UseShellExecute = false;
-					Process.Start();
-
-					Output = Process.StandardOutput.ReadToEnd();
+					process.StartInfo.FileName = "system_profiler";
+					process.StartInfo.Arguments = "SPHardwareDataType SPSoftwareDataType -xml";
+					process.StartInfo.CreateNoWindow = true;
+					process.StartInfo.RedirectStandardOutput = true;
+					process.StartInfo.RedirectStandardInput = false;
+					process.StartInfo.UseShellExecute = false;
+					process.Start();
+
+					output = process.StandardOutput.ReadToEnd();
 				}
 
-				XmlDocument Xml = new XmlDocument();
-				Xml.LoadXml(Output);
+				XmlDocument xml = new XmlDocument();
+				xml.LoadXml(output);
 
-				XmlNode? HardwareNode = Xml.SelectSingleNode("/plist/array/dict/key[. = '_dataType']/following-sibling::string[. = 'SPHardwareDataType']/../key[. = '_items']/following-sibling::array/dict");
-				if(HardwareNode != null)
+				XmlNode? hardwareNode = xml.SelectSingleNode("/plist/array/dict/key[. = '_dataType']/following-sibling::string[. = 'SPHardwareDataType']/../key[. = '_items']/following-sibling::array/dict");
+				if(hardwareNode != null)
 				{
-					XmlNode? Model = HardwareNode.SelectSingleNode("key[. = 'machine_model']/following-sibling::string");
-					if(Model != null)
+					XmlNode? model = hardwareNode.SelectSingleNode("key[. = 'machine_model']/following-sibling::string");
+					if(model != null)
 					{
-						PrimaryDevice.Properties.Add($"Model={Model.InnerText}");
+						primaryDevice.Properties.Add($"Model={model.InnerText}");
 					}
 
-					XmlNode? CpuTypeNode = HardwareNode.SelectSingleNode("key[. = 'cpu_type']/following-sibling::string");
-					XmlNode? CpuSpeedNode = HardwareNode.SelectSingleNode("key[. = 'current_processor_speed']/following-sibling::string");
-					XmlNode? CpuPackagesNode = HardwareNode.SelectSingleNode("key[. = 'packages']/following-sibling::integer");
-					if(CpuTypeNode != null && CpuSpeedNode != null && CpuPackagesNode != null)
+					XmlNode? cpuTypeNode = hardwareNode.SelectSingleNode("key[. = 'cpu_type']/following-sibling::string");
+					XmlNode? cpuSpeedNode = hardwareNode.SelectSingleNode("key[. = 'current_processor_speed']/following-sibling::string");
+					XmlNode? cpuPackagesNode = hardwareNode.SelectSingleNode("key[. = 'packages']/following-sibling::integer");
+					if(cpuTypeNode != null && cpuSpeedNode != null && cpuPackagesNode != null)
 					{
-						PrimaryDevice.Properties.Add((CpuPackagesNode.InnerText != "1")? $"CPU={CpuPackagesNode.InnerText} x {CpuTypeNode.InnerText} @ {CpuSpeedNode.InnerText}" : $"CPU={CpuTypeNode.InnerText} @ {CpuSpeedNode.InnerText}");
+						primaryDevice.Properties.Add((cpuPackagesNode.InnerText != "1")? $"CPU={cpuPackagesNode.InnerText} x {cpuTypeNode.InnerText} @ {cpuSpeedNode.InnerText}" : $"CPU={cpuTypeNode.InnerText} @ {cpuSpeedNode.InnerText}");
 					}
 
-					PrimaryDevice.Properties.Add($"LogicalCores={Environment.ProcessorCount}");
+					primaryDevice.Properties.Add($"LogicalCores={Environment.ProcessorCount}");
 
-					XmlNode? CpuCountNode = HardwareNode.SelectSingleNode("key[. = 'number_processors']/following-sibling::integer");
-					if(CpuCountNode != null)
+					XmlNode? cpuCountNode = hardwareNode.SelectSingleNode("key[. = 'number_processors']/following-sibling::integer");
+					if(cpuCountNode != null)
 					{
-						PrimaryDevice.Properties.Add($"PhysicalCores={CpuCountNode.InnerText}");
+						primaryDevice.Properties.Add($"PhysicalCores={cpuCountNode.InnerText}");
 					}
 
-					XmlNode? MemoryNode = HardwareNode.SelectSingleNode("key[. = 'physical_memory']/following-sibling::string");
-					if(MemoryNode != null)
+					XmlNode? memoryNode = hardwareNode.SelectSingleNode("key[. = 'physical_memory']/following-sibling::string");
+					if(memoryNode != null)
 					{
-						string[] Parts = MemoryNode.InnerText.Split(new char[]{' '}, StringSplitOptions.RemoveEmptyEntries);
-						if(Parts.Length == 2 && Parts[1] == "GB")
+						string[] parts = memoryNode.InnerText.Split(new char[]{' '}, StringSplitOptions.RemoveEmptyEntries);
+						if(parts.Length == 2 && parts[1] == "GB")
 						{
-							PrimaryDevice.Properties.Add($"RAM={Parts[0]}");
+							primaryDevice.Properties.Add($"RAM={parts[0]}");
 						}
 					}
 				}
 
-				XmlNode? SoftwareNode = Xml.SelectSingleNode("/plist/array/dict/key[. = '_dataType']/following-sibling::string[. = 'SPSoftwareDataType']/../key[. = '_items']/following-sibling::array/dict");
-				if(SoftwareNode != null)
+				XmlNode? softwareNode = xml.SelectSingleNode("/plist/array/dict/key[. = '_dataType']/following-sibling::string[. = 'SPSoftwareDataType']/../key[. = '_items']/following-sibling::array/dict");
+				if(softwareNode != null)
 				{
-					XmlNode? OsVersionNode = SoftwareNode.SelectSingleNode("key[. = 'os_version']/following-sibling::string");
-					if(OsVersionNode != null)
+					XmlNode? osVersionNode = softwareNode.SelectSingleNode("key[. = 'os_version']/following-sibling::string");
+					if(osVersionNode != null)
 					{
-						PrimaryDevice.Properties.Add($"OSDistribution={OsVersionNode.InnerText}");
+						primaryDevice.Properties.Add($"OSDistribution={osVersionNode.InnerText}");
 					}
 
-					XmlNode? KernelVersionNode = SoftwareNode.SelectSingleNode("key[. = 'kernel_version']/following-sibling::string");
-					if(KernelVersionNode != null)
+					XmlNode? kernelVersionNode = softwareNode.SelectSingleNode("key[. = 'kernel_version']/following-sibling::string");
+					if(kernelVersionNode != null)
 					{
-						PrimaryDevice.Properties.Add($"OSKernelVersion={KernelVersionNode.InnerText}");
+						primaryDevice.Properties.Add($"OSKernelVersion={kernelVersionNode.InnerText}");
 					}
 				}
 			}
@@ -1762,239 +1759,239 @@ namespace Horde.Agent.Services
 			// Get the IP addresses
 			try
 			{
-				IPHostEntry Entry = Dns.GetHostEntry(Dns.GetHostName());
-				foreach (IPAddress Address in Entry.AddressList)
+				IPHostEntry entry = Dns.GetHostEntry(Dns.GetHostName());
+				foreach (IPAddress address in entry.AddressList)
 				{
-					if (Address.AddressFamily == System.Net.Sockets.AddressFamily.InterNetwork)
+					if (address.AddressFamily == System.Net.Sockets.AddressFamily.InterNetwork)
 					{
-						PrimaryDevice.Properties.Add($"Ipv4={Address}");
+						primaryDevice.Properties.Add($"Ipv4={address}");
 					}
-					else if (Address.AddressFamily == System.Net.Sockets.AddressFamily.InterNetworkV6)
+					else if (address.AddressFamily == System.Net.Sockets.AddressFamily.InterNetworkV6)
 					{
-						PrimaryDevice.Properties.Add($"Ipv6={Address}");
+						primaryDevice.Properties.Add($"Ipv6={address}");
 					}
 				}
 			}
-			catch (Exception Ex)
+			catch (Exception ex)
 			{
-				Logger.LogDebug(Ex, "Unable to get local IP address");
+				logger.LogDebug(ex, "Unable to get local IP address");
 			}
 
 			// Get the time that the machine booted
-			PrimaryDevice.Properties.Add($"BootTime={BootTime}");
-			PrimaryDevice.Properties.Add($"StartTime={StartTime}");
+			primaryDevice.Properties.Add($"BootTime={s_bootTime}");
+			primaryDevice.Properties.Add($"StartTime={s_startTime}");
 
 			// Add disk info based on platform
-			string? DriveName;
+			string? driveName;
 			if (RuntimeInformation.IsOSPlatform(OSPlatform.Windows))
 			{
-				DriveName = Path.GetPathRoot(WorkingDir.FullName);
+				driveName = Path.GetPathRoot(workingDir.FullName);
 			}
 			else
 			{
-				DriveName = WorkingDir.FullName;
+				driveName = workingDir.FullName;
 			}
 
-			if (DriveName != null)
+			if (driveName != null)
 			{
 				try
 				{
-					DriveInfo Info = new DriveInfo(DriveName);
-					PrimaryDevice.Properties.Add($"DiskFreeSpace={Info.AvailableFreeSpace}");
-					PrimaryDevice.Properties.Add($"DiskTotalSize={Info.TotalSize}");
+					DriveInfo info = new DriveInfo(driveName);
+					primaryDevice.Properties.Add($"DiskFreeSpace={info.AvailableFreeSpace}");
+					primaryDevice.Properties.Add($"DiskTotalSize={info.TotalSize}");
 				}
-				catch (Exception Ex)
+				catch (Exception ex)
 				{
-					Logger.LogWarning(Ex, "Unable to query disk info for path '{DriveName}'", DriveName);
+					logger.LogWarning(ex, "Unable to query disk info for path '{DriveName}'", driveName);
 				}
 			}
-			PrimaryDevice.Properties.Add($"WorkingDir={WorkingDir}");
+			primaryDevice.Properties.Add($"WorkingDir={workingDir}");
 
 			// Add any horde. env vars for custom properties.
-			IEnumerable<string> EnvVars = Environment.GetEnvironmentVariables().Keys.Cast<string>();
-			foreach (string EnvVar in EnvVars.Where(x => x.StartsWith("horde.", StringComparison.InvariantCultureIgnoreCase)))
+			IEnumerable<string> envVars = Environment.GetEnvironmentVariables().Keys.Cast<string>();
+			foreach (string envVar in envVars.Where(x => x.StartsWith("horde.", StringComparison.InvariantCultureIgnoreCase)))
 			{
-				PrimaryDevice.Properties.Add($"{EnvVar}={Environment.GetEnvironmentVariable(EnvVar)}");
+				primaryDevice.Properties.Add($"{envVar}={Environment.GetEnvironmentVariable(envVar)}");
 			}
 
 			// Create the worker
-			AgentCapabilities Agent = new AgentCapabilities();
-			Agent.Devices.Add(PrimaryDevice);
-			Agent.Devices.AddRange(OtherDevices);
-			return Agent;
+			AgentCapabilities agent = new AgentCapabilities();
+			agent.Devices.Add(primaryDevice);
+			agent.Devices.AddRange(otherDevices);
+			return agent;
 		}
 
-		static void AddCpuInfo(DeviceCapabilities PrimaryDevice, Dictionary<string, int> NameToCount, int NumLogicalCores, int NumPhysicalCores)
+		static void AddCpuInfo(DeviceCapabilities primaryDevice, Dictionary<string, int> nameToCount, int numLogicalCores, int numPhysicalCores)
 		{
-			if (NameToCount.Count > 0)
+			if (nameToCount.Count > 0)
 			{
-				PrimaryDevice.Properties.Add("CPU=" + String.Join(", ", NameToCount.Select(x => (x.Value > 1) ? $"{x.Key} x {x.Value}" : x.Key)));
+				primaryDevice.Properties.Add("CPU=" + String.Join(", ", nameToCount.Select(x => (x.Value > 1) ? $"{x.Key} x {x.Value}" : x.Key)));
 			}
 
-			if (NumLogicalCores > 0)
+			if (numLogicalCores > 0)
 			{
-				PrimaryDevice.Properties.Add($"LogicalCores={NumLogicalCores}");
+				primaryDevice.Properties.Add($"LogicalCores={numLogicalCores}");
 			}
 
-			if (NumPhysicalCores > 0)
+			if (numPhysicalCores > 0)
 			{
-				PrimaryDevice.Properties.Add($"PhysicalCores={NumPhysicalCores}");
+				primaryDevice.Properties.Add($"PhysicalCores={numPhysicalCores}");
 			}
 		}
 
-		static async Task<List<Dictionary<string, string>>?> ReadLinuxHwPropsAsync(string FileName, ILogger Logger)
+		static async Task<List<Dictionary<string, string>>?> ReadLinuxHwPropsAsync(string fileName, ILogger logger)
 		{
-			List<Dictionary<string, string>>? Records = null;
-			if (File.Exists(FileName))
+			List<Dictionary<string, string>>? records = null;
+			if (File.Exists(fileName))
 			{
-				Records = new List<Dictionary<string, string>>();
-				using (StreamReader Reader = new StreamReader(FileName))
+				records = new List<Dictionary<string, string>>();
+				using (StreamReader reader = new StreamReader(fileName))
 				{
-					Dictionary<string, string> Record = new Dictionary<string, string>(StringComparer.Ordinal);
+					Dictionary<string, string> record = new Dictionary<string, string>(StringComparer.Ordinal);
 
-					string? Line;
-					while ((Line = await Reader.ReadLineAsync()) != null)
+					string? line;
+					while ((line = await reader.ReadLineAsync()) != null)
 					{
-						int Idx = Line.IndexOf(':');
-						if (Idx == -1)
+						int idx = line.IndexOf(':');
+						if (idx == -1)
 						{
-							if (Record.Count > 0)
+							if (record.Count > 0)
 							{
-								Records.Add(Record);
-								Record = new Dictionary<string, string>(StringComparer.Ordinal);
+								records.Add(record);
+								record = new Dictionary<string, string>(StringComparer.Ordinal);
 							}
 						}
 						else
 						{
-							string Key = Line.Substring(0, Idx).Trim();
-							string Value = Line.Substring(Idx + 1).Trim();
+							string key = line.Substring(0, idx).Trim();
+							string value = line.Substring(idx + 1).Trim();
 
-							if (Record.TryGetValue(Key, out string? PrevValue))
+							if (record.TryGetValue(key, out string? prevValue))
 							{
-								Logger.LogWarning("Multiple entries for {Key} in {File} (was '{Prev}', now '{Next}')", Key, FileName, PrevValue, Value);
+								logger.LogWarning("Multiple entries for {Key} in {File} (was '{Prev}', now '{Next}')", key, fileName, prevValue, value);
 							}
 							else
 							{
-								Record.Add(Key, Value);
+								record.Add(key, value);
 							}
 						}
 					}
 
-					if (Record.Count > 0)
+					if (record.Count > 0)
 					{
-						Records.Add(Record);
+						records.Add(record);
 					}
 				}
 			}
-			return Records;
+			return records;
 		}
 
 		class WmiProperties
 		{
-			Dictionary<string, object> Properties = new Dictionary<string, object>(StringComparer.Ordinal);
+			readonly Dictionary<string, object> _properties = new Dictionary<string, object>(StringComparer.Ordinal);
 
-			public WmiProperties(ManagementObject Row)
+			public WmiProperties(ManagementObject row)
 			{
-				foreach (PropertyData Property in Row.Properties)
+				foreach (PropertyData property in row.Properties)
 				{
-					Properties[Property.Name] = Property.Value;
+					_properties[property.Name] = property.Value;
 				}
 			}
 
-			public bool TryGetValue(string Name, [NotNullWhen(true)] out string? Value)
+			public bool TryGetValue(string name, [NotNullWhen(true)] out string? value)
 			{
-				object? Object;
-				if (Properties.TryGetValue(Name, out Object))
+				object? @object;
+				if (_properties.TryGetValue(name, out @object))
 				{
-					Value = Object.ToString();
-					return Value != null;
+					value = @object.ToString();
+					return value != null;
 				}
 				else
 				{
-					Value = null!;
+					value = null!;
 					return false;
 				}
 			}
 
-			public bool TryGetValue(string Name, out long Value)
+			public bool TryGetValue(string name, out long value)
 			{
-				object? Object;
-				if (Properties.TryGetValue(Name, out Object))
+				object? objectValue;
+				if (_properties.TryGetValue(name, out objectValue))
 				{
-					if (Object is int)
+					if (objectValue is int intValue)
 					{
-						Value = (int)Object;
+						value = intValue;
 						return true;
 					}
-					else if (Object is uint)
+					else if (objectValue is uint uintValue)
 					{
-						Value = (uint)Object;
+						value = uintValue;
 						return true;
 					}
-					else if (Object is long)
+					else if (objectValue is long longValue)
 					{
-						Value = (long)Object;
+						value = longValue;
 						return true;
 					}
-					else if (Object is ulong)
+					else if (objectValue is ulong ulongValue)
 					{
-						Value = (long)(ulong)Object;
+						value = (long)ulongValue;
 						return true;
 					}
 				}
 
-				Value = 0;
+				value = 0;
 				return false;
 			}
 		}
 
-		static Dictionary<string, object> GetWmiProperties(ManagementObject Row)
+		static Dictionary<string, object> GetWmiProperties(ManagementObject row)
 		{
-			Dictionary<string, object> Properties = new Dictionary<string, object>(StringComparer.Ordinal);
-			foreach (PropertyData Property in Row.Properties)
+			Dictionary<string, object> properties = new Dictionary<string, object>(StringComparer.Ordinal);
+			foreach (PropertyData property in row.Properties)
 			{
-				Properties[Property.Name] = Property.Value;
+				properties[property.Name] = property.Value;
 			}
-			return Properties;
+			return properties;
 		}
 
-		static async Task AddAwsProperties(IList<string> Properties, ILogger Logger)
+		static async Task AddAwsProperties(IList<string> properties, ILogger logger)
 		{
 			if (EC2InstanceMetadata.IdentityDocument != null)
 			{
-				Properties.Add("EC2=1");
-				AddAwsProperty("aws-instance-id", "/instance-id", Properties);
-				AddAwsProperty("aws-instance-type", "/instance-type", Properties);
-				AddAwsProperty("aws-region", "/region", Properties);
+				properties.Add("EC2=1");
+				AddAwsProperty("aws-instance-id", "/instance-id", properties);
+				AddAwsProperty("aws-instance-type", "/instance-type", properties);
+				AddAwsProperty("aws-region", "/region", properties);
 
 				try
 				{
-					using (AmazonEC2Client Client = new AmazonEC2Client())
+					using (AmazonEC2Client client = new AmazonEC2Client())
 					{
-						DescribeTagsRequest Request = new DescribeTagsRequest();
-						Request.Filters = new List<Filter>();
-						Request.Filters.Add(new Filter("resource-id", new List<string> { EC2InstanceMetadata.InstanceId }));
+						DescribeTagsRequest request = new DescribeTagsRequest();
+						request.Filters = new List<Filter>();
+						request.Filters.Add(new Filter("resource-id", new List<string> { EC2InstanceMetadata.InstanceId }));
 
-						DescribeTagsResponse Response = await Client.DescribeTagsAsync(Request);
-						foreach (TagDescription Tag in Response.Tags)
+						DescribeTagsResponse response = await client.DescribeTagsAsync(request);
+						foreach (TagDescription tag in response.Tags)
 						{
-							Properties.Add($"aws-tag={Tag.Key}:{Tag.Value}");
+							properties.Add($"aws-tag={tag.Key}:{tag.Value}");
 						}
 					}
 				}
-				catch (Exception Ex)
+				catch (Exception ex)
 				{
-					Logger.LogDebug(Ex, "Unable to query EC2 tags.");
+					logger.LogDebug(ex, "Unable to query EC2 tags.");
 				}
 			}
 		}			
 
-		static void AddAwsProperty(string Name, string AwsKey, IList<string> Properties)
+		static void AddAwsProperty(string name, string awsKey, IList<string> properties)
 		{
-			string? Value = EC2InstanceMetadata.GetData(AwsKey);
-			if (Value != null)
+			string? value = EC2InstanceMetadata.GetData(awsKey);
+			if (value != null)
 			{
-				Properties.Add($"{Name}={Value}");
+				properties.Add($"{name}={value}");
 			}
 		}
 	}
