Horde: Modify agent code to conform to NET coding style.

diff --git a/Engine/Source/Programs/Horde/Horde.Agent/Tracing.cs b/Engine/Source/Programs/Horde/Horde.Agent/Tracing.cs
--- a/Engine/Source/Programs/Horde/Horde.Agent/Tracing.cs
+++ b/Engine/Source/Programs/Horde/Horde.Agent/Tracing.cs
@@ -9,22 +9,22 @@ namespace Horde.Agent
 {
 	static class Tracing
 	{
-		public static ISpanBuilder WithResourceName(this ISpanBuilder Builder, string? ResourceName)
+		public static ISpanBuilder WithResourceName(this ISpanBuilder builder, string? resourceName)
 		{
-			if (ResourceName != null)
+			if (resourceName != null)
 			{
-				Builder = Builder.WithTag(Datadog.Trace.OpenTracing.DatadogTags.ResourceName, ResourceName);
+				builder = builder.WithTag(Datadog.Trace.OpenTracing.DatadogTags.ResourceName, resourceName);
 			}
-			return Builder;
+			return builder;
 		}
 
-		public static ISpanBuilder WithServiceName(this ISpanBuilder Builder, string? ServiceName)
+		public static ISpanBuilder WithServiceName(this ISpanBuilder builder, string? serviceName)
 		{
-			if(ServiceName != null)
+			if(serviceName != null)
 			{
-				Builder = Builder.WithTag(Datadog.Trace.OpenTracing.DatadogTags.ServiceName, ServiceName);
+				builder = builder.WithTag(Datadog.Trace.OpenTracing.DatadogTags.ServiceName, serviceName);
 			}
-			return Builder;
+			return builder;
 		}
 }
 }
diff --git a/Engine/Source/Programs/Horde/Horde.Agent/Utilities/HordeRpc.cs b/Engine/Source/Programs/Horde/Horde.Agent/Utilities/HordeRpc.cs
--- a/Engine/Source/Programs/Horde/Horde.Agent/Utilities/HordeRpc.cs
+++ b/Engine/Source/Programs/Horde/Horde.Agent/Utilities/HordeRpc.cs
@@ -12,49 +12,48 @@ using System.Text;
 
 namespace HordeCommon.Rpc
 {
-#pragma warning disable CS1591
 	partial class Property
 	{
-		public Property(string Name, string Value)
+		public Property(string name, string value)
 		{
-			this.Name = Name;
-			this.Value = Value;
+			Name = name;
+			Value = value;
 		}
 
-		public Property(KeyValuePair<string, string> Pair)
+		public Property(KeyValuePair<string, string> pair)
 		{
-			this.Name = Pair.Key;
-			this.Value = Pair.Value;
+			Name = pair.Key;
+			Value = pair.Value;
 		}
 	}
 
 	partial class PropertyUpdate
 	{
-		public PropertyUpdate(string Name, string? Value)
+		public PropertyUpdate(string name, string? value)
 		{
-			this.Name = Name;
-			this.Value = Value;
+			Name = name;
+			Value = value;
 		}
 	}
 
 	static class PropertyExtensions
 	{
-		public static string GetValue(this RepeatedField<Property> Properties, string Name)
+		public static string GetValue(this RepeatedField<Property> properties, string name)
 		{
-			return Properties.First(x => x.Name == Name).Value;
+			return properties.First(x => x.Name == name).Value;
 		}
 
-		public static bool TryGetValue(this RepeatedField<Property> Properties, string Name, [MaybeNullWhen(false)] out string Result)
+		public static bool TryGetValue(this RepeatedField<Property> properties, string name, [MaybeNullWhen(false)] out string result)
 		{
-			Property? Property = Properties.FirstOrDefault(x => x.Name == Name);
-			if (Property == null)
+			Property? property = properties.FirstOrDefault(x => x.Name == name);
+			if (property == null)
 			{
-				Result = null!;
+				result = null!;
 				return false;
 			}
 			else
 			{
-				Result = Property.Value;
+				result = property.Value;
 				return true;
 			}
 		}
@@ -62,128 +61,127 @@ namespace HordeCommon.Rpc
 
 	partial class GetStreamRequest
 	{
-		public GetStreamRequest(string StreamId)
+		public GetStreamRequest(string streamId)
 		{
-			this.StreamId = StreamId;
+			StreamId = streamId;
 		}
 	}
 
 	partial class UpdateStreamRequest
 	{
-		public UpdateStreamRequest(string StreamId, Dictionary<string, string?> Properties)
+		public UpdateStreamRequest(string streamId, Dictionary<string, string?> properties)
 		{
-			this.StreamId = StreamId;
-			this.Properties.AddRange(Properties.Select(x => new PropertyUpdate(x.Key, x.Value)));
+			StreamId = streamId;
+			Properties.AddRange(properties.Select(x => new PropertyUpdate(x.Key, x.Value)));
 		}
 	}
 
 	partial class GetJobRequest
 	{
-		public GetJobRequest(string JobId)
+		public GetJobRequest(string jobId)
 		{
-			this.JobId = JobId;
+			JobId = jobId;
 		}
 	}
 
 	partial class BeginBatchRequest
 	{
-		public BeginBatchRequest(string JobId, string BatchId, string LeaseId)
+		public BeginBatchRequest(string jobId, string batchId, string leaseId)
 		{
-			this.JobId = JobId;
-			this.BatchId = BatchId;
-			this.LeaseId = LeaseId;
+			JobId = jobId;
+			BatchId = batchId;
+			LeaseId = leaseId;
 		}
 	}
 
 	partial class FinishBatchRequest
 	{
-		public FinishBatchRequest(string JobId, string BatchId, string LeaseId)
+		public FinishBatchRequest(string jobId, string batchId, string leaseId)
 		{
-			this.JobId = JobId;
-			this.BatchId = BatchId;
-			this.LeaseId = LeaseId;
+			JobId = jobId;
+			BatchId = batchId;
+			LeaseId = leaseId;
 		}
 	}
 
 	partial class BeginStepRequest
 	{
-		public BeginStepRequest(string JobId, string BatchId, string LeaseId)
+		public BeginStepRequest(string jobId, string batchId, string leaseId)
 		{
-			this.JobId = JobId;
-			this.BatchId = BatchId;
-			this.LeaseId = LeaseId;
+			JobId = jobId;
+			BatchId = batchId;
+			LeaseId = leaseId;
 		}
 	}
 
 	partial class UpdateStepRequest
 	{
-		public UpdateStepRequest(string JobId, string BatchId, string StepId, JobStepState State, JobStepOutcome Outcome)
+		public UpdateStepRequest(string jobId, string batchId, string stepId, JobStepState state, JobStepOutcome outcome)
 		{
-			this.JobId = JobId;
-			this.BatchId = BatchId;
-			this.StepId = StepId;
-			this.State = State;
-			this.Outcome = Outcome;
+			JobId = jobId;
+			BatchId = batchId;
+			StepId = stepId;
+			State = state;
+			Outcome = outcome;
 		}
 	}
 	
 	partial class GetStepRequest
 	{
-		public GetStepRequest(string JobId, string BatchId, string StepId)
+		public GetStepRequest(string jobId, string batchId, string stepId)
 		{
-			this.JobId = JobId;
-			this.BatchId = BatchId;
-			this.StepId = StepId;
+			JobId = jobId;
+			BatchId = batchId;
+			StepId = stepId;
 		}
 	}
 	
 	partial class GetStepResponse
 	{
-		public GetStepResponse(JobStepOutcome Outcome, JobStepState State, bool AbortRequested)
+		public GetStepResponse(JobStepOutcome outcome, JobStepState state, bool abortRequested)
 		{
-			this.Outcome = Outcome;
-			this.State = State;
-			this.AbortRequested = AbortRequested;
+			Outcome = outcome;
+			State = state;
+			AbortRequested = abortRequested;
 		}
 	}
 
 	partial class CreateEventRequest
 	{
-		public CreateEventRequest(EventSeverity Severity, string LogId, int LineIndex, int LineCount)
+		public CreateEventRequest(EventSeverity severity, string logId, int lineIndex, int lineCount)
 		{
-			this.Severity = Severity;
-			this.LogId = LogId;
-			this.LineIndex = LineIndex;
-			this.LineCount = LineCount;
+			Severity = severity;
+			LogId = logId;
+			LineIndex = lineIndex;
+			LineCount = lineCount;
 		}
 	}
 
 	partial class CreateEventsRequest
 	{
-		public CreateEventsRequest(IEnumerable<CreateEventRequest> Events)
+		public CreateEventsRequest(IEnumerable<CreateEventRequest> events)
 		{
-			this.Events.AddRange(Events);
+			Events.AddRange(events);
 		}
 	}
 
 	partial class WriteOutputRequest
 	{
-		public WriteOutputRequest(string LogId, long Offset, int LineIndex, byte[] Data, bool Flush)
+		public WriteOutputRequest(string logId, long offset, int lineIndex, byte[] data, bool flush)
 		{
-			this.LogId = LogId;
-			this.Offset = Offset;
-			this.LineIndex = LineIndex;
-			this.Data = ByteString.CopyFrom(Data);
-			this.Flush = Flush;
+			LogId = logId;
+			Offset = offset;
+			LineIndex = lineIndex;
+			Data = ByteString.CopyFrom(data);
+			Flush = flush;
 		}
 	}
 
 	partial class DownloadSoftwareRequest
 	{
-		public DownloadSoftwareRequest(string Version)
+		public DownloadSoftwareRequest(string version)
 		{
-			this.Version = Version;
+			Version = version;
 		}
 	}
-#pragma warning restore CS1591
 }
diff --git a/Engine/Source/Programs/Horde/Horde.Agent/Utilities/ProtoMethods.cs b/Engine/Source/Programs/Horde/Horde.Agent/Utilities/ProtoMethods.cs
--- a/Engine/Source/Programs/Horde/Horde.Agent/Utilities/ProtoMethods.cs
+++ b/Engine/Source/Programs/Horde/Horde.Agent/Utilities/ProtoMethods.cs
@@ -17,18 +17,18 @@ namespace HordeCommon.Rpc.Tasks
 		/// <summary>
 		/// Constructor
 		/// </summary>
-		public ComputeTaskResultMessage(RefIdWrapper ResultRefId)
+		public ComputeTaskResultMessage(RefIdWrapper resultRefId)
 		{
-			this.ResultRefId = ResultRefId;
+			ResultRefId = resultRefId;
 		}
 
 		/// <summary>
 		/// Constructor
 		/// </summary>
-		public ComputeTaskResultMessage(ComputeTaskOutcome Outcome, string? Detail = null)
+		public ComputeTaskResultMessage(ComputeTaskOutcome outcome, string? detail = null)
 		{
-			this.Outcome = (int)Outcome;
-			this.Detail = Detail;
+			Outcome = (int)outcome;
+			Detail = detail;
 		}
 	}
 }
diff --git a/Engine/Source/Programs/Horde/Horde.Agent/Utility/ArtifactUploader.cs b/Engine/Source/Programs/Horde/Horde.Agent/Utility/ArtifactUploader.cs
--- a/Engine/Source/Programs/Horde/Horde.Agent/Utility/ArtifactUploader.cs
+++ b/Engine/Source/Programs/Horde/Horde.Agent/Utility/ArtifactUploader.cs
@@ -25,7 +25,7 @@ namespace Horde.Agent.Utility
         /// <summary>
         /// Dictionary of horde common mime types
         /// </summary>
-        private static Dictionary<string, string> HordeMimeTypes = new Dictionary<string, string>()
+        private static readonly Dictionary<string, string> s_hordeMimeTypes = new Dictionary<string, string>()
         {
             { ".bin", "application/octet-stream" },
             { ".json", "application/json" },
@@ -40,40 +40,40 @@ namespace Horde.Agent.Utility
         /// <summary>
         /// Gets a mime type from a file extension
         /// </summary>
-        /// <param name="File">The file object to parse</param>
+        /// <param name="file">The file object to parse</param>
         /// <returns>the mimetype if present in the dictionary, binary otherwise</returns>
-        private static string GetMimeType(FileReference File)
+        private static string GetMimeType(FileReference file)
         {
-            string? ContentType;
-            if (!HordeMimeTypes.TryGetValue(File.GetExtension(), out ContentType))
+            string? contentType;
+            if (!s_hordeMimeTypes.TryGetValue(file.GetExtension(), out contentType))
             {
-                ContentType = "application/octet-stream";
+                contentType = "application/octet-stream";
             }
-            return ContentType;
+            return contentType;
         }
         
         /// <summary>
         /// Uploads an artifact (with retries)
         /// </summary>
-        /// <param name="RpcConnection">The grpc client</param>
-        /// <param name="JobId">Job id</param>
-        /// <param name="BatchId">Job batch id</param>
-        /// <param name="StepId">Job step id</param>
-        /// <param name="ArtifactName">Name of the artifact</param>
-        /// <param name="ArtifactFile">File to upload</param>
-        /// <param name="Logger">Logger interfact</param>
-        /// <param name="CancellationToken">Cancellation token</param>
+        /// <param name="rpcConnection">The grpc client</param>
+        /// <param name="jobId">Job id</param>
+        /// <param name="batchId">Job batch id</param>
+        /// <param name="stepId">Job step id</param>
+        /// <param name="artifactName">Name of the artifact</param>
+        /// <param name="artifactFile">File to upload</param>
+        /// <param name="logger">Logger interfact</param>
+        /// <param name="cancellationToken">Cancellation token</param>
         /// <returns></returns>
-        public static async Task<string?> UploadAsync(IRpcConnection RpcConnection, string JobId, string BatchId, string StepId, string ArtifactName, FileReference ArtifactFile, ILogger Logger, CancellationToken CancellationToken)
+        public static async Task<string?> UploadAsync(IRpcConnection rpcConnection, string jobId, string batchId, string stepId, string artifactName, FileReference artifactFile, ILogger logger, CancellationToken cancellationToken)
         {
 			try
 			{
-				string ArtifactId = await RpcConnection.InvokeAsync(RpcClient => DoUploadAsync(RpcClient, JobId, BatchId, StepId, ArtifactName, ArtifactFile, Logger, CancellationToken), new RpcContext(), CancellationToken);
-				return ArtifactId;
+				string artifactId = await rpcConnection.InvokeAsync(rpcClient => DoUploadAsync(rpcClient, jobId, batchId, stepId, artifactName, artifactFile, logger, cancellationToken), new RpcContext(), cancellationToken);
+				return artifactId;
 			}
-			catch (Exception Ex)
+			catch (Exception ex)
 			{
-				Logger.LogInformation(KnownLogEvents.Systemic_Horde_ArtifactUpload, Ex, "Exception while attempting to upload artifact: {Message}", Ex.Message);
+				logger.LogInformation(KnownLogEvents.Systemic_Horde_ArtifactUpload, ex, "Exception while attempting to upload artifact: {Message}", ex.Message);
 				return null;
 			}
         }
@@ -81,59 +81,59 @@ namespace Horde.Agent.Utility
         /// <summary>
 		/// Uploads an artifact
 		/// </summary>
-		/// <param name="Client">The grpc client</param>
-		/// <param name="JobId">Job id</param>
-		/// <param name="BatchId">Job batch id</param>
-		/// <param name="StepId">Job step id</param>
-		/// <param name="ArtifactName">Name of the artifact</param>
-		/// <param name="ArtifactFile">File to upload</param>
-		/// <param name="Logger">Logger interfact</param>
-		/// <param name="CancellationToken">Cancellation token</param>
+		/// <param name="client">The grpc client</param>
+		/// <param name="jobId">Job id</param>
+		/// <param name="batchId">Job batch id</param>
+		/// <param name="stepId">Job step id</param>
+		/// <param name="artifactName">Name of the artifact</param>
+		/// <param name="artifactFile">File to upload</param>
+		/// <param name="logger">Logger interfact</param>
+		/// <param name="cancellationToken">Cancellation token</param>
 		/// <returns></returns>
-		private static async Task<string> DoUploadAsync(HordeRpc.HordeRpcClient Client, string JobId, string BatchId, string StepId, string ArtifactName, FileReference ArtifactFile, ILogger Logger, CancellationToken CancellationToken)
+		private static async Task<string> DoUploadAsync(HordeRpc.HordeRpcClient client, string jobId, string batchId, string stepId, string artifactName, FileReference artifactFile, ILogger logger, CancellationToken cancellationToken)
         {
-			Logger.LogInformation("Uploading artifact {ArtifactName} from {ArtifactFile}", ArtifactName, ArtifactFile);
-            using (FileStream ArtifactStream = FileReference.Open(ArtifactFile, FileMode.Open, FileAccess.Read, FileShare.Read))
+			logger.LogInformation("Uploading artifact {ArtifactName} from {ArtifactFile}", artifactName, artifactFile);
+            using (FileStream artifactStream = FileReference.Open(artifactFile, FileMode.Open, FileAccess.Read, FileShare.Read))
 			{
-				using (AsyncClientStreamingCall<UploadArtifactRequest, UploadArtifactResponse> Cursor = Client.UploadArtifact(null, null, CancellationToken))
+				using (AsyncClientStreamingCall<UploadArtifactRequest, UploadArtifactResponse> cursor = client.UploadArtifact(null, null, cancellationToken))
                 {
 					// Upload the metadata in the initial request
-					UploadArtifactMetadata Metadata = new UploadArtifactMetadata();
-					Metadata.JobId = JobId;
-					Metadata.BatchId = BatchId;
-					Metadata.StepId = StepId;
-					Metadata.Name = ArtifactName;
-					Metadata.MimeType = GetMimeType(ArtifactFile);
-					Metadata.Length = ArtifactStream.Length;
+					UploadArtifactMetadata metadata = new UploadArtifactMetadata();
+					metadata.JobId = jobId;
+					metadata.BatchId = batchId;
+					metadata.StepId = stepId;
+					metadata.Name = artifactName;
+					metadata.MimeType = GetMimeType(artifactFile);
+					metadata.Length = artifactStream.Length;
 
-					UploadArtifactRequest InitialRequest = new UploadArtifactRequest();
-					InitialRequest.Metadata = Metadata;
+					UploadArtifactRequest initialRequest = new UploadArtifactRequest();
+					initialRequest.Metadata = metadata;
 
-					await Cursor.RequestStream.WriteAsync(InitialRequest);
+					await cursor.RequestStream.WriteAsync(initialRequest);
 
 					// Upload the data in chunks
-					byte[] Buffer = new byte[4096];
-					for (int Offset = 0; Offset < Metadata.Length; )
+					byte[] buffer = new byte[4096];
+					for (int offset = 0; offset < metadata.Length; )
                     {
-                        int BytesRead = await ArtifactStream.ReadAsync(Buffer, 0, Buffer.Length);
-						if(BytesRead == 0)
+                        int bytesRead = await artifactStream.ReadAsync(buffer, 0, buffer.Length);
+						if(bytesRead == 0)
 						{
-							throw new InvalidDataException($"Unable to read data from {ArtifactFile} beyond offset {Offset}; expected length to be {Metadata.Length}");
+							throw new InvalidDataException($"Unable to read data from {artifactFile} beyond offset {offset}; expected length to be {metadata.Length}");
 						}
 
-						UploadArtifactRequest Request = new UploadArtifactRequest();
-						Request.Data = Google.Protobuf.ByteString.CopyFrom(Buffer, 0, BytesRead);
-						await Cursor.RequestStream.WriteAsync(Request);
+						UploadArtifactRequest request = new UploadArtifactRequest();
+						request.Data = Google.Protobuf.ByteString.CopyFrom(buffer, 0, bytesRead);
+						await cursor.RequestStream.WriteAsync(request);
 
-                        Offset += BytesRead;
+                        offset += bytesRead;
                     }
 
 					// Close the stream
-                    await Cursor.RequestStream.CompleteAsync();
+                    await cursor.RequestStream.CompleteAsync();
 
 					// Read the response
-					UploadArtifactResponse Response = await Cursor.ResponseAsync;
-					return Response.Id;
+					UploadArtifactResponse response = await cursor.ResponseAsync;
+					return response.Id;
                 }
             }
         }
diff --git a/Engine/Source/Programs/Horde/Horde.Agent/Utility/CertificateHelper.cs b/Engine/Source/Programs/Horde/Horde.Agent/Utility/CertificateHelper.cs
--- a/Engine/Source/Programs/Horde/Horde.Agent/Utility/CertificateHelper.cs
+++ b/Engine/Source/Programs/Horde/Horde.Agent/Utility/CertificateHelper.cs
@@ -20,93 +20,92 @@ namespace Horde.Agent.Utility
 		/// <summary>
 		/// Provides additional diagnostic information for SSL certificate validation
 		/// </summary>
-		/// <param name="Logger">The logger instance</param>
-		/// <param name="Sender"></param>
-		/// <param name="Certificate"></param>
-		/// <param name="Chain"></param>
-		/// <param name="SslPolicyErrors"></param>
-		/// <param name="ServerProfile">The server profile</param>
+		/// <param name="logger">The logger instance</param>
+		/// <param name="sender"></param>
+		/// <param name="certificate"></param>
+		/// <param name="chain"></param>
+		/// <param name="sslPolicyErrors"></param>
+		/// <param name="serverProfile">The server profile</param>
 		/// <returns>True if the certificate is allowed, false otherwise</returns>
-		public static bool CertificateValidationCallBack(ILogger Logger, object Sender, X509Certificate Certificate, X509Chain Chain, SslPolicyErrors SslPolicyErrors, ServerProfile ServerProfile)
+		public static bool CertificateValidationCallBack(ILogger logger, object sender, X509Certificate certificate, X509Chain chain, SslPolicyErrors sslPolicyErrors, ServerProfile serverProfile)
 		{
 			// If the certificate is a valid, signed certificate, return true.
-			if (SslPolicyErrors == SslPolicyErrors.None)
+			if (sslPolicyErrors == SslPolicyErrors.None)
 			{
 				return true;
 			}
 
 			// Trust the remote certificate if it has the right thumbprint
-			if (SslPolicyErrors == SslPolicyErrors.RemoteCertificateChainErrors)
+			if (sslPolicyErrors == SslPolicyErrors.RemoteCertificateChainErrors)
 			{
-				if (Chain.ChainElements.Count == 1)
+				if (chain.ChainElements.Count == 1)
 				{
-					X509ChainElement Element = Chain.ChainElements[0];
-					if (Element.ChainElementStatus.Length == 1 && (Element.ChainElementStatus[0].Status == X509ChainStatusFlags.UntrustedRoot || Element.ChainElementStatus[0].Status == X509ChainStatusFlags.PartialChain))
+					X509ChainElement element = chain.ChainElements[0];
+					if (element.ChainElementStatus.Length == 1 && (element.ChainElementStatus[0].Status == X509ChainStatusFlags.UntrustedRoot || element.ChainElementStatus[0].Status == X509ChainStatusFlags.PartialChain))
 					{
-						if (ServerProfile.IsTrustedCertificate(Element.Certificate.Thumbprint))
+						if (serverProfile.IsTrustedCertificate(element.Certificate.Thumbprint))
 						{
-							Logger.LogDebug("Trusting server certificate {Thumbprint}", Element.Certificate.Thumbprint);
+							logger.LogDebug("Trusting server certificate {Thumbprint}", element.Certificate.Thumbprint);
 							return true;
 						}
 					}
-
 				}
 			}
 
 			// Generate diagnostic information
-			StringBuilder Builder = new StringBuilder();
-			if (Sender != null)
+			StringBuilder builder = new StringBuilder();
+			if (sender != null)
 			{
-				HttpRequestMessage? Message = Sender as HttpRequestMessage;
-				if (Message != null)
+				HttpRequestMessage? message = sender as HttpRequestMessage;
+				if (message != null)
 				{
-					Builder.Append($"\nSender: {Message.Method} {Message.RequestUri}");
+					builder.Append($"\nSender: {message.Method} {message.RequestUri}");
 				}
 				else
 				{
-					string SenderInfo = StringUtils.Indent(Sender.ToString() ?? String.Empty, "    ");
-					Builder.Append($"\nSender:\n{SenderInfo}");
+					string senderInfo = StringUtils.Indent(sender.ToString() ?? String.Empty, "    ");
+					builder.Append($"\nSender:\n{senderInfo}");
 				}
 			}
-			if (Certificate != null)
+			if (certificate != null)
 			{
-				Builder.Append($"\nCertificate: {Certificate.Subject}");
+				builder.Append($"\nCertificate: {certificate.Subject}");
 			}
-			if (Chain != null)
+			if (chain != null)
 			{
-				if (Chain.ChainStatus != null && Chain.ChainStatus.Length > 0)
+				if (chain.ChainStatus != null && chain.ChainStatus.Length > 0)
 				{
-					Builder.Append("\nChain status:");
-					foreach (X509ChainStatus Status in Chain.ChainStatus)
+					builder.Append("\nChain status:");
+					foreach (X509ChainStatus status in chain.ChainStatus)
 					{
-						Builder.Append($"\n  {Status.StatusInformation}");
+						builder.Append($"\n  {status.StatusInformation}");
 					}
 				}
-				if (Chain.ChainElements != null)
+				if (chain.ChainElements != null)
 				{
-					Builder.Append("\nChain elements:");
-					for (int Idx = 0; Idx < Chain.ChainElements.Count; Idx++)
+					builder.Append("\nChain elements:");
+					for (int idx = 0; idx < chain.ChainElements.Count; idx++)
 					{
-						X509ChainElement Element = Chain.ChainElements[Idx];
-						Builder.Append($"\n  {Idx,4} - Certificate: {Element.Certificate.Subject}");
-						Builder.Append($"\n         Thumbprint: {Element.Certificate.Thumbprint}");
-						if (Element.ChainElementStatus != null && Element.ChainElementStatus.Length > 0)
+						X509ChainElement element = chain.ChainElements[idx];
+						builder.Append($"\n  {idx,4} - Certificate: {element.Certificate.Subject}");
+						builder.Append($"\n         Thumbprint: {element.Certificate.Thumbprint}");
+						if (element.ChainElementStatus != null && element.ChainElementStatus.Length > 0)
 						{
-							foreach (X509ChainStatus Status in Element.ChainElementStatus)
+							foreach (X509ChainStatus status in element.ChainElementStatus)
 							{
-								Builder.Append($"\n         Status: {Status.StatusInformation} ({Status.Status})");
+								builder.Append($"\n         Status: {status.StatusInformation} ({status.Status})");
 							}
 						}
-						if (!String.IsNullOrEmpty(Element.Information))
+						if (!String.IsNullOrEmpty(element.Information))
 						{
-							Builder.Append($"\n         Info: {Element.Information}");
+							builder.Append($"\n         Info: {element.Information}");
 						}
 					}
 				}
 			}
 
 			// Print out additional diagnostic information
-			Logger.LogError("TLS certificate validation failed ({Errors}).{AdditionalInfo}", SslPolicyErrors, StringUtils.Indent(Builder.ToString(), "    "));
+			logger.LogError("TLS certificate validation failed ({Errors}).{AdditionalInfo}", sslPolicyErrors, StringUtils.Indent(builder.ToString(), "    "));
 			return false;
 		}
 	}
diff --git a/Engine/Source/Programs/Horde/Horde.Agent/Utility/JsonLogger.cs b/Engine/Source/Programs/Horde/Horde.Agent/Utility/JsonLogger.cs
--- a/Engine/Source/Programs/Horde/Horde.Agent/Utility/JsonLogger.cs
+++ b/Engine/Source/Programs/Horde/Horde.Agent/Utility/JsonLogger.cs
@@ -32,66 +32,66 @@ namespace Horde.Agent.Utility
 		/// <summary>
 		/// Constructor
 		/// </summary>
-		/// <param name="Warnings">Whether to include warnings in the output</param>
-		/// <param name="Inner">Additional logger to write to</param>
-		public JsonLogger(bool? Warnings, ILogger Inner)
+		/// <param name="warnings">Whether to include warnings in the output</param>
+		/// <param name="inner">Additional logger to write to</param>
+		public JsonLogger(bool? warnings, ILogger inner)
 		{
-			this.Warnings = Warnings;
-			this.Inner = Inner;
+			Warnings = warnings;
+			Inner = inner;
 		}
 
 		/// <inheritdoc/>
-		public IDisposable? BeginScope<TState>(TState State)
+		public IDisposable? BeginScope<TState>(TState state)
 		{
 			return null;
 		}
 
 		/// <inheritdoc/>
-		public bool IsEnabled(LogLevel LogLevel)
+		public bool IsEnabled(LogLevel logLevel)
 		{
 			return true;
 		}
 
 		/// <inheritdoc/>
-		public void Log<TState>(LogLevel LogLevel, EventId EventId, TState State, Exception? Exception, Func<TState, Exception?, string> Formatter)
+		public void Log<TState>(LogLevel logLevel, EventId eventId, TState state, Exception? exception, Func<TState, Exception?, string> formatter)
 		{
 			// Downgrade warnings to information if not required
-			if (LogLevel == LogLevel.Warning && !(Warnings ?? true))
+			if (logLevel == LogLevel.Warning && !(Warnings ?? true))
 			{
-				LogLevel = LogLevel.Information;
+				logLevel = LogLevel.Information;
 			}
 
-			if(State is JsonLogEvent JsonEvent)
+			if(state is JsonLogEvent jsonEvent)
 			{
-				WriteFormattedEvent(LogLevel, JsonEvent.LineIndex, JsonEvent.LineCount, JsonEvent.Data.ToArray());
+				WriteFormattedEvent(logLevel, jsonEvent.LineIndex, jsonEvent.LineCount, jsonEvent.Data.ToArray());
 				return;
 			}
 
-			LogEvent? Event = State as LogEvent;
-			if (Event == null)
+			LogEvent? logEvent = state as LogEvent;
+			if (logEvent == null)
 			{
-				Event = LogEvent.FromState(LogLevel, EventId, State, Exception, Formatter);
+				logEvent = LogEvent.FromState(logLevel, eventId, state, exception, formatter);
 			}
-			WriteFormattedEvent(Event.Level, Event.LineIndex, Event.LineCount, Event.ToJsonBytes());
+			WriteFormattedEvent(logEvent.Level, logEvent.LineIndex, logEvent.LineCount, logEvent.ToJsonBytes());
 		}
 
 		/// <summary>
 		/// Writes a formatted event
 		/// </summary>
-		/// <param name="Level">The log level</param>
-		/// <param name="LineIndex">Index of the current line within this event</param>
-		/// <param name="LineCount">Number of lines in this event</param>
-		/// <param name="Line">Utf-8 encoded JSON line data</param>
-		protected abstract void WriteFormattedEvent(LogLevel Level, int LineIndex, int LineCount, byte[] Line);
+		/// <param name="level">The log level</param>
+		/// <param name="lineIndex">Index of the current line within this event</param>
+		/// <param name="lineCount">Number of lines in this event</param>
+		/// <param name="line">Utf-8 encoded JSON line data</param>
+		protected abstract void WriteFormattedEvent(LogLevel level, int lineIndex, int lineCount, byte[] line);
 
 		/// <summary>
 		/// Callback to write a systemic event
 		/// </summary>
-		/// <param name="EventId">The event id</param>
-		/// <param name="Text">The event text</param>
-		protected virtual void WriteSystemicEvent(EventId EventId, string Text)
+		/// <param name="eventId">The event id</param>
+		/// <param name="text">The event text</param>
+		protected virtual void WriteSystemicEvent(EventId eventId, string text)
 		{
-			Inner.LogWarning("Systemic event {KnownLogEventId}: {Text}", EventId.Id, Text);
+			Inner.LogWarning("Systemic event {KnownLogEventId}: {Text}", eventId.Id, text);
 		}
 	}
 }
diff --git a/Engine/Source/Programs/Horde/Horde.Agent/Utility/JsonRpcLogger.cs b/Engine/Source/Programs/Horde/Horde.Agent/Utility/JsonRpcLogger.cs
--- a/Engine/Source/Programs/Horde/Horde.Agent/Utility/JsonRpcLogger.cs
+++ b/Engine/Source/Programs/Horde/Horde.Agent/Utility/JsonRpcLogger.cs
@@ -27,13 +27,13 @@ namespace Horde.Agent.Parser
 	{
 		class QueueItem
 		{
-			public byte[] Data;
-			public CreateEventRequest? Event;
+			public byte[] Data { get; }
+			public CreateEventRequest? CreateEvent { get; }
 
-			public QueueItem(byte[] Data, CreateEventRequest? Event)
+			public QueueItem(byte[] data, CreateEventRequest? createEvent)
 			{
-				this.Data = Data;
-				this.Event = Event;
+				Data = data;
+				CreateEvent = createEvent;
 			}
 
 			public override string ToString()
@@ -42,13 +42,13 @@ namespace Horde.Agent.Parser
 			}
 		}
 
-		readonly IRpcConnection RpcClient;
-		readonly string LogId;
-		readonly string? JobId;
-		readonly string? JobBatchId;
-		readonly string? JobStepId;
-		Channel<QueueItem> DataChannel;
-		Task? DataWriter;
+		readonly IRpcConnection _rpcClient;
+		readonly string _logId;
+		readonly string? _jobId;
+		readonly string? _jobBatchId;
+		readonly string? _jobStepId;
+		readonly Channel<QueueItem> _dataChannel;
+		Task? _dataWriter;
 
 		/// <summary>
 		/// The current outcome for this step. Updated to reflect any errors and warnings that occurred.
@@ -62,51 +62,51 @@ namespace Horde.Agent.Parser
 		/// <summary>
 		/// Constructor
 		/// </summary>
-		/// <param name="RpcClient">RPC client to use for server requests</param>
-		/// <param name="LogId">The log id to write to</param>
-		/// <param name="JobId">Id of the job being executed</param>
-		/// <param name="JobBatchId">Batch being executed</param>
-		/// <param name="JobStepId">Id of the step being executed</param>
-		/// <param name="Warnings">Whether to include warnings in the output</param>
-		/// <param name="Inner">Additional logger to write to</param>
-		public JsonRpcLogger(IRpcConnection RpcClient, string LogId, string? JobId, string? JobBatchId, string? JobStepId, bool? Warnings, ILogger Inner)
-			: base(Warnings, Inner)
+		/// <param name="rpcClient">RPC client to use for server requests</param>
+		/// <param name="logId">The log id to write to</param>
+		/// <param name="jobId">Id of the job being executed</param>
+		/// <param name="jobBatchId">Batch being executed</param>
+		/// <param name="jobStepId">Id of the step being executed</param>
+		/// <param name="warnings">Whether to include warnings in the output</param>
+		/// <param name="inner">Additional logger to write to</param>
+		public JsonRpcLogger(IRpcConnection rpcClient, string logId, string? jobId, string? jobBatchId, string? jobStepId, bool? warnings, ILogger inner)
+			: base(warnings, inner)
 		{
-			this.RpcClient = RpcClient;
-			this.LogId = LogId;
-			this.JobId = JobId;
-			this.JobBatchId = JobBatchId;
-			this.JobStepId = JobStepId;
-			this.DataChannel = Channel.CreateUnbounded<QueueItem>();
-			this.DataWriter = Task.Run(() => RunDataWriter());
-			this.Outcome = JobStepOutcome.Success;
+			_rpcClient = rpcClient;
+			_logId = logId;
+			_jobId = jobId;
+			_jobBatchId = jobBatchId;
+			_jobStepId = jobStepId;
+			_dataChannel = Channel.CreateUnbounded<QueueItem>();
+			_dataWriter = Task.Run(() => RunDataWriter());
+			Outcome = JobStepOutcome.Success;
 		}
 
-		protected override void WriteFormattedEvent(LogLevel Level, int LineIndex, int LineCount, byte[] Line)
+		protected override void WriteFormattedEvent(LogLevel level, int lineIndex, int lineCount, byte[] line)
 		{
 			// Update the state of this job if this is an error status
-			if (Level == LogLevel.Error || Level == LogLevel.Critical)
+			if (level == LogLevel.Error || level == LogLevel.Critical)
 			{
 				Outcome = JobStepOutcome.Failure;
 			}
-			else if (Level == LogLevel.Warning && Outcome != JobStepOutcome.Failure)
+			else if (level == LogLevel.Warning && Outcome != JobStepOutcome.Failure)
 			{
 				Outcome = JobStepOutcome.Warnings;
 			}
 
 			// If we want an event for this log event, create one now
-			CreateEventRequest? Event = null;
-			if (LineIndex == 0)
+			CreateEventRequest? eventRequest = null;
+			if (lineIndex == 0)
 			{
-				if (Level == LogLevel.Warning || Level == LogLevel.Error || Level == LogLevel.Critical)
+				if (level == LogLevel.Warning || level == LogLevel.Error || level == LogLevel.Critical)
 				{
-					Event = CreateEvent(Level, LineCount);
+					eventRequest = CreateEvent(level, lineCount);
 				}
 			}
 
 			// Write the data to the output channel
-			QueueItem QueueItem = new QueueItem(Line, Event);
-			if (!DataChannel.Writer.TryWrite(QueueItem))
+			QueueItem queueItem = new QueueItem(line, eventRequest);
+			if (!_dataChannel.Writer.TryWrite(queueItem))
 			{
 				throw new InvalidOperationException("Expected unbounded writer to complete immediately");
 			}
@@ -115,37 +115,37 @@ namespace Horde.Agent.Parser
 		/// <summary>
 		/// Callback to write a systemic event
 		/// </summary>
-		/// <param name="EventId">The event id</param>
-		/// <param name="Text">The event text</param>
-		protected override void WriteSystemicEvent(EventId EventId, string Text)
+		/// <param name="eventId">The event id</param>
+		/// <param name="text">The event text</param>
+		protected override void WriteSystemicEvent(EventId eventId, string text)
 		{
-			if (JobId == null)
+			if (_jobId == null)
 			{
-				Inner.LogWarning("Systemic event {KnownLogEventId} in log {LogId}: {Text}", EventId.Id, Text);
+				Inner.LogWarning("Systemic event {KnownLogEventId} in log {LogId}: {Text}", eventId.Id, text);
 			}
-			else if (JobBatchId == null)
+			else if (_jobBatchId == null)
 			{
-				Inner.LogWarning("Systemic event {KnownLogEventId} in log {LogId} (job {JobId}): {Text})", EventId.Id, JobId, Text);
+				Inner.LogWarning("Systemic event {KnownLogEventId} in log {LogId} (job {JobId}): {Text})", eventId.Id, _jobId, text);
 			}
-			else if (JobStepId == null)
+			else if (_jobStepId == null)
 			{
-				Inner.LogWarning("Systemic event {KnownLogEventId} in log {LogId} (job batch {JobId}:{BatchId}): {Text})", EventId.Id, JobId, JobStepId, Text);
+				Inner.LogWarning("Systemic event {KnownLogEventId} in log {LogId} (job batch {JobId}:{BatchId}): {Text})", eventId.Id, _jobId, _jobStepId, text);
 			}
 			else
 			{
-				Inner.LogWarning("Systemic event {KnownLogEventId} in log {LogId} (job step {JobId}:{BatchId}:{StepId}): {Text}", EventId.Id, LogId, JobId, JobBatchId, JobStepId, Text);
+				Inner.LogWarning("Systemic event {KnownLogEventId} in log {LogId} (job step {JobId}:{BatchId}:{StepId}): {Text}", eventId.Id, _logId, _jobId, _jobBatchId, _jobStepId, text);
 			}
 		}
 
 		/// <summary>
 		/// Makes a <see cref="CreateEventRequest"/> for the given parameters
 		/// </summary>
-		/// <param name="LogLevel">Level for this log event</param>
-		/// <param name="LineCount">Number of lines in the event</param>
-		CreateEventRequest CreateEvent(LogLevel LogLevel, int LineCount)
+		/// <param name="logLevel">Level for this log event</param>
+		/// <param name="lineCount">Number of lines in the event</param>
+		CreateEventRequest CreateEvent(LogLevel logLevel, int lineCount)
 		{
-			EventSeverity Severity = (LogLevel == LogLevel.Warning) ? EventSeverity.Warning : EventSeverity.Error;
-			return new CreateEventRequest(Severity, LogId, 0, LineCount);
+			EventSeverity severity = (logLevel == LogLevel.Warning) ? EventSeverity.Warning : EventSeverity.Error;
+			return new CreateEventRequest(severity, _logId, 0, lineCount);
 		}
 
 		/// <summary>
@@ -154,11 +154,11 @@ namespace Horde.Agent.Parser
 		/// <returns>Async task</returns>
 		public async Task StopAsync()
 		{
-			if (DataWriter != null)
+			if (_dataWriter != null)
 			{
-				DataChannel.Writer.TryComplete();
-				await DataWriter;
-				DataWriter = null;
+				_dataChannel.Writer.TryComplete();
+				await _dataWriter;
+				_dataWriter = null;
 			}
 		}
 
@@ -177,64 +177,64 @@ namespace Horde.Agent.Parser
 		async Task RunDataWriter()
 		{
 			// Current position and line number in the log file
-			long Offset = 0;
-			int LineIndex = 0;
+			long offset = 0;
+			int lineIndex = 0;
 
 			// Total number of errors and warnings
 			const int MaxErrors = 50;
-			int NumErrors = 0;
+			int numErrors = 0;
 			const int MaxWarnings = 50;
-			int NumWarnings = 0;
+			int numWarnings = 0;
 
 			// Buffers for chunks and events read in a single iteration
-			ArrayBufferWriter<byte> Writer = new ArrayBufferWriter<byte>();
-			List<CreateEventRequest> Events = new List<CreateEventRequest>();
+			ArrayBufferWriter<byte> writer = new ArrayBufferWriter<byte>();
+			List<CreateEventRequest> events = new List<CreateEventRequest>();
 
 			// Line separator for JSON events
-			byte[] Newline = { (byte)'\n' };
-			JsonEncodedText Timestamp = JsonEncodedText.Encode("");
+			byte[] newline = { (byte)'\n' };
+			JsonEncodedText timestamp = JsonEncodedText.Encode("");
 
 			// The current jobstep outcome
-			JobStepOutcome PostedOutcome = JobStepOutcome.Success;
+			JobStepOutcome postedOutcome = JobStepOutcome.Success;
 
 			// Whether we've written the flush command
 			for (; ; )
 			{
-				Writer.Clear();
-				Events.Clear();
+				writer.Clear();
+				events.Clear();
 
 				// Save off the current line number for sending to the server
-				int InitialLineIndex = LineIndex;
+				int initialLineIndex = lineIndex;
 
 				// Get the next data
-				Task WaitTask = Task.Delay(TimeSpan.FromSeconds(2.0));
-				while (Writer.WrittenCount < 256 * 1024)
+				Task waitTask = Task.Delay(TimeSpan.FromSeconds(2.0));
+				while (writer.WrittenCount < 256 * 1024)
 				{
-					QueueItem? Data;
-					if (DataChannel.Reader.TryRead(out Data))
+					QueueItem? data;
+					if (_dataChannel.Reader.TryRead(out data))
 					{
-						if (Data.Event != null)
+						if (data.CreateEvent != null)
 						{
-							if ((Data.Event.Severity == EventSeverity.Warning && ++NumWarnings <= MaxWarnings) || (Data.Event.Severity == EventSeverity.Error && ++NumErrors <= MaxErrors))
+							if ((data.CreateEvent.Severity == EventSeverity.Warning && ++numWarnings <= MaxWarnings) || (data.CreateEvent.Severity == EventSeverity.Error && ++numErrors <= MaxErrors))
 							{
-								Data.Event.LineIndex = LineIndex;
-								Events.Add(Data.Event);
+								data.CreateEvent.LineIndex = lineIndex;
+								events.Add(data.CreateEvent);
 							}
 						}
 
-						Writer.Write(Data.Data);
-						Writer.Write(Newline);
+						writer.Write(data.Data);
+						writer.Write(newline);
 
-						LineIndex++;
+						lineIndex++;
 					}
 					else
 					{
-						Task<bool> ReadTask = DataChannel.Reader.WaitToReadAsync().AsTask();
-						if (await Task.WhenAny(ReadTask, WaitTask) == WaitTask)
+						Task<bool> readTask = _dataChannel.Reader.WaitToReadAsync().AsTask();
+						if (await Task.WhenAny(readTask, waitTask) == waitTask)
 						{
 							break;
 						}
-						if (!await ReadTask)
+						if (!await readTask)
 						{
 							break;
 						}
@@ -242,57 +242,57 @@ namespace Horde.Agent.Parser
 				}
 
 				// Upload it to the server
-				if (Writer.WrittenCount > 0)
+				if (writer.WrittenCount > 0)
 				{
-					byte[] Data = Writer.WrittenSpan.ToArray();
+					byte[] data = writer.WrittenSpan.ToArray();
 					try
 					{
-						await RpcClient.InvokeAsync(x => x.WriteOutputAsync(new WriteOutputRequest(LogId, Offset, InitialLineIndex, Data, false)), new RpcContext(), CancellationToken.None);
+						await _rpcClient.InvokeAsync(x => x.WriteOutputAsync(new WriteOutputRequest(_logId, offset, initialLineIndex, data, false)), new RpcContext(), CancellationToken.None);
 					}
-					catch (Exception Ex)
+					catch (Exception ex)
 					{
-						Inner.LogWarning(Ex, "Unable to write data to server (log {LogId}, offset {Offset})", LogId, Offset);
+						Inner.LogWarning(ex, "Unable to write data to server (log {LogId}, offset {Offset})", _logId, offset);
 					}
-					Offset += Data.Length;
+					offset += data.Length;
 				}
 
 				// Write all the events
-				if (Events.Count > 0)
+				if (events.Count > 0)
 				{
 					try
 					{
-						await RpcClient.InvokeAsync(x => x.CreateEventsAsync(new CreateEventsRequest(Events)), new RpcContext(), CancellationToken.None);
+						await _rpcClient.InvokeAsync(x => x.CreateEventsAsync(new CreateEventsRequest(events)), new RpcContext(), CancellationToken.None);
 					}
-					catch (Exception Ex)
+					catch (Exception ex)
 					{
-						Inner.LogWarning(Ex, "Unable to create events");
+						Inner.LogWarning(ex, "Unable to create events");
 					}
 				}
 
 				// Update the outcome of this jobstep
-				if (JobId != null && JobBatchId != null && JobStepId != null && Outcome != PostedOutcome)
+				if (_jobId != null && _jobBatchId != null && _jobStepId != null && Outcome != postedOutcome)
 				{
 					try
 					{
-						await RpcClient.InvokeAsync(x => x.UpdateStepAsync(new UpdateStepRequest(JobId, JobBatchId, JobStepId, JobStepState.Unspecified, Outcome)), new RpcContext(), CancellationToken.None);
+						await _rpcClient.InvokeAsync(x => x.UpdateStepAsync(new UpdateStepRequest(_jobId, _jobBatchId, _jobStepId, JobStepState.Unspecified, Outcome)), new RpcContext(), CancellationToken.None);
 					}
-					catch (Exception Ex)
+					catch (Exception ex)
 					{
-						Inner.LogWarning(Ex, "Unable to update step outcome to {NewOutcome}", Outcome);
+						Inner.LogWarning(ex, "Unable to update step outcome to {NewOutcome}", Outcome);
 					}
-					PostedOutcome = Outcome;
+					postedOutcome = Outcome;
 				}
 
 				// Wait for more data to be available
-				if (!await DataChannel.Reader.WaitToReadAsync())
+				if (!await _dataChannel.Reader.WaitToReadAsync())
 				{
 					try
 					{
-						await RpcClient.InvokeAsync(x => x.WriteOutputAsync(new WriteOutputRequest(LogId, Offset, LineIndex, Array.Empty<byte>(), true)), new RpcContext(), CancellationToken.None);
+						await _rpcClient.InvokeAsync(x => x.WriteOutputAsync(new WriteOutputRequest(_logId, offset, lineIndex, Array.Empty<byte>(), true)), new RpcContext(), CancellationToken.None);
 					}
-					catch (Exception Ex)
+					catch (Exception ex)
 					{
-						Inner.LogWarning(Ex, "Unable to flush data to server (log {LogId}, offset {Offset})", LogId, Offset);
+						Inner.LogWarning(ex, "Unable to flush data to server (log {LogId}, offset {Offset})", _logId, offset);
 					}
 					break;
 				}
diff --git a/Engine/Source/Programs/Horde/Horde.Agent/Utility/NetworkShare.cs b/Engine/Source/Programs/Horde/Horde.Agent/Utility/NetworkShare.cs
--- a/Engine/Source/Programs/Horde/Horde.Agent/Utility/NetworkShare.cs
+++ b/Engine/Source/Programs/Horde/Horde.Agent/Utility/NetworkShare.cs
@@ -15,6 +15,7 @@ namespace Horde.Agent.Utility
 	static class NetworkShare
 	{
 		[StructLayout(LayoutKind.Sequential, CharSet = CharSet.Unicode)]
+		[SuppressMessage("Style", "IDE1006:Naming Styles", Justification = "<Pending>")]
 		private struct NETRESOURCE
 		{
 			public uint dwScope;
@@ -30,12 +31,11 @@ namespace Horde.Agent.Utility
 		const uint RESOURCETYPE_DISK = 1;
 		const uint CONNECT_TEMPORARY = 4;
 
-
 		const int ERROR_ALREADY_ASSIGNED = 85;
 		const int ERROR_MORE_DATA = 234;
 
 		[DllImport("Mpr.dll", CharSet = CharSet.Unicode)]
-		private static extern int WNetAddConnection2W(ref NETRESOURCE lpNetResource, string? lpPassword, string? lpUsername, System.UInt32 dwFlags );
+		private static extern int WNetAddConnection2W(ref NETRESOURCE lpNetResource, string? lpPassword, string? lpUsername, uint dwFlags);
 
 		[DllImport("Mpr.dll", CharSet = CharSet.Unicode)]
 		private static extern int WNetCancelConnectionW(string lpName, [MarshalAs(UnmanagedType.Bool)] bool fForce);
@@ -46,84 +46,84 @@ namespace Horde.Agent.Utility
 		/// <summary>
 		/// Mounts a network share at the given location
 		/// </summary>
-		/// <param name="MountPoint">Mount point for the share</param>
-		/// <param name="RemotePath">Path to the remote resource</param>
-		public static void Mount(string MountPoint, string RemotePath)
+		/// <param name="mountPoint">Mount point for the share</param>
+		/// <param name="remotePath">Path to the remote resource</param>
+		public static void Mount(string mountPoint, string remotePath)
 		{
-			NETRESOURCE NetResource = new NETRESOURCE();
-			NetResource.dwType = RESOURCETYPE_DISK;
-			NetResource.lpLocalName = MountPoint;
-			NetResource.lpRemoteName = RemotePath;
+			NETRESOURCE netResource = new NETRESOURCE();
+			netResource.dwType = RESOURCETYPE_DISK;
+			netResource.lpLocalName = mountPoint;
+			netResource.lpRemoteName = remotePath;
 
-			int Result = WNetAddConnection2W(ref NetResource, null, null, CONNECT_TEMPORARY);
-			if(Result != 0)
+			int result = WNetAddConnection2W(ref netResource, null, null, CONNECT_TEMPORARY);
+			if(result != 0)
 			{
-				if (Result == ERROR_ALREADY_ASSIGNED)
+				if (result == ERROR_ALREADY_ASSIGNED)
 				{
-					string? CurRemotePath;
-					if (TryGetRemotePath(MountPoint, out CurRemotePath))
+					string? curRemotePath;
+					if (TryGetRemotePath(mountPoint, out curRemotePath))
 					{
-						if (CurRemotePath.Equals(RemotePath, StringComparison.OrdinalIgnoreCase))
+						if (curRemotePath.Equals(remotePath, StringComparison.OrdinalIgnoreCase))
 						{
 							return;
 						}
 						else
 						{
-							throw new Win32Exception(Result, $"Unable to mount network share {RemotePath} as {MountPoint} ({Result}). Currently connected to {CurRemotePath}.");
+							throw new Win32Exception(result, $"Unable to mount network share {remotePath} as {mountPoint} ({result}). Currently connected to {curRemotePath}.");
 						}
 					}
 				}
-				throw new Win32Exception(Result, $"Unable to mount network share {RemotePath} as {MountPoint} ({Result}: {new Win32Exception(Result).Message})");
+				throw new Win32Exception(result, $"Unable to mount network share {remotePath} as {mountPoint} ({result}: {new Win32Exception(result).Message})");
 			}
 		}
 
 		/// <summary>
 		/// Unmounts a network share
 		/// </summary>
-		/// <param name="MountPoint">The mount point to remove</param>
-		public static void Unmount(string MountPoint)
+		/// <param name="mountPoint">The mount point to remove</param>
+		public static void Unmount(string mountPoint)
 		{
-			int Result = WNetCancelConnectionW(MountPoint, true);
-			if(Result != 0)
+			int result = WNetCancelConnectionW(mountPoint, true);
+			if(result != 0)
 			{
-				throw new Win32Exception(Result, $"Unable to unmount {MountPoint} ({Result})");
+				throw new Win32Exception(result, $"Unable to unmount {mountPoint} ({result})");
 			}
 		}
 
 		/// <summary>
 		/// Gets the currently mounted path for a network share
 		/// </summary>
-		/// <param name="MountPoint">The mount point</param>
-		/// <param name="OutRemotePath">Receives the remote path</param>
+		/// <param name="mountPoint">The mount point</param>
+		/// <param name="outRemotePath">Receives the remote path</param>
 		/// <returns>True if the remote path is returned</returns>
-		public static bool TryGetRemotePath(string MountPoint, [NotNullWhen(true)] out string? OutRemotePath)
+		public static bool TryGetRemotePath(string mountPoint, [NotNullWhen(true)] out string? outRemotePath)
 		{
-			int Length = 260 * sizeof(char);
+			int length = 260 * sizeof(char);
 			for (; ; )
 			{
-				IntPtr RemotePath = Marshal.AllocHGlobal(Length);
+				IntPtr remotePath = Marshal.AllocHGlobal(length);
 				try
 				{
-					int PrevLength = Length;
-					int Result = WNetGetConnectionW(MountPoint, RemotePath, ref Length);
-					if (Result == 0)
+					int prevLength = length;
+					int result = WNetGetConnectionW(mountPoint, remotePath, ref length);
+					if (result == 0)
 					{
-						OutRemotePath = Marshal.PtrToStringUni(RemotePath)!;
-						return OutRemotePath != null;
+						outRemotePath = Marshal.PtrToStringUni(remotePath)!;
+						return outRemotePath != null;
 					}
-					else if(Result == ERROR_MORE_DATA && Length > PrevLength)
+					else if(result == ERROR_MORE_DATA && length > prevLength)
 					{
 						continue;
 					}
 					else
 					{
-						OutRemotePath = null!;
+						outRemotePath = null!;
 						return false;
 					}
 				}
 				finally
 				{
-					Marshal.FreeHGlobal(RemotePath);
+					Marshal.FreeHGlobal(remotePath);
 				}
 			}
 		}
diff --git a/Engine/Source/Programs/Horde/Horde.Agent/Utility/RpcConnection.cs b/Engine/Source/Programs/Horde/Horde.Agent/Utility/RpcConnection.cs
--- a/Engine/Source/Programs/Horde/Horde.Agent/Utility/RpcConnection.cs
+++ b/Engine/Source/Programs/Horde/Horde.Agent/Utility/RpcConnection.cs
@@ -59,20 +59,20 @@ namespace Horde.Agent.Utility
 		/// <summary>
 		/// Constructor
 		/// </summary>
-		/// <param name="Text">Text to display</param>
-		public RpcContext(string Text)
+		/// <param name="text">Text to display</param>
+		public RpcContext(string text)
 		{
-			this.Text = Text;
+			Text = text;
 		}
 
 		/// <summary>
 		/// Constructor
 		/// </summary>
-		/// <param name="Line">Line number of the file</param>
-		/// <param name="File">Calling file</param>
-		public RpcContext([CallerLineNumber] int Line = 0, [CallerFilePath] string File = "")
+		/// <param name="line">Line number of the file</param>
+		/// <param name="file">Calling file</param>
+		public RpcContext([CallerLineNumber] int line = 0, [CallerFilePath] string file = "")
 		{
-			this.Text = $"{File}({Line})";
+			Text = $"{file}({line})";
 		}
 	}
 
@@ -82,53 +82,53 @@ namespace Horde.Agent.Utility
 		/// Attempts to get a client reference, returning immediately if there's not one available
 		/// </summary>
 		/// <returns></returns>
-		IRpcClientRef? TryGetClientRef(RpcContext Context);
+		IRpcClientRef? TryGetClientRef(RpcContext context);
 
 		/// <summary>
 		/// Obtains a new client reference object
 		/// </summary>
-		/// <param name="Context">Context for the call, for debugging</param>
-		/// <param name="CancellationToken">Cancellation token for this request</param>
+		/// <param name="context">Context for the call, for debugging</param>
+		/// <param name="cancellationToken">Cancellation token for this request</param>
 		/// <returns>New client reference</returns>
-		Task<IRpcClientRef> GetClientRef(RpcContext Context, CancellationToken CancellationToken);
+		Task<IRpcClientRef> GetClientRef(RpcContext context, CancellationToken cancellationToken);
 
 		/// <summary>
 		/// Invokes an asynchronous command with a HordeRpcClient instance
 		/// </summary>
-		/// <param name="Func">The function to execute</param>
-		/// <param name="Context">Context for the call, for debugging</param>
-		/// <param name="CancellationToken">Cancellation token for the call</param>
+		/// <param name="func">The function to execute</param>
+		/// <param name="context">Context for the call, for debugging</param>
+		/// <param name="cancellationToken">Cancellation token for the call</param>
 		/// <returns>Response from the call</returns>
-		Task<T> InvokeOnceAsync<T>(Func<HordeRpc.HordeRpcClient, Task<T>> Func, RpcContext Context, CancellationToken CancellationToken);
+		Task<T> InvokeOnceAsync<T>(Func<HordeRpc.HordeRpcClient, Task<T>> func, RpcContext context, CancellationToken cancellationToken);
 
 		/// <summary>
 		/// Obtains a client reference and executes a unary RPC
 		/// </summary>
 		/// <typeparam name="T">Return type from the call</typeparam>
-		/// <param name="Func">Method to execute with the RPC instance</param>
-		/// <param name="Context">Context for the call, for debugging</param>
-		/// <param name="CancellationToken">Cancellation token for the call</param>
+		/// <param name="func">Method to execute with the RPC instance</param>
+		/// <param name="context">Context for the call, for debugging</param>
+		/// <param name="cancellationToken">Cancellation token for the call</param>
 		/// <returns>Response from the call</returns>
-		Task<T> InvokeOnceAsync<T>(Func<HordeRpc.HordeRpcClient, AsyncUnaryCall<T>> Func, RpcContext Context, CancellationToken CancellationToken);
+		Task<T> InvokeOnceAsync<T>(Func<HordeRpc.HordeRpcClient, AsyncUnaryCall<T>> func, RpcContext context, CancellationToken cancellationToken);
 
 		/// <summary>
 		/// Invokes an asynchronous command with a HordeRpcClient instance
 		/// </summary>
-		/// <param name="Func">The function to execute</param>
-		/// <param name="Context">Context for the call, for debugging</param>
-		/// <param name="CancellationToken">Cancellation token for the call</param>
+		/// <param name="func">The function to execute</param>
+		/// <param name="context">Context for the call, for debugging</param>
+		/// <param name="cancellationToken">Cancellation token for the call</param>
 		/// <returns>Response from the call</returns>
-		Task<T> InvokeAsync<T>(Func<HordeRpc.HordeRpcClient, Task<T>> Func, RpcContext Context, CancellationToken CancellationToken);
+		Task<T> InvokeAsync<T>(Func<HordeRpc.HordeRpcClient, Task<T>> func, RpcContext context, CancellationToken cancellationToken);
 
 		/// <summary>
 		/// Obtains a client reference and executes a unary RPC
 		/// </summary>
 		/// <typeparam name="T">Return type from the call</typeparam>
-		/// <param name="Func">Method to execute with the RPC instance</param>
-		/// <param name="Context">Context for the call, for debugging</param>
-		/// <param name="CancellationToken">Cancellation token for the call</param>
+		/// <param name="func">Method to execute with the RPC instance</param>
+		/// <param name="context">Context for the call, for debugging</param>
+		/// <param name="cancellationToken">Cancellation token for the call</param>
 		/// <returns>Response from the call</returns>
-		Task<T> InvokeAsync<T>(Func<HordeRpc.HordeRpcClient, AsyncUnaryCall<T>> Func, RpcContext Context, CancellationToken CancellationToken);
+		Task<T> InvokeAsync<T>(Func<HordeRpc.HordeRpcClient, AsyncUnaryCall<T>> func, RpcContext context, CancellationToken cancellationToken);
 	}
 
 	/// <summary>
@@ -170,49 +170,49 @@ namespace Horde.Agent.Utility
 			/// <summary>
 			/// Task which is set to indicate the connection is being disposed
 			/// </summary>
-			public Task DisposingTask => DisposingTaskSource.Task;
+			public Task DisposingTask => _disposingTaskSource.Task;
 
 			/// <summary>
 			/// Logger for messages about refcount changes
 			/// </summary>
-			ILogger Logger;
+			private readonly ILogger _logger;
 
 			/// <summary>
 			/// The current refcount
 			/// </summary>
-			int RefCount;
+			private int _refCount;
 
 			/// <summary>
 			/// Whether the connection is terminating
 			/// </summary>
-			bool bDisposing;
+			private bool _disposing;
 
 			/// <summary>
 			/// Cancellation token set 
 			/// </summary>
-			TaskCompletionSource<bool> DisposingTaskSource = new TaskCompletionSource<bool>();
+			private readonly TaskCompletionSource<bool> _disposingTaskSource = new TaskCompletionSource<bool>();
 
 			/// <summary>
 			/// Wraps a task allowing the disposer to wait for clients to finish using this connection
 			/// </summary>
-			TaskCompletionSource<bool> DisposedTaskSource = new TaskCompletionSource<bool>();
+			private readonly TaskCompletionSource<bool> _disposedTaskSource = new TaskCompletionSource<bool>();
 
 			/// <summary>
 			/// Constructor
 			/// </summary>
-			/// <param name="ConnectionId">The connection id</param>
-			/// <param name="Name">Name of the server</param>
-			/// <param name="Channel">The channel instance</param>
-			/// <param name="Client">The client instance</param>
-			/// <param name="Logger">Logger for debug messages</param>
-			public RpcSubConnection(int ConnectionId, string Name, GrpcChannel Channel, HordeRpc.HordeRpcClient Client, ILogger Logger)
+			/// <param name="connectionId">The connection id</param>
+			/// <param name="name">Name of the server</param>
+			/// <param name="channel">The channel instance</param>
+			/// <param name="client">The client instance</param>
+			/// <param name="logger">Logger for debug messages</param>
+			public RpcSubConnection(int connectionId, string name, GrpcChannel channel, HordeRpc.HordeRpcClient client, ILogger logger)
 			{
-				this.ConnectionId = ConnectionId;
-				this.Name = Name;
-				this.Channel = Channel;
-				this.Client = Client;
-				this.Logger = Logger;
-				this.RefCount = 1;
+				ConnectionId = connectionId;
+				Name = name;
+				Channel = channel;
+				Client = client;
+				_logger = logger;
+				_refCount = 1;
 			}
 
 			/// <summary>
@@ -223,12 +223,12 @@ namespace Horde.Agent.Utility
 			{
 				for (; ; )
 				{
-					int InitialRefCount = RefCount;
-					if (InitialRefCount == 0)
+					int initialRefCount = _refCount;
+					if (initialRefCount == 0)
 					{
 						return false;
 					}
-					if (Interlocked.CompareExchange(ref RefCount, InitialRefCount + 1, InitialRefCount) == InitialRefCount)
+					if (Interlocked.CompareExchange(ref _refCount, initialRefCount + 1, initialRefCount) == initialRefCount)
 					{
 						return true;
 					}
@@ -240,14 +240,14 @@ namespace Horde.Agent.Utility
 			/// </summary>
 			public void Release()
 			{
-				int NewRefCount = Interlocked.Decrement(ref RefCount);
-				if (bDisposing)
+				int newRefCount = Interlocked.Decrement(ref _refCount);
+				if (_disposing)
 				{
-					Logger.LogInformation("Connection {ConnectionId}: Decrementing refcount to rpc server {ServerName} ({RefCount} rpc(s) in flight)", ConnectionId, Name, NewRefCount);
+					_logger.LogInformation("Connection {ConnectionId}: Decrementing refcount to rpc server {ServerName} ({RefCount} rpc(s) in flight)", ConnectionId, Name, newRefCount);
 				}
-				if (NewRefCount == 0)
+				if (newRefCount == 0)
 				{
-					DisposedTaskSource.SetResult(true);
+					_disposedTaskSource.SetResult(true);
 				}
 			}
 
@@ -257,11 +257,11 @@ namespace Horde.Agent.Utility
 			/// <returns>Async task</returns>
 			public ValueTask DisposeAsync()
 			{
-				Logger.LogInformation("Connection {ConnectionId}: Disposing sub-connection", ConnectionId);
-				DisposingTaskSource.TrySetResult(true);
-				bDisposing = true;
+				_logger.LogInformation("Connection {ConnectionId}: Disposing sub-connection", ConnectionId);
+				_disposingTaskSource.TrySetResult(true);
+				_disposing = true;
 				Release();
-				return new ValueTask(DisposedTaskSource.Task);
+				return new ValueTask(_disposedTaskSource.Task);
 			}
 		}
 
@@ -273,61 +273,46 @@ namespace Horde.Agent.Utility
 			/// <summary>
 			/// The subconnection containing the client
 			/// </summary>
-			public RpcSubConnection SubConnection
-			{
-				get;
-			}
+			public RpcSubConnection SubConnection { get; }
 
 			/// <summary>
 			/// Context for the reference
 			/// </summary>
-			public RpcContext Context
-			{
-				get;
-			}
+			public RpcContext Context { get; }
 
 			/// <summary>
 			/// The channel instance
 			/// </summary>
-			public GrpcChannel Channel
-			{
-				get { return SubConnection.Channel; }
-			}
+			public GrpcChannel Channel => SubConnection.Channel;
 
 			/// <summary>
 			/// The client instance
 			/// </summary>
-			public HordeRpc.HordeRpcClient Client
-			{
-				get { return SubConnection.Client; }
-			}
+			public HordeRpc.HordeRpcClient Client => SubConnection.Client;
 
 			/// <summary>
 			/// Task which completes when the connection should be disposed
 			/// </summary>
-			public Task DisposingTask
-			{
-				get { return SubConnection.DisposingTask; }
-			}
+			public Task DisposingTask => SubConnection.DisposingTask;
 
 			/// <summary>
 			/// Tracks all the active refs
 			/// </summary>
-			static List<RpcClientRef> CurrentRefs = new List<RpcClientRef>();
+			private static readonly List<RpcClientRef> s_currentRefs = new List<RpcClientRef>();
 
 			/// <summary>
 			/// Constructor
 			/// </summary>
-			/// <param name="SharedClient">The subconnection containing the client</param>
-			/// <param name="Context">Description of the reference, for debugging</param>
-			public RpcClientRef(RpcSubConnection SharedClient, RpcContext Context)
+			/// <param name="sharedClient">The subconnection containing the client</param>
+			/// <param name="context">Description of the reference, for debugging</param>
+			public RpcClientRef(RpcSubConnection sharedClient, RpcContext context)
 			{
-				this.SubConnection = SharedClient;
-				this.Context = Context;
+				SubConnection = sharedClient;
+				Context = context;
 
-				lock (CurrentRefs)
+				lock (s_currentRefs)
 				{
-					CurrentRefs.Add(this);
+					s_currentRefs.Add(this);
 				}
 			}
 
@@ -338,9 +323,9 @@ namespace Horde.Agent.Utility
 			{
 				SubConnection.Release();
 
-				lock(CurrentRefs)
+				lock(s_currentRefs)
 				{
-					CurrentRefs.Remove(this);
+					s_currentRefs.Remove(this);
 				}
 			}
 
@@ -350,74 +335,52 @@ namespace Horde.Agent.Utility
 			/// <returns>List of ref descriptions</returns>
 			public static List<string> GetActiveRefs()
 			{
-				lock(CurrentRefs)
+				lock(s_currentRefs)
 				{
-					return CurrentRefs.Select(x => x.Context.Text).ToList();
+					return s_currentRefs.Select(x => x.Context.Text).ToList();
 				}
 			}
 		}
 
-		/// <summary>
-		/// Delays before retrying RPC calls
-		/// </summary>
-		static TimeSpan[] RetryTimes =
+		private static readonly TimeSpan[] s_retryTimes =
 		{
 			TimeSpan.FromSeconds(1.0),
 			TimeSpan.FromSeconds(10.0),
 			TimeSpan.FromSeconds(30.0),
 		};
 
-		/// <summary>
-		/// Delegate to create a new GrpcChannel
-		/// </summary>
-		Func<GrpcChannel> CreateGrpcChannel;
-
-		/// <summary>
-		/// Source for tokens asking the background thread to terminate
-		/// </summary>
-		TaskCompletionSource<bool> StoppingTaskSource = new TaskCompletionSource<bool>();
-
-		/// <summary>
-		/// Task source for waiting for a new subconnection to be made
-		/// </summary>
-		TaskCompletionSource<RpcSubConnection> SubConnectionTaskSource = new TaskCompletionSource<RpcSubConnection>();
-
-		/// <summary>
-		/// Background task to maintain the connection
-		/// </summary>
-		Task? BackgroundTask;
-
-		/// <summary>
-		/// Logger instance
-		/// </summary>
-		ILogger Logger;
+		private readonly Func<GrpcChannel> _createGrpcChannel;
+		private readonly TaskCompletionSource<bool> _stoppingTaskSource = new TaskCompletionSource<bool>();
+		private TaskCompletionSource<RpcSubConnection> _subConnectionTaskSource = new TaskCompletionSource<RpcSubConnection>();
+		private Task? _backgroundTask;
+		private readonly ILogger _logger;
 
 		/// <summary>
 		/// Constructor
 		/// </summary>
-		/// <param name="CreateGrpcChannel">Factory method for creating new GRPC channels</param>
-		/// <param name="Logger">Logger instance</param>
-		public RpcConnection(Func<GrpcChannel> CreateGrpcChannel, ILogger Logger)
+		/// <param name="createGrpcChannel">Factory method for creating new GRPC channels</param>
+		/// <param name="logger">Logger instance</param>
+		public RpcConnection(Func<GrpcChannel> createGrpcChannel, ILogger logger)
 		{
-			this.CreateGrpcChannel = CreateGrpcChannel;
-			this.Logger = Logger;
+			_createGrpcChannel = createGrpcChannel;
+			_logger = logger;
 
-			BackgroundTask = Task.Run(() => ExecuteAsync());
+			_backgroundTask = Task.Run(() => ExecuteAsync());
 		}
 
 		/// <summary>
 		/// Attempts to get a client reference, returning immediately if there's not one available
 		/// </summary>
 		/// <returns></returns>
-		public IRpcClientRef? TryGetClientRef(RpcContext Context)
+		public IRpcClientRef? TryGetClientRef(RpcContext context)
 		{
-			Task<RpcSubConnection> SubConnectionTask = SubConnectionTaskSource.Task;
-			if (SubConnectionTask.IsCompleted)
+			Task<RpcSubConnection> subConnectionTask = _subConnectionTaskSource.Task;
+			if (subConnectionTask.IsCompleted)
 			{
-				RpcSubConnection DefaultClient = SubConnectionTask.Result;
-				if(DefaultClient.TryAddRef())
+				RpcSubConnection defaultClient = subConnectionTask.Result;
+				if(defaultClient.TryAddRef())
 				{
-					return new RpcClientRef(DefaultClient, Context);
+					return new RpcClientRef(defaultClient, context);
 				}
 			}
 			return null;
@@ -426,26 +389,26 @@ namespace Horde.Agent.Utility
 		/// <summary>
 		/// Obtains a new client reference object
 		/// </summary>
-		/// <param name="Context">Context for the call, for debugging</param>
-		/// <param name="CancellationToken">Cancellation token for this request</param>
+		/// <param name="context">Context for the call, for debugging</param>
+		/// <param name="cancellationToken">Cancellation token for this request</param>
 		/// <returns>New client reference</returns>
-		public async Task<IRpcClientRef> GetClientRef(RpcContext Context, CancellationToken CancellationToken)
+		public async Task<IRpcClientRef> GetClientRef(RpcContext context, CancellationToken cancellationToken)
 		{
 			for (; ; )
 			{
 				// Print a message while we're waiting for a client to be available
-				if (!SubConnectionTaskSource.Task.IsCompleted)
+				if (!_subConnectionTaskSource.Task.IsCompleted)
 				{
-					Stopwatch Timer = Stopwatch.StartNew();
-					while (!SubConnectionTaskSource.Task.IsCompleted)
+					Stopwatch timer = Stopwatch.StartNew();
+					while (!_subConnectionTaskSource.Task.IsCompleted)
 					{
-						Task Delay = Task.Delay(TimeSpan.FromSeconds(30.0), CancellationToken);
-						if (await Task.WhenAny(Delay, SubConnectionTaskSource.Task, StoppingTaskSource.Task) == Delay)
+						Task delay = Task.Delay(TimeSpan.FromSeconds(30.0), cancellationToken);
+						if (await Task.WhenAny(delay, _subConnectionTaskSource.Task, _stoppingTaskSource.Task) == delay)
 						{
-							await Delay; // Allow the cancellation token to throw
-							Logger.LogInformation("Thread stalled for {Time:0.0}s while waiting for IRpcClientRef. Source:\n{Context}\nActive references:\n{References}\nStack trace:\n{StackTrace}", Timer.Elapsed.TotalSeconds, Context.Text, String.Join("\n", RpcClientRef.GetActiveRefs()), Environment.StackTrace);
+							await delay; // Allow the cancellation token to throw
+							_logger.LogInformation("Thread stalled for {Time:0.0}s while waiting for IRpcClientRef. Source:\n{Context}\nActive references:\n{References}\nStack trace:\n{StackTrace}", timer.Elapsed.TotalSeconds, context.Text, String.Join("\n", RpcClientRef.GetActiveRefs()), Environment.StackTrace);
 						}
-						if (StoppingTaskSource.Task.IsCompleted)
+						if (_stoppingTaskSource.Task.IsCompleted)
 						{
 							throw new OperationCanceledException("RpcConnection is being terminated");
 						}
@@ -453,10 +416,10 @@ namespace Horde.Agent.Utility
 				}
 
 				// Get the new task and increase the refcount
-				RpcSubConnection DefaultClient = await SubConnectionTaskSource.Task;
-				if (DefaultClient.TryAddRef())
+				RpcSubConnection defaultClient = await _subConnectionTaskSource.Task;
+				if (defaultClient.TryAddRef())
 				{
-					return new RpcClientRef(DefaultClient, Context);
+					return new RpcClientRef(defaultClient, context);
 				}
 			}
 		}
@@ -464,15 +427,15 @@ namespace Horde.Agent.Utility
 		/// <summary>
 		/// Invokes an asynchronous command with a HordeRpcClient instance
 		/// </summary>
-		/// <param name="Func">The function to execute</param>
-		/// <param name="Context">Context for the call, for debugging</param>
-		/// <param name="CancellationToken">Cancellation token for the call</param>
+		/// <param name="func">The function to execute</param>
+		/// <param name="context">Context for the call, for debugging</param>
+		/// <param name="cancellationToken">Cancellation token for the call</param>
 		/// <returns>Response from the call</returns>
-		public async Task<T> InvokeOnceAsync<T>(Func<HordeRpc.HordeRpcClient, Task<T>> Func, RpcContext Context, CancellationToken CancellationToken)
+		public async Task<T> InvokeOnceAsync<T>(Func<HordeRpc.HordeRpcClient, Task<T>> func, RpcContext context, CancellationToken cancellationToken)
 		{
-			using (IRpcClientRef ClientRef = await GetClientRef(Context, CancellationToken))
+			using (IRpcClientRef clientRef = await GetClientRef(context, cancellationToken))
 			{
-				return await Func(ClientRef.Client);
+				return await func(clientRef.Client);
 			}
 		}
 
@@ -480,84 +443,84 @@ namespace Horde.Agent.Utility
 		/// Obtains a client reference and executes a unary RPC
 		/// </summary>
 		/// <typeparam name="T">Return type from the call</typeparam>
-		/// <param name="Func">Method to execute with the RPC instance</param>
-		/// <param name="Context">Context for the call, for debugging</param>
-		/// <param name="CancellationToken">Cancellation token for the call</param>
+		/// <param name="func">Method to execute with the RPC instance</param>
+		/// <param name="context">Context for the call, for debugging</param>
+		/// <param name="cancellationToken">Cancellation token for the call</param>
 		/// <returns>Response from the call</returns>
-		public Task<T> InvokeOnceAsync<T>(Func<HordeRpc.HordeRpcClient, AsyncUnaryCall<T>> Func, RpcContext Context, CancellationToken CancellationToken)
+		public Task<T> InvokeOnceAsync<T>(Func<HordeRpc.HordeRpcClient, AsyncUnaryCall<T>> func, RpcContext context, CancellationToken cancellationToken)
 		{
-			return InvokeOnceAsync(async (x) => await Func(x), Context, CancellationToken);
+			return InvokeOnceAsync(async (x) => await func(x), context, cancellationToken);
 		}
 
 		/// <summary>
 		/// Invokes an asynchronous command with a HordeRpcClient instance
 		/// </summary>
-		/// <param name="Func">The function to execute</param>
-		/// <param name="Context">Context for the call, for debugging</param>
-		/// <param name="CancellationToken">Cancellation token for the call</param>
+		/// <param name="func">The function to execute</param>
+		/// <param name="context">Context for the call, for debugging</param>
+		/// <param name="cancellationToken">Cancellation token for the call</param>
 		/// <returns>Response from the call</returns>
-		public async Task<T> InvokeAsync<T>(Func<HordeRpc.HordeRpcClient, Task<T>> Func, RpcContext Context, CancellationToken CancellationToken)
+		public async Task<T> InvokeAsync<T>(Func<HordeRpc.HordeRpcClient, Task<T>> func, RpcContext context, CancellationToken cancellationToken)
 		{
-			for (int Attempt = 0; ;Attempt++)
+			for (int attempt = 0; ;attempt++)
 			{
 				// Attempt the RPC call
-				using (IRpcClientRef ClientRef = await GetClientRef(Context, CancellationToken))
+				using (IRpcClientRef clientRef = await GetClientRef(context, cancellationToken))
 				{
 					try
 					{
-						return await Func(ClientRef.Client);
+						return await func(clientRef.Client);
 					}
-					catch (Exception Ex)
+					catch (Exception ex)
 					{
 						// Otherwise check if we can try again
-						if (Attempt >= RetryTimes.Length)
+						if (attempt >= s_retryTimes.Length)
 						{
-							Logger.LogError("Max number of retry attempts reached");
+							_logger.LogError("Max number of retry attempts reached");
 							throw;
 						}
 
 						// Check if the user requested cancellation
-						if (CancellationToken.IsCancellationRequested)
+						if (cancellationToken.IsCancellationRequested)
 						{
-							Logger.LogInformation("Cancellation of gRPC call requested");
+							_logger.LogInformation("Cancellation of gRPC call requested");
 							throw;
 						}
 
 						// Check we can retry it
-						if (!CanRetry(Ex))
+						if (!CanRetry(ex))
 						{
-							Logger.LogError(Ex, "Exception during RPC: {Message}", Ex.Message);
+							_logger.LogError(ex, "Exception during RPC: {Message}", ex.Message);
 							throw;
 						}
 
-						RpcException? RpcEx = Ex as RpcException;
-						if (RpcEx != null && RpcEx.StatusCode == StatusCode.Unavailable)
+						RpcException? rpcEx = ex as RpcException;
+						if (rpcEx != null && rpcEx.StatusCode == StatusCode.Unavailable)
 						{
-							Logger.LogInformation(Ex, "Failure #{Attempt} during gRPC call (service unavailable). Retrying...", Attempt + 1);
+							_logger.LogInformation(ex, "Failure #{Attempt} during gRPC call (service unavailable). Retrying...", attempt + 1);
 						}
 						else
 						{
-							Logger.LogError(Ex, "Failure #{Attempt} during gRPC call.", Attempt + 1);
+							_logger.LogError(ex, "Failure #{Attempt} during gRPC call.", attempt + 1);
 						}
 					}
 				}
 
 				// Wait before retrying
-				await Task.Delay(RetryTimes[Attempt]);
+				await Task.Delay(s_retryTimes[attempt]);
 			}
 		}
 
 		/// <summary>
 		/// Determines if the given exception should be retried
 		/// </summary>
-		/// <param name="Ex"></param>
+		/// <param name="ex"></param>
 		/// <returns></returns>
-		public static bool CanRetry(Exception Ex)
+		public static bool CanRetry(Exception ex)
 		{
-			RpcException? RpcEx = Ex as RpcException;
-			if(RpcEx != null)
+			RpcException? rpcEx = ex as RpcException;
+			if(rpcEx != null)
 			{
-				return RpcEx.StatusCode == StatusCode.Cancelled || RpcEx.StatusCode == StatusCode.Unavailable || RpcEx.StatusCode == StatusCode.DataLoss /* Interrupted streaming call */;
+				return rpcEx.StatusCode == StatusCode.Cancelled || rpcEx.StatusCode == StatusCode.Unavailable || rpcEx.StatusCode == StatusCode.DataLoss /* Interrupted streaming call */;
 			}
 			else
 			{
@@ -569,13 +532,13 @@ namespace Horde.Agent.Utility
 		/// Obtains a client reference and executes a unary RPC
 		/// </summary>
 		/// <typeparam name="T">Return type from the call</typeparam>
-		/// <param name="Func">Method to execute with the RPC instance</param>
-		/// <param name="Context">Context for the call, for debugging</param>
-		/// <param name="CancellationToken">Cancellation token for the call</param>
+		/// <param name="func">Method to execute with the RPC instance</param>
+		/// <param name="context">Context for the call, for debugging</param>
+		/// <param name="cancellationToken">Cancellation token for the call</param>
 		/// <returns>Response from the call</returns>
-		public Task<T> InvokeAsync<T>(Func<HordeRpc.HordeRpcClient, AsyncUnaryCall<T>> Func, RpcContext Context, CancellationToken CancellationToken)
+		public Task<T> InvokeAsync<T>(Func<HordeRpc.HordeRpcClient, AsyncUnaryCall<T>> func, RpcContext context, CancellationToken cancellationToken)
 		{
-			return InvokeAsync(async (x) => await Func(x), Context, CancellationToken);
+			return InvokeAsync(async (x) => await func(x), context, cancellationToken);
 		}
 
 		/// <summary>
@@ -585,48 +548,48 @@ namespace Horde.Agent.Utility
 		async Task ExecuteAsync()
 		{
 			// Counter for the connection number. Used to track log messages through async threads.
-			int ConnectionId = 0;
+			int connectionId = 0;
 
 			// This task source is marked as completed once a RPC connection receives a notification from the server that it's shutting down.
 			// Another connection will immediately be allocated for other connection requests to use, and the load balancer should route us to
 			// the new server.
-			TaskCompletionSource<bool>? ReconnectTaskSource = null;
+			TaskCompletionSource<bool>? reconnectTaskSource = null;
 
 			// List of connection handling async tasks we're tracking
-			List<Task> Tasks = new List<Task>();
-			while (Tasks.Count > 0 || !StoppingTaskSource.Task.IsCompleted)
+			List<Task> tasks = new List<Task>();
+			while (tasks.Count > 0 || !_stoppingTaskSource.Task.IsCompleted)
 			{
 				// Get the task indicating we want a new connection (or create a new one)
-				Task? ReconnectTask = null;
-				if (!StoppingTaskSource.Task.IsCompleted)
+				Task? reconnectTask = null;
+				if (!_stoppingTaskSource.Task.IsCompleted)
 				{
-					if (ReconnectTaskSource == null || ReconnectTaskSource.Task.IsCompleted)
+					if (reconnectTaskSource == null || reconnectTaskSource.Task.IsCompleted)
 					{
 						// Need to avoid lambda capture of the reconnect task source
-						int NewConnectionId = ++ConnectionId;
-						TaskCompletionSource<bool> NewReconnectTaskSource = new TaskCompletionSource<bool>();
-						Tasks.Add(Task.Run(() => HandleConnectionAsync(NewConnectionId, NewReconnectTaskSource)));
-						ReconnectTaskSource = NewReconnectTaskSource;
+						int newConnectionId = ++connectionId;
+						TaskCompletionSource<bool> newReconnectTaskSource = new TaskCompletionSource<bool>();
+						tasks.Add(Task.Run(() => HandleConnectionAsync(newConnectionId, newReconnectTaskSource)));
+						reconnectTaskSource = newReconnectTaskSource;
 					}
-					ReconnectTask = ReconnectTaskSource.Task;
+					reconnectTask = reconnectTaskSource.Task;
 				}
 
 				// Build a list of tasks to wait for
-				List<Task> WaitTasks = new List<Task>(Tasks);
-				if (ReconnectTask != null)
+				List<Task> waitTasks = new List<Task>(tasks);
+				if (reconnectTask != null)
 				{
-					WaitTasks.Add(ReconnectTask);
+					waitTasks.Add(reconnectTask);
 				}
-				await Task.WhenAny(WaitTasks);
+				await Task.WhenAny(waitTasks);
 
 				// Remove any complete tasks, logging any exceptions
-				for (int Idx = 0; Idx < Tasks.Count; Idx++)
+				for (int idx = 0; idx < tasks.Count; idx++)
 				{
-					Task Task = Tasks[Idx];
-					if(Task.IsCompleted)
+					Task task = tasks[idx];
+					if(task.IsCompleted)
 					{
-						await Task;
-						Tasks.RemoveAt(Idx--);
+						await task;
+						tasks.RemoveAt(idx--);
 					}
 				}
 			}
@@ -635,29 +598,29 @@ namespace Horde.Agent.Utility
 		/// <summary>
 		/// Background task to monitor the server state
 		/// </summary>
-		/// <param name="ConnectionId">The connection id</param>
-		/// <param name="ReconnectTaskSource">Task source for reconnecting</param>
+		/// <param name="connectionId">The connection id</param>
+		/// <param name="reconnectTaskSource">Task source for reconnecting</param>
 		/// <returns>Async task</returns>
-		async Task HandleConnectionAsync(int ConnectionId, TaskCompletionSource<bool> ReconnectTaskSource)
+		async Task HandleConnectionAsync(int connectionId, TaskCompletionSource<bool> reconnectTaskSource)
 		{
-			using IScope Scope = GlobalTracer.Instance.BuildSpan("RpcConnection {ConnectionId}").StartActive();
+			using IScope scope = GlobalTracer.Instance.BuildSpan("RpcConnection {ConnectionId}").StartActive();
 
-			Logger.LogInformation("Connection {ConnectionId}: Creating connection", ConnectionId);
+			_logger.LogInformation("Connection {ConnectionId}: Creating connection", connectionId);
 			try
 			{
-				await HandleConnectionInternalAsync(ConnectionId, ReconnectTaskSource);
+				await HandleConnectionInternalAsync(connectionId, reconnectTaskSource);
 			}
-			catch (RpcException Ex) when (Ex.StatusCode == StatusCode.Unavailable)
+			catch (RpcException ex) when (ex.StatusCode == StatusCode.Unavailable)
 			{
-				Logger.LogInformation("Connection {ConnectionId}: Rpc service is unavailable. Reconnecting.", ConnectionId);
+				_logger.LogInformation("Connection {ConnectionId}: Rpc service is unavailable. Reconnecting.", connectionId);
 			}
-			catch (Exception Ex)
+			catch (Exception ex)
 			{
-				Logger.LogError(Ex, "Connection {ConnectionId}: Exception handling connection: {Message}", ConnectionId, Ex.Message);
+				_logger.LogError(ex, "Connection {ConnectionId}: Exception handling connection: {Message}", connectionId, ex.Message);
 			}
 			finally
 			{
-				TriggerReconnect(ReconnectTaskSource);
+				TriggerReconnect(reconnectTaskSource);
 			}
 		}
 
@@ -665,81 +628,81 @@ namespace Horde.Agent.Utility
 		/// Background task to monitor the server state
 		/// </summary>
 		/// <returns>Async task</returns>
-		async Task HandleConnectionInternalAsync(int ConnectionId, TaskCompletionSource<bool> ReconnectTaskSource)
+		async Task HandleConnectionInternalAsync(int connectionId, TaskCompletionSource<bool> reconnectTaskSource)
 		{
-			using (GrpcChannel Channel = CreateGrpcChannel())
+			using (GrpcChannel channel = _createGrpcChannel())
 			{
-				HordeRpc.HordeRpcClient Client = new HordeRpc.HordeRpcClient(Channel);//.Intercept(new ClientTracingInterceptor(GlobalTracer.Instance)));
+				HordeRpc.HordeRpcClient client = new HordeRpc.HordeRpcClient(channel);//.Intercept(new ClientTracingInterceptor(GlobalTracer.Instance)));
 
-				RpcSubConnection? SubConnection = null;
+				RpcSubConnection? subConnection = null;
 				try
 				{
-					while (!ReconnectTaskSource.Task.IsCompleted && !StoppingTaskSource.Task.IsCompleted)
+					while (!reconnectTaskSource.Task.IsCompleted && !_stoppingTaskSource.Task.IsCompleted)
 					{
 						// Client HAS TO BE the one to end the call, by closing the stream. The server must keep receiving requests from us until it does.
-						using (AsyncDuplexStreamingCall<QueryServerStateRequest, QueryServerStateResponse> Call = Client.QueryServerStateV2())
+						using (AsyncDuplexStreamingCall<QueryServerStateRequest, QueryServerStateResponse> call = client.QueryServerStateV2())
 						{
 							// Send the name of this agent
-							QueryServerStateRequest Request = new QueryServerStateRequest();
-							Request.Name = Dns.GetHostName();
-							await Call.RequestStream.WriteAsync(Request);
+							QueryServerStateRequest request = new QueryServerStateRequest();
+							request.Name = Dns.GetHostName();
+							await call.RequestStream.WriteAsync(request);
 
 							// Read the server info
-							Task<bool> MoveNextTask = Call.ResponseStream.MoveNext();
-							if (!await MoveNextTask)
+							Task<bool> moveNextTask = call.ResponseStream.MoveNext();
+							if (!await moveNextTask)
 							{
-								Logger.LogError("Connection {ConnectionId}: No response from server", ConnectionId);
+								_logger.LogError("Connection {ConnectionId}: No response from server", connectionId);
 								break;
 							}
 
 							// Read the server response
-							QueryServerStateResponse Response = Call.ResponseStream.Current;
-							if (!Response.Stopping)
+							QueryServerStateResponse response = call.ResponseStream.Current;
+							if (!response.Stopping)
 							{
 								// The first time we connect, log the server name and create the subconnection
-								if (SubConnection == null)
+								if (subConnection == null)
 								{
-									Logger.LogInformation("Connection {ConnectionId}: Connected to rpc server {ServerName}", ConnectionId, Response.Name);
-									SubConnection = new RpcSubConnection(ConnectionId, Response.Name, Channel, Client, Logger);
-									SubConnectionTaskSource.SetResult(SubConnection);
+									_logger.LogInformation("Connection {ConnectionId}: Connected to rpc server {ServerName}", connectionId, response.Name);
+									subConnection = new RpcSubConnection(connectionId, response.Name, channel, client, _logger);
+									_subConnectionTaskSource.SetResult(subConnection);
 								}
 
 								// Wait for the StoppingTask token to be set, or the server to inform that it's shutting down
-								MoveNextTask = Call.ResponseStream.MoveNext();
-								await Task.WhenAny(StoppingTaskSource.Task, Task.Delay(TimeSpan.FromSeconds(45.0)), MoveNextTask);
+								moveNextTask = call.ResponseStream.MoveNext();
+								await Task.WhenAny(_stoppingTaskSource.Task, Task.Delay(TimeSpan.FromSeconds(45.0)), moveNextTask);
 
 								// Update the response
-								if (MoveNextTask.IsCompleted && MoveNextTask.Result)
+								if (moveNextTask.IsCompleted && moveNextTask.Result)
 								{
-									Response = Call.ResponseStream.Current;
+									response = call.ResponseStream.Current;
 								}
 							}
 
 							// If the server is stopping, start reconnecting
-							if (Response.Stopping)
+							if (response.Stopping)
 							{
-								Logger.LogInformation("Connection {ConnectionId}: Server is stopping. Triggering reconnect.", ConnectionId);
-								TriggerReconnect(ReconnectTaskSource);
+								_logger.LogInformation("Connection {ConnectionId}: Server is stopping. Triggering reconnect.", connectionId);
+								TriggerReconnect(reconnectTaskSource);
 							}
 
 							// Close the request stream
-							await Call.RequestStream.CompleteAsync();
+							await call.RequestStream.CompleteAsync();
 
 							// Wait for the server to finish posting responses
-							while (await MoveNextTask)
+							while (await moveNextTask)
 							{
-								Response = Call.ResponseStream.Current;
-								MoveNextTask = Call.ResponseStream.MoveNext();
+								response = call.ResponseStream.Current;
+								moveNextTask = call.ResponseStream.MoveNext();
 							}
 						}
 					}
-					Logger.LogInformation("Connection {ConnectionId}: Closing connection to rpc server", ConnectionId);
+					_logger.LogInformation("Connection {ConnectionId}: Closing connection to rpc server", connectionId);
 				}
 				finally
 				{
-					if (SubConnection != null)
+					if (subConnection != null)
 					{
-						await SubConnection.DisposeAsync();
+						await subConnection.DisposeAsync();
 					}
 				}
 
@@ -751,19 +714,19 @@ namespace Horde.Agent.Utility
 		/// <summary>
 		/// Clears the current subconnection and triggers a reconnect
 		/// </summary>
-		/// <param name="ReconnectTaskSource"></param>
-		void TriggerReconnect(TaskCompletionSource<bool> ReconnectTaskSource)
+		/// <param name="reconnectTaskSource"></param>
+		void TriggerReconnect(TaskCompletionSource<bool> reconnectTaskSource)
 		{
-			if (!ReconnectTaskSource.Task.IsCompleted)
+			if (!reconnectTaskSource.Task.IsCompleted)
 			{
 				// First reset the task source, so that nothing new starts to use the current connection
-				if (SubConnectionTaskSource.Task.IsCompleted)
+				if (_subConnectionTaskSource.Task.IsCompleted)
 				{
-					SubConnectionTaskSource = new TaskCompletionSource<RpcSubConnection>();
+					_subConnectionTaskSource = new TaskCompletionSource<RpcSubConnection>();
 				}
 
 				// Now trigger another connection
-				ReconnectTaskSource.SetResult(true);
+				reconnectTaskSource.SetResult(true);
 			}
 		}
 
@@ -773,18 +736,18 @@ namespace Horde.Agent.Utility
 		/// <returns>Async task</returns>
 		public async ValueTask DisposeAsync()
 		{
-			if (BackgroundTask != null)
+			if (_backgroundTask != null)
 			{
-				StoppingTaskSource.TrySetResult(true);
+				_stoppingTaskSource.TrySetResult(true);
 
-				await BackgroundTask;
-				BackgroundTask = null!;
+				await _backgroundTask;
+				_backgroundTask = null!;
 			}
 		}
 
-		public static IRpcConnection Create(Func<GrpcChannel> CreateGrpcChannel, ILogger Logger)
+		public static IRpcConnection Create(Func<GrpcChannel> createGrpcChannel, ILogger logger)
 		{
-			return new RpcConnection(CreateGrpcChannel, Logger);
+			return new RpcConnection(createGrpcChannel, logger);
 		}
 	}
 }
diff --git a/Engine/Source/Programs/Horde/Horde.Agent/Utility/Shutdown.cs b/Engine/Source/Programs/Horde/Horde.Agent/Utility/Shutdown.cs
--- a/Engine/Source/Programs/Horde/Horde.Agent/Utility/Shutdown.cs
+++ b/Engine/Source/Programs/Horde/Horde.Agent/Utility/Shutdown.cs
@@ -5,6 +5,7 @@ using Microsoft.Win32.SafeHandles;
 using System;
 using System.Collections.Generic;
 using System.Diagnostics;
+using System.Diagnostics.CodeAnalysis;
 using System.Net.NetworkInformation;
 using System.Runtime.InteropServices;
 using System.Text;
@@ -13,30 +14,33 @@ namespace Horde.Agent.Utility
 {
 	static class Shutdown
 	{
-		public const UInt32 TOKEN_QUERY = 0x0008;
-		public const UInt32 TOKEN_ADJUST_PRIVILEGES = 0x0020;
+		public const uint TOKEN_QUERY = 0x0008;
+		public const uint TOKEN_ADJUST_PRIVILEGES = 0x0020;
 
 		[DllImport("advapi32.dll", SetLastError = true)]
 		[return: MarshalAs(UnmanagedType.Bool)]
-		static extern bool OpenProcessToken(IntPtr ProcessHandle, UInt32 DesiredAccess, out IntPtr TokenHandle);
+		static extern bool OpenProcessToken(IntPtr processHandle, uint desiredAccess, out IntPtr tokenHandle);
 
 		[StructLayout(LayoutKind.Sequential)]
+		[SuppressMessage("Style", "IDE1006:Naming Styles", Justification = "Native struct")]
 		public struct LUID
 		{
-			public UInt32 LowPart;
-			public Int32 HighPart;
+			public uint LowPart;
+			public int HighPart;
 		}
 
 		[DllImport("advapi32.dll")]
 		static extern bool LookupPrivilegeValue(string? lpSystemName, string lpName, ref LUID lpLuid);
 
 		[StructLayout(LayoutKind.Sequential, Pack = 4)]
+		[SuppressMessage("Style", "IDE1006:Naming Styles", Justification = "Native struct")]
 		public struct LUID_AND_ATTRIBUTES
 		{
 			public LUID Luid;
-			public UInt32 Attributes;
+			public uint Attributes;
 		}
 
+		[SuppressMessage("Style", "IDE1006:Naming Styles", Justification = "Native struct")]
 		struct TOKEN_PRIVILEGES
 		{
 			public int PrivilegeCount;
@@ -46,11 +50,12 @@ namespace Horde.Agent.Utility
 
 		// Use this signature if you do not want the previous state
 		[DllImport("advapi32.dll", SetLastError = true)]
+		[SuppressMessage("Style", "IDE1006:Naming Styles", Justification = "Native struct")]
 		[return: MarshalAs(UnmanagedType.Bool)]
 		static extern bool AdjustTokenPrivileges(IntPtr TokenHandle,
 		   [MarshalAs(UnmanagedType.Bool)]bool DisableAllPrivileges,
 		   ref TOKEN_PRIVILEGES NewState,
-		   UInt32 Zero,
+		   uint Zero,
 		   IntPtr Null1,
 		   IntPtr Null2);
 
@@ -67,92 +72,92 @@ namespace Horde.Agent.Utility
 		/// <summary>
 		/// Initiate a shutdown operation
 		/// </summary>
-		/// <param name="bRestartAfterShutdown">Whether to restart after the shutdown</param>
-		/// <param name="Logger">Logger for the operation</param>
+		/// <param name="restartAfterShutdown">Whether to restart after the shutdown</param>
+		/// <param name="logger">Logger for the operation</param>
 		/// <returns></returns>
-		public static bool InitiateShutdown(bool bRestartAfterShutdown, ILogger Logger)
+		public static bool InitiateShutdown(bool restartAfterShutdown, ILogger logger)
 		{
 			if (RuntimeInformation.IsOSPlatform(OSPlatform.Windows))
 			{
-				if (bRestartAfterShutdown)
+				if (restartAfterShutdown)
 				{
-					Logger.LogInformation("Triggering restart");
+					logger.LogInformation("Triggering restart");
 				}
 				else
 				{
-					Logger.LogInformation("Triggering shutdown");
+					logger.LogInformation("Triggering shutdown");
 				}
 
-				IntPtr TokenHandle;
-				if (!OpenProcessToken(Process.GetCurrentProcess().Handle, TOKEN_ADJUST_PRIVILEGES | TOKEN_QUERY, out TokenHandle))
+				IntPtr tokenHandle;
+				if (!OpenProcessToken(Process.GetCurrentProcess().Handle, TOKEN_ADJUST_PRIVILEGES | TOKEN_QUERY, out tokenHandle))
 				{
-					Logger.LogError("OpenProcessToken() failed (code 0x{Code:x8})", Marshal.GetLastWin32Error());
+					logger.LogError("OpenProcessToken() failed (code 0x{Code:x8})", Marshal.GetLastWin32Error());
 					return false;
 				}
 
 				// Get the LUID for the shutdown privilege. 
-				LUID Luid = new LUID();
-				if (!LookupPrivilegeValue(null, SE_SHUTDOWN_NAME, ref Luid))
+				LUID luid = new LUID();
+				if (!LookupPrivilegeValue(null, SE_SHUTDOWN_NAME, ref luid))
 				{
-					Logger.LogError("LookupPrivilegeValue() failed (code 0x{Code:x8})", Marshal.GetLastWin32Error());
+					logger.LogError("LookupPrivilegeValue() failed (code 0x{Code:x8})", Marshal.GetLastWin32Error());
 					return false;
 				}
 
-				TOKEN_PRIVILEGES Privileges = new TOKEN_PRIVILEGES();
-				Privileges.PrivilegeCount = 1;
-				Privileges.Privileges = new LUID_AND_ATTRIBUTES[1];
-				Privileges.Privileges[0].Attributes = SE_PRIVILEGE_ENABLED;
-				Privileges.Privileges[0].Luid = Luid;
+				TOKEN_PRIVILEGES privileges = new TOKEN_PRIVILEGES();
+				privileges.PrivilegeCount = 1;
+				privileges.Privileges = new LUID_AND_ATTRIBUTES[1];
+				privileges.Privileges[0].Attributes = SE_PRIVILEGE_ENABLED;
+				privileges.Privileges[0].Luid = luid;
 
-				if (!AdjustTokenPrivileges(TokenHandle, false, ref Privileges, 0, IntPtr.Zero, IntPtr.Zero))
+				if (!AdjustTokenPrivileges(tokenHandle, false, ref privileges, 0, IntPtr.Zero, IntPtr.Zero))
 				{
-					Logger.LogError("AdjustTokenPrivileges() failed (code 0x{Code:x8})", Marshal.GetLastWin32Error());
+					logger.LogError("AdjustTokenPrivileges() failed (code 0x{Code:x8})", Marshal.GetLastWin32Error());
 					return false;
 				}
 
-				if (!InitiateSystemShutdownEx(null, "HordeAgent has initiated shutdown", 10, true, bRestartAfterShutdown, SHTDN_REASON_MAJOR_APPLICATION | SHTDN_REASON_MINOR_MAINTENANCE))
+				if (!InitiateSystemShutdownEx(null, "HordeAgent has initiated shutdown", 10, true, restartAfterShutdown, SHTDN_REASON_MAJOR_APPLICATION | SHTDN_REASON_MINOR_MAINTENANCE))
 				{
-					Logger.LogError("Shutdown failed (0x{Code:x8})", Marshal.GetLastWin32Error());
+					logger.LogError("Shutdown failed (0x{Code:x8})", Marshal.GetLastWin32Error());
 					return false;
 				}
 				return true;
 			}
 			else if (RuntimeInformation.IsOSPlatform(OSPlatform.Linux))
 			{
-				String ShutdownArgs;
-				if (bRestartAfterShutdown)
+				string shutdownArgs;
+				if (restartAfterShutdown)
                 {
-					ShutdownArgs = "sudo shutdown -r +0 \"Horde Agent is restarting\"";
+					shutdownArgs = "sudo shutdown -r +0 \"Horde Agent is restarting\"";
 				}
 				else
                 {
-					ShutdownArgs = "sudo shutdown +0 \"Horde Agent is shutting down\"";
+					shutdownArgs = "sudo shutdown +0 \"Horde Agent is shutting down\"";
 				}
  
-				Process ShutdownProcess = new Process()
+				Process shutdownProcess = new Process()
 				{
 					StartInfo = new ProcessStartInfo
 					{
 						FileName = "/bin/sh",
-						Arguments = String.Format("-c \"{0}\"", ShutdownArgs),
+						Arguments = String.Format("-c \"{0}\"", shutdownArgs),
 						UseShellExecute = false,
 						CreateNoWindow = true
 					}
 				};
  
-				ShutdownProcess.Start();
-				ShutdownProcess.WaitForExit();
-				int ExitCode = ShutdownProcess.ExitCode;
-				if (ExitCode != 0)
+				shutdownProcess.Start();
+				shutdownProcess.WaitForExit();
+				int exitCode = shutdownProcess.ExitCode;
+				if (exitCode != 0)
 				{
-					Logger.LogError("Shutdown failed ({0})", ExitCode);
+					logger.LogError("Shutdown failed ({0})", exitCode);
 					return false;
 				}
 				return true;
 			}
 			else
 			{
-				Logger.LogError("Shutdown is not implemented on this platform");
+				logger.LogError("Shutdown is not implemented on this platform");
 				return false;
 			}
 		}
