Horde: Modify agent code to conform to NET coding style.

diff --git a/Engine/Source/Programs/Horde/Horde.Agent/Utility/WindowsServiceManager.cs b/Engine/Source/Programs/Horde/Horde.Agent/Utility/WindowsServiceManager.cs
--- a/Engine/Source/Programs/Horde/Horde.Agent/Utility/WindowsServiceManager.cs
+++ b/Engine/Source/Programs/Horde/Horde.Agent/Utility/WindowsServiceManager.cs
@@ -6,6 +6,7 @@ using System;
 using System.Collections.Generic;
 using System.ComponentModel;
 using System.Diagnostics;
+using System.Diagnostics.CodeAnalysis;
 using System.Reflection;
 using System.Runtime.InteropServices;
 using System.Text;
@@ -26,10 +27,10 @@ namespace Horde.Agent.Utility
 			{
 			}
 
-			public ServiceHandle(IntPtr Handle)
+			public ServiceHandle(IntPtr handle)
 				: base(true)
 			{
-				SetHandle(Handle);
+				SetHandle(handle);
 			}
 
 			protected override bool ReleaseHandle()
@@ -76,6 +77,7 @@ namespace Horde.Agent.Utility
 
 		public const int SC_STATUS_PROCESS_INFO = 0;
 
+		[System.Diagnostics.CodeAnalysis.SuppressMessage("Style", "IDE1006:Naming Styles", Justification = "<Pending>")]
 		public struct SERVICE_STATUS_PROCESS
 		{
 			public uint dwServiceType;
@@ -91,8 +93,9 @@ namespace Horde.Agent.Utility
 
 		[DllImport("advapi32.dll", SetLastError = true)]
 		[return: MarshalAs(UnmanagedType.Bool)]
-		public static extern bool QueryServiceStatusEx(ServiceHandle hService, int InfoLevel, ref SERVICE_STATUS_PROCESS pBuffer, int cbBufSize, out int pcbBytesNeeded);
+		public static extern bool QueryServiceStatusEx(ServiceHandle hService, int infoLevel, ref SERVICE_STATUS_PROCESS pBuffer, int cbBufSize, out int pcbBytesNeeded);
 
+		[System.Diagnostics.CodeAnalysis.SuppressMessage("Style", "IDE1006:Naming Styles", Justification = "<Pending>")]
 		public struct SERVICE_STATUS
 		{
 			public uint dwServiceType;
@@ -113,10 +116,11 @@ namespace Horde.Agent.Utility
 		public const int SERVICE_CONFIG_DESCRIPTION = 1;
 
 		[StructLayout(LayoutKind.Sequential)]
+		[SuppressMessage("Style", "IDE1006:Naming Styles", Justification = "Native struct")]
 		public struct SERVICE_DESCRIPTION
 		{
 			[MarshalAs(UnmanagedType.LPWStr)]
-			public String lpDescription;
+			public string lpDescription;
 		}
 
 		[DllImport("advapi32.dll", SetLastError = true, CharSet = CharSet.Auto)]
@@ -171,24 +175,21 @@ namespace Horde.Agent.Utility
 		/// <summary>
 		/// Handle to the service
 		/// </summary>
-		Native.ServiceHandle Handle;
+		private readonly Native.ServiceHandle _handle;
 
 		/// <summary>
 		/// Constructor
 		/// </summary>
-		/// <param name="Handle">Handle to the service</param>
-		public WindowsService(Native.ServiceHandle Handle)
+		/// <param name="handle">Handle to the service</param>
+		public WindowsService(Native.ServiceHandle handle)
 		{
-			this.Handle = Handle;
+			_handle = handle;
 		}
 
 		/// <summary>
 		/// Determines if the service is valid
 		/// </summary>
-		public bool IsValid
-		{
-			get { return !Handle.IsInvalid; }
-		}
+		public bool IsValid => !_handle.IsInvalid;
 
 		/// <summary>
 		/// Gets the current service status
@@ -196,28 +197,28 @@ namespace Horde.Agent.Utility
 		/// <returns>Status code</returns>
 		public WindowsServiceStatus GetStatus()
 		{
-			Native.SERVICE_STATUS_PROCESS Status = new Native.SERVICE_STATUS_PROCESS();
-			if (!Native.QueryServiceStatusEx(Handle, Native.SC_STATUS_PROCESS_INFO, ref Status, Marshal.SizeOf(Status), out int BytesNeeded))
+			Native.SERVICE_STATUS_PROCESS status = new Native.SERVICE_STATUS_PROCESS();
+			if (!Native.QueryServiceStatusEx(_handle, Native.SC_STATUS_PROCESS_INFO, ref status, Marshal.SizeOf(status), out int _))
 			{
 				throw new Win32Exception(String.Format("Unable to query process status (0x{0:X8)}", Marshal.GetLastWin32Error()));
 			}
-			return (WindowsServiceStatus)Status.dwCurrentState;
+			return (WindowsServiceStatus)status.dwCurrentState;
 		}
 
 		/// <summary>
 		/// Waits for the status to change
 		/// </summary>
-		/// <param name="TransitionStatus">Expected status while transitioning</param>
-		/// <param name="MaxWaitTime">Maximum time to wait</param>
-		public WindowsServiceStatus WaitForStatusChange(WindowsServiceStatus TransitionStatus, TimeSpan MaxWaitTime)
+		/// <param name="transitionStatus">Expected status while transitioning</param>
+		/// <param name="maxWaitTime">Maximum time to wait</param>
+		public WindowsServiceStatus WaitForStatusChange(WindowsServiceStatus transitionStatus, TimeSpan maxWaitTime)
 		{
-			Stopwatch Timer = Stopwatch.StartNew();
+			Stopwatch timer = Stopwatch.StartNew();
 			for(; ;)
 			{
-				WindowsServiceStatus Status = GetStatus();
-				if (Status != TransitionStatus || Timer.Elapsed > MaxWaitTime)
+				WindowsServiceStatus status = GetStatus();
+				if (status != transitionStatus || timer.Elapsed > maxWaitTime)
 				{
-					return Status;
+					return status;
 				}
 				else
 				{
@@ -229,12 +230,12 @@ namespace Horde.Agent.Utility
 		/// <summary>
 		/// Sets the service description
 		/// </summary>
-		/// <param name="Description">Description for the service</param>
-		public void SetDescription(string Description)
+		/// <param name="description">Description for the service</param>
+		public void SetDescription(string description)
 		{
-			Native.SERVICE_DESCRIPTION DescriptionData = new Native.SERVICE_DESCRIPTION();
-			DescriptionData.lpDescription = Description;
-			Native.ChangeServiceConfig2(Handle, Native.SERVICE_CONFIG_DESCRIPTION, ref DescriptionData);
+			Native.SERVICE_DESCRIPTION descriptionData = new Native.SERVICE_DESCRIPTION();
+			descriptionData.lpDescription = description;
+			Native.ChangeServiceConfig2(_handle, Native.SERVICE_CONFIG_DESCRIPTION, ref descriptionData);
 		}
 
 		/// <summary>
@@ -242,7 +243,7 @@ namespace Horde.Agent.Utility
 		/// </summary>
 		public void Start()
 		{
-			if (!Native.StartService(Handle, 0, null))
+			if (!Native.StartService(_handle, 0, null))
 			{
 				throw new Win32Exception(String.Format("Unable to start service (error 0x{0:X8})", Marshal.GetLastWin32Error()));
 			}
@@ -253,13 +254,13 @@ namespace Horde.Agent.Utility
 		/// </summary>
 		public void Stop()
 		{
-			Native.SERVICE_STATUS Status = new Native.SERVICE_STATUS();
-			if (!Native.ControlService(Handle, Native.SERVICE_CONTROL_STOP, ref Status))
+			Native.SERVICE_STATUS status = new Native.SERVICE_STATUS();
+			if (!Native.ControlService(_handle, Native.SERVICE_CONTROL_STOP, ref status))
 			{
-				int Error = Marshal.GetLastWin32Error();
-				if (Error != Native.ERROR_SERVICE_NOT_ACTIVE)
+				int error = Marshal.GetLastWin32Error();
+				if (error != Native.ERROR_SERVICE_NOT_ACTIVE)
 				{
-					throw new Win32Exception(String.Format("Unable to stop service (error 0x{0:X8})", Error));
+					throw new Win32Exception(String.Format("Unable to stop service (error 0x{0:X8})", error));
 				}
 			}
 		}
@@ -269,7 +270,7 @@ namespace Horde.Agent.Utility
 		/// </summary>
 		public void Delete()
 		{
-			if (!Native.DeleteService(Handle))
+			if (!Native.DeleteService(_handle))
 			{
 				throw new Win32Exception(String.Format("Unable to delete service (error 0x{0:X8})", Marshal.GetLastWin32Error()));
 			}
@@ -280,7 +281,7 @@ namespace Horde.Agent.Utility
 		/// </summary>
 		public void Dispose()
 		{
-			Handle.Close();
+			_handle.Close();
 		}
 	}
 
@@ -292,24 +293,24 @@ namespace Horde.Agent.Utility
 		/// <summary>
 		/// Native handle to the service manager
 		/// </summary>
-		Native.ServiceHandle ServiceManagerHandle;
+		private readonly Native.ServiceHandle _serviceManagerHandle;
 
 		/// <summary>
 		/// Constructor. Opens a handle to the service manager.
 		/// </summary>
 		public WindowsServiceManager()
 		{
-			ServiceManagerHandle = Native.OpenSCManager(null, null, Native.SC_MANAGER_ALL_ACCESS);
-			if (ServiceManagerHandle.IsInvalid)
+			_serviceManagerHandle = Native.OpenSCManager(null, null, Native.SC_MANAGER_ALL_ACCESS);
+			if (_serviceManagerHandle.IsInvalid)
 			{
-				int ErrorCode = Marshal.GetLastWin32Error();
-				if (ErrorCode == Native.ERROR_ACCESS_DENIED)
+				int errorCode = Marshal.GetLastWin32Error();
+				if (errorCode == Native.ERROR_ACCESS_DENIED)
 				{
 					throw new Win32Exception("Unable to open service manager (access denied). Check you're running as administrator.");
 				}
 				else
 				{
-					throw new Win32Exception(String.Format("Unable to open service manager (0x{0:X8)).", ErrorCode));
+					throw new Win32Exception(String.Format("Unable to open service manager (0x{0:X8)).", errorCode));
 				}
 			}
 		}
@@ -319,45 +320,45 @@ namespace Horde.Agent.Utility
 		/// </summary>
 		public void Dispose()
 		{
-			ServiceManagerHandle.Close();
+			_serviceManagerHandle.Close();
 		}
 
 		/// <summary>
 		/// Opens a service with the given name
 		/// </summary>
-		/// <param name="ServiceName">Name of the service</param>
+		/// <param name="serviceName">Name of the service</param>
 		/// <returns>New service wrapper</returns>
-		public WindowsService Open(string ServiceName)
+		public WindowsService Open(string serviceName)
 		{
-			Native.ServiceHandle ServiceHandle = Native.OpenService(ServiceManagerHandle, ServiceName, Native.SERVICE_ALL_ACCESS);
-			if (ServiceHandle.IsInvalid)
+			Native.ServiceHandle serviceHandle = Native.OpenService(_serviceManagerHandle, serviceName, Native.SERVICE_ALL_ACCESS);
+			if (serviceHandle.IsInvalid)
 			{
-				int ErrorCode = Marshal.GetLastWin32Error();
-				if (ErrorCode != Native.ERROR_SERVICE_DOES_NOT_EXIST)
+				int errorCode = Marshal.GetLastWin32Error();
+				if (errorCode != Native.ERROR_SERVICE_DOES_NOT_EXIST)
 				{
 					throw new Win32Exception("Unable to open handle to service");
 				}
 			}
-			return new WindowsService(ServiceHandle);
+			return new WindowsService(serviceHandle);
 		}
 
 		/// <summary>
 		/// Creates a service with the given settings
 		/// </summary>
-		/// <param name="Name">Name of the service</param>
-		/// <param name="DisplayName">Display name</param>
-		/// <param name="CommandLine">Command line to use when starting the service</param>
-		/// <param name="UserName">Username to run this service as</param>
-		/// <param name="Password">Password for the account the service is to run under</param>
+		/// <param name="name">Name of the service</param>
+		/// <param name="displayName">Display name</param>
+		/// <param name="commandLine">Command line to use when starting the service</param>
+		/// <param name="userName">Username to run this service as</param>
+		/// <param name="password">Password for the account the service is to run under</param>
 		/// <returns>New service instance</returns>
-		public WindowsService Create(string Name, string DisplayName, string CommandLine, string? UserName, string? Password)
+		public WindowsService Create(string name, string displayName, string commandLine, string? userName, string? password)
 		{
-			Native.ServiceHandle NewServiceHandle = Native.CreateService(ServiceManagerHandle, Name, DisplayName, Native.SERVICE_ALL_ACCESS, Native.SERVICE_WIN32_OWN_PROCESS, Native.SERVICE_AUTO_START, Native.SERVICE_ERROR_NORMAL, CommandLine, null, null, null, UserName, Password);
-			if (NewServiceHandle.IsInvalid)
+			Native.ServiceHandle newServiceHandle = Native.CreateService(_serviceManagerHandle, name, displayName, Native.SERVICE_ALL_ACCESS, Native.SERVICE_WIN32_OWN_PROCESS, Native.SERVICE_AUTO_START, Native.SERVICE_ERROR_NORMAL, commandLine, null, null, null, userName, password);
+			if (newServiceHandle.IsInvalid)
 			{
 				throw new Win32Exception(String.Format("Unable to create service (0x{0:X8})", Marshal.GetLastWin32Error()));
 			}
-			return new WindowsService(NewServiceHandle);
+			return new WindowsService(newServiceHandle);
 		}
 	}
 }
diff --git a/Engine/Source/Programs/Horde/Horde.Agent/Utility/WorkspaceInfo.cs b/Engine/Source/Programs/Horde/Horde.Agent/Utility/WorkspaceInfo.cs
--- a/Engine/Source/Programs/Horde/Horde.Agent/Utility/WorkspaceInfo.cs
+++ b/Engine/Source/Programs/Horde/Horde.Agent/Utility/WorkspaceInfo.cs
@@ -34,174 +34,144 @@ namespace Horde.Agent.Utility
 		/// <summary>
 		/// The Perforce server and port. This is checked to not be null in the constructor.
 		/// </summary>
-		public string ServerAndPort
-		{
-			get { return PerforceClient.Settings.ServerAndPort!; }
-		}
+		public string ServerAndPort => PerforceClient.Settings.ServerAndPort!;
 
 		/// <summary>
 		/// The Perforce client name. This is checked to not be null in the constructor.
 		/// </summary>
-		public string ClientName
-		{
-			get { return PerforceClient.Settings.ClientName!; }
-		}
+		public string ClientName => PerforceClient.Settings.ClientName!;
 
 		/// <summary>
 		/// The Perforce user name. This is checked to not be null in the constructor.
 		/// </summary>
-		public string UserName
-		{
-			get { return PerforceClient.Settings.UserName!; }
-		}
+		public string UserName => PerforceClient.Settings.UserName!;
 
 		/// <summary>
 		/// The hostname
 		/// </summary>
-		public string HostName
-		{
-			get;
-		}
+		public string HostName { get; }
 
 		/// <summary>
 		/// The current stream
 		/// </summary>
-		public string StreamName
-		{
-			get;
-		}
+		public string StreamName { get; }
 
 		/// <summary>
 		/// The directory containing metadata for this workspace
 		/// </summary>
-		public DirectoryReference MetadataDir
-		{
-			get;
-		}
+		public DirectoryReference MetadataDir { get; }
 
 		/// <summary>
 		/// The directory containing workspace data
 		/// </summary>
-		public DirectoryReference WorkspaceDir
-		{
-			get;
-		}
+		public DirectoryReference WorkspaceDir { get; }
 
 		/// <summary>
 		/// The view for files to sync
 		/// </summary>
-		public IReadOnlyList<string> View
-		{
-			get;
-		}			
+		public IReadOnlyList<string> View { get; }
 
 		/// <summary>
 		/// Whether untracked files should be removed from this workspace
 		/// </summary>
-		public bool bRemoveUntrackedFiles
-		{
-			get;
-		}
+		public bool RemoveUntrackedFiles { get; }
 
 		/// <summary>
 		/// The managed repository instance
 		/// </summary>
-		public ManagedWorkspace Repository
-		{
-			get;
-		}
+		public ManagedWorkspace Repository { get; }
 
 		/// <summary>
 		/// Constructor
 		/// </summary>
-		/// <param name="Perforce">The perforce connection</param>
-		/// <param name="HostName">Name of this host</param>
-		/// <param name="StreamName">Name of the stream to sync</param>
-		/// <param name="MetadataDir">Path to the metadata directory</param>
-		/// <param name="WorkspaceDir">Path to the workspace directory</param>
-		/// <param name="View">View for files to be synced</param>
-		/// <param name="bRemoveUntrackedFiles">Whether to remove untracked files when syncing</param>
-		/// <param name="Repository">The repository instance</param>
-		public WorkspaceInfo(IPerforceConnection Perforce, string HostName, string StreamName, DirectoryReference MetadataDir, DirectoryReference WorkspaceDir, IList<string>? View, bool bRemoveUntrackedFiles, ManagedWorkspace Repository)
+		/// <param name="perforce">The perforce connection</param>
+		/// <param name="hostName">Name of this host</param>
+		/// <param name="streamName">Name of the stream to sync</param>
+		/// <param name="metadataDir">Path to the metadata directory</param>
+		/// <param name="workspaceDir">Path to the workspace directory</param>
+		/// <param name="view">View for files to be synced</param>
+		/// <param name="removeUntrackedFiles">Whether to remove untracked files when syncing</param>
+		/// <param name="repository">The repository instance</param>
+		public WorkspaceInfo(IPerforceConnection perforce, string hostName, string streamName, DirectoryReference metadataDir, DirectoryReference workspaceDir, IList<string>? view, bool removeUntrackedFiles, ManagedWorkspace repository)
 		{
-			this.PerforceClient = Perforce;
+			PerforceClient = perforce;
 
-			if (Perforce.Settings.ClientName == null)
+			if (perforce.Settings.ClientName == null)
 			{
 				throw new ArgumentException("PerforceConnection does not have valid client name");
 			}
 
-			this.HostName = HostName;
-			this.StreamName = StreamName;
-			this.MetadataDir = MetadataDir;
-			this.WorkspaceDir = WorkspaceDir;
-			this.View = (View == null) ? new List<string>() : new List<string>(View);
-			this.bRemoveUntrackedFiles = bRemoveUntrackedFiles;
-			this.Repository = Repository;
+			HostName = hostName;
+			StreamName = streamName;
+			MetadataDir = metadataDir;
+			WorkspaceDir = workspaceDir;
+			View = (view == null) ? new List<string>() : new List<string>(view);
+			RemoveUntrackedFiles = removeUntrackedFiles;
+			Repository = repository;
 		}
 
 		/// <summary>
 		/// Creates a new managed workspace
 		/// </summary>
-		/// <param name="Workspace">The workspace definition</param>
-		/// <param name="RootDir">Root directory for storing the workspace</param>
-		/// <param name="Logger">Logger output</param>
-		/// <param name="CancellationToken">Cancellation token</param>
+		/// <param name="workspace">The workspace definition</param>
+		/// <param name="rootDir">Root directory for storing the workspace</param>
+		/// <param name="logger">Logger output</param>
+		/// <param name="cancellationToken">Cancellation token</param>
 		/// <returns>New workspace info</returns>
-		public static async Task<WorkspaceInfo> SetupWorkspaceAsync(AgentWorkspace Workspace, DirectoryReference RootDir, ILogger Logger, CancellationToken CancellationToken)
+		public static async Task<WorkspaceInfo> SetupWorkspaceAsync(AgentWorkspace workspace, DirectoryReference rootDir, ILogger logger, CancellationToken cancellationToken)
 		{
 			// Fill in the default credentials iff they are not set
-			string? ServerAndPort = String.IsNullOrEmpty(Workspace.ServerAndPort)? null : Workspace.ServerAndPort;
-			string? UserName = String.IsNullOrEmpty(Workspace.UserName)? null : Workspace.UserName;
-			string? Password = String.IsNullOrEmpty(Workspace.Password)? null : Workspace.Password;
+			string? serverAndPort = String.IsNullOrEmpty(workspace.ServerAndPort)? null : workspace.ServerAndPort;
+			string? userName = String.IsNullOrEmpty(workspace.UserName)? null : workspace.UserName;
+			string? password = String.IsNullOrEmpty(workspace.Password)? null : workspace.Password;
 
-			if (ServerAndPort == null)
+			if (serverAndPort == null)
 			{
-				ServerAndPort = PerforceSettings.Default.ServerAndPort;
-				Logger.LogInformation("Using locally configured Perforce server: '{ServerAndPort}'", ServerAndPort);
+				serverAndPort = PerforceSettings.Default.ServerAndPort;
+				logger.LogInformation("Using locally configured Perforce server: '{ServerAndPort}'", serverAndPort);
 			}
 
-			if (UserName == null)
+			if (userName == null)
 			{
-				UserName = PerforceSettings.Default.UserName;
-				Logger.LogInformation("Using locally configured and logged in Perforce user: '{UserName}'", UserName);
+				userName = PerforceSettings.Default.UserName;
+				logger.LogInformation("Using locally configured and logged in Perforce user: '{UserName}'", userName);
 			}
 
 			// Create the connection
-			IPerforceConnection Perforce = await PerforceConnection.CreateAsync(new PerforceSettings(ServerAndPort, UserName), Logger);
-			if (UserName != null)
+			IPerforceConnection perforce = await PerforceConnection.CreateAsync(new PerforceSettings(serverAndPort, userName), logger);
+			if (userName != null)
 			{
-				if (Password != null)
+				if (password != null)
 				{
-					await Perforce.LoginAsync(Password, CancellationToken);
+					await perforce.LoginAsync(password, cancellationToken);
 				}
 				else
 				{
-					Logger.LogInformation($"Using locally logged in session for {UserName}");
+					logger.LogInformation($"Using locally logged in session for {userName}");
 				}
 			}
-			return await SetupWorkspaceAsync(Perforce, Workspace.Stream, Workspace.Identifier, Workspace.View, !Workspace.Incremental, RootDir, Logger, CancellationToken);
+			return await SetupWorkspaceAsync(perforce, workspace.Stream, workspace.Identifier, workspace.View, !workspace.Incremental, rootDir, logger, cancellationToken);
 		}
 
 		/// <summary>
 		/// Creates a new managed workspace
 		/// </summary>
-		/// <param name="Perforce">The perforce connection to use</param>
-		/// <param name="StreamName">The stream being synced</param>
-		/// <param name="Identifier">Identifier to use to share </param>
-		/// <param name="View">View for this workspace</param>
-		/// <param name="bRemoveUntrackedFiles">Whether untracked files should be removed when cleaning this workspace</param>
-		/// <param name="RootDir">Root directory for storing the workspace</param>
-		/// <param name="Logger">Logger output</param>
-		/// <param name="CancellationToken">Cancellation token</param>
+		/// <param name="perforce">The perforce connection to use</param>
+		/// <param name="streamName">The stream being synced</param>
+		/// <param name="identifier">Identifier to use to share </param>
+		/// <param name="view">View for this workspace</param>
+		/// <param name="removeUntrackedFiles">Whether untracked files should be removed when cleaning this workspace</param>
+		/// <param name="rootDir">Root directory for storing the workspace</param>
+		/// <param name="logger">Logger output</param>
+		/// <param name="cancellationToken">Cancellation token</param>
 		/// <returns>New workspace info</returns>
-		public static async Task<WorkspaceInfo> SetupWorkspaceAsync(IPerforceConnection Perforce, string StreamName, string Identifier, IList<string> View, bool bRemoveUntrackedFiles, DirectoryReference RootDir, ILogger Logger, CancellationToken CancellationToken)
+		public static async Task<WorkspaceInfo> SetupWorkspaceAsync(IPerforceConnection perforce, string streamName, string identifier, IList<string> view, bool removeUntrackedFiles, DirectoryReference rootDir, ILogger logger, CancellationToken cancellationToken)
 		{
 			// Get the host name, and fill in any missing metadata about the connection
-			InfoRecord Info = await Perforce.GetInfoAsync(InfoOptions.ShortOutput, CancellationToken);
+			InfoRecord info = await perforce.GetInfoAsync(InfoOptions.ShortOutput, cancellationToken);
 
-			string? HostName = Info.ClientHost;
-			if(HostName == null)
+			string? hostName = info.ClientHost;
+			if(hostName == null)
 			{
 				throw new Exception("Unable to determine Perforce host name");
 			}
@@ -213,175 +183,173 @@ namespace Horde.Agent.Utility
 			//			$workspace_identifier .= sprintf("+%02d", $slot_idx) if $slot_idx;
 
 			// if running on an edge server, append the server id to the client name
-			string EdgeSuffix = String.Empty;
-			if(Info.Services != null && Info.ServerID != null)
+			string edgeSuffix = String.Empty;
+			if(info.Services != null && info.ServerID != null)
 			{
-				string[] Services = Info.Services.Split(' ', StringSplitOptions.RemoveEmptyEntries);
-				if (Services.Any(x => x.Equals("edge-server", StringComparison.OrdinalIgnoreCase)))
+				string[] services = info.Services.Split(' ', StringSplitOptions.RemoveEmptyEntries);
+				if (services.Any(x => x.Equals("edge-server", StringComparison.OrdinalIgnoreCase)))
 				{
-					EdgeSuffix = $"+{Info.ServerID}";
+					edgeSuffix = $"+{info.ServerID}";
 				}
 			}
 
 			// get all the workspace settings
-			string ClientName = $"Horde+{GetNormalizedHostName(HostName)}+{Identifier}{EdgeSuffix}";
+			string clientName = $"Horde+{GetNormalizedHostName(hostName)}+{identifier}{edgeSuffix}";
 
 			// Create the client Perforce connection
-			IPerforceConnection PerforceClient = await Perforce.WithClientAsync(ClientName);
+			IPerforceConnection perforceClient = await perforce.WithClientAsync(clientName);
 
 			// get the workspace names
-			DirectoryReference MetadataDir = DirectoryReference.Combine(RootDir, Identifier);
-			DirectoryReference WorkspaceDir = DirectoryReference.Combine(MetadataDir, "Sync");
+			DirectoryReference metadataDir = DirectoryReference.Combine(rootDir, identifier);
+			DirectoryReference workspaceDir = DirectoryReference.Combine(metadataDir, "Sync");
 
 			// Create the repository
-			ManagedWorkspace NewRepository = await ManagedWorkspace.LoadOrCreateAsync(HostName, MetadataDir, true, Logger, CancellationToken);
-			if (bRemoveUntrackedFiles)
+			ManagedWorkspace newRepository = await ManagedWorkspace.LoadOrCreateAsync(hostName, metadataDir, true, logger, cancellationToken);
+			if (removeUntrackedFiles)
 			{
-				await NewRepository.DeleteClientAsync(PerforceClient, CancellationToken);
+				await newRepository.DeleteClientAsync(perforceClient, cancellationToken);
 			}
-			await NewRepository.SetupAsync(PerforceClient, StreamName, CancellationToken);
+			await newRepository.SetupAsync(perforceClient, streamName, cancellationToken);
 
 			// Revert any open files
-			await NewRepository.RevertAsync(PerforceClient, CancellationToken);
+			await newRepository.RevertAsync(perforceClient, cancellationToken);
 
 			// Create the workspace info
-			return new WorkspaceInfo(PerforceClient, HostName, StreamName, MetadataDir, WorkspaceDir, View, bRemoveUntrackedFiles, NewRepository);
+			return new WorkspaceInfo(perforceClient, hostName, streamName, metadataDir, workspaceDir, view, removeUntrackedFiles, newRepository);
 		}
 
 		/// <summary>
 		/// Gets the latest change in the stream
 		/// </summary>
-		/// <param name="CancellationToken">Cancellation token</param>
+		/// <param name="cancellationToken">Cancellation token</param>
 		/// <returns>Latest changelist number</returns>
-		public Task<int> GetLatestChangeAsync(CancellationToken CancellationToken)
+		public Task<int> GetLatestChangeAsync(CancellationToken cancellationToken)
 		{
-			return Repository.GetLatestChangeAsync(PerforceClient, StreamName, CancellationToken);
+			return Repository.GetLatestChangeAsync(PerforceClient, StreamName, cancellationToken);
 		}
 
 		/// <summary>
 		/// Revert any open files in the workspace and clean it
 		/// </summary>
-		/// <param name="CancellationToken">Cancellation token</param>
+		/// <param name="cancellationToken">Cancellation token</param>
 		/// <returns>Async task</returns>
-		public async Task CleanAsync(CancellationToken CancellationToken)
+		public async Task CleanAsync(CancellationToken cancellationToken)
 		{
-			await Repository.RevertAsync(PerforceClient, CancellationToken);
-			await Repository.CleanAsync(bRemoveUntrackedFiles, CancellationToken);
+			await Repository.RevertAsync(PerforceClient, cancellationToken);
+			await Repository.CleanAsync(RemoveUntrackedFiles, cancellationToken);
 		}
 
 		/// <summary>
 		/// Removes the given cache file if it's invalid, and updates the metadata to reflect the version about to be synced
 		/// </summary>
-		/// <param name="CacheFile">Path to the cache file</param>
-		/// <param name="Change">The current change being built</param>
-		/// <param name="PreflightChange">The preflight changelist number</param>
+		/// <param name="cacheFile">Path to the cache file</param>
+		/// <param name="change">The current change being built</param>
+		/// <param name="preflightChange">The preflight changelist number</param>
 		/// <returns>Async task</returns>
-		public async Task UpdateLocalCacheMarker(FileReference CacheFile, int Change, int PreflightChange)
+		public async Task UpdateLocalCacheMarker(FileReference cacheFile, int change, int preflightChange)
 		{
 			// Create the new cache file descriptor
-			string NewDescriptor;
-			if(PreflightChange <= 0)
+			string newDescriptor;
+			if(preflightChange <= 0)
 			{
-				NewDescriptor = $"CL {Change}";
+				newDescriptor = $"CL {change}";
 			}
 			else
 			{
-				NewDescriptor = $"CL {PreflightChange} with base CL {Change}";
+				newDescriptor = $"CL {preflightChange} with base CL {change}";
 			}
 
 			// Remove the cache file if the current descriptor doesn't match
-			FileReference DescriptorFile = CacheFile.ChangeExtension(".txt");
-			if (FileReference.Exists(CacheFile))
+			FileReference descriptorFile = cacheFile.ChangeExtension(".txt");
+			if (FileReference.Exists(cacheFile))
 			{
-				if (FileReference.Exists(DescriptorFile))
+				if (FileReference.Exists(descriptorFile))
 				{
-					string OldDescriptor = await FileReference.ReadAllTextAsync(DescriptorFile);
-					if (OldDescriptor == NewDescriptor)
+					string oldDescriptor = await FileReference.ReadAllTextAsync(descriptorFile);
+					if (oldDescriptor == newDescriptor)
 					{
 						return;
 					}
 					else
 					{
-						FileReference.Delete(DescriptorFile);
+						FileReference.Delete(descriptorFile);
 					}
 				}
-				FileReference.Delete(CacheFile);
+				FileReference.Delete(cacheFile);
 			}
 
 			// Write the new descriptor file
-			FileReference.WriteAllText(DescriptorFile, NewDescriptor);
+			FileReference.WriteAllText(descriptorFile, newDescriptor);
 		}
 
 		/// <summary>
 		/// Sync the workspace to a given changelist, and capture it so we can quickly clean in the future
 		/// </summary>
-		/// <param name="Change">The changelist to sync to</param>
-		/// <param name="PreflightChange">Change to preflight, or 0</param>
-		/// <param name="CacheFile">Path to the cache file to use</param>
-		/// <param name="Logger">The logger for output messages</param>
-		/// <param name="CancellationToken">Cancellation token</param>
+		/// <param name="change">The changelist to sync to</param>
+		/// <param name="preflightChange">Change to preflight, or 0</param>
+		/// <param name="cacheFile">Path to the cache file to use</param>
+		/// <param name="cancellationToken">Cancellation token</param>
 		/// <returns>Async task</returns>
-		public async Task SyncAsync(int Change, int PreflightChange, FileReference? CacheFile, ILogger Logger, CancellationToken CancellationToken)
+		public async Task SyncAsync(int change, int preflightChange, FileReference? cacheFile, CancellationToken cancellationToken)
 		{
-			await Repository.SyncAsync(PerforceClient, StreamName, Change, View, bRemoveUntrackedFiles, false, CacheFile, CancellationToken);
+			await Repository.SyncAsync(PerforceClient, StreamName, change, View, RemoveUntrackedFiles, false, cacheFile, cancellationToken);
 
 			// Purge the cache for incremental workspaces
-			if (!bRemoveUntrackedFiles)
+			if (!RemoveUntrackedFiles)
 			{
-				await Repository.PurgeAsync(0, CancellationToken);
+				await Repository.PurgeAsync(0, cancellationToken);
 			}
 
 			// if we're running a preflight
-			if (PreflightChange > 0)
+			if (preflightChange > 0)
 			{
-				await Repository.UnshelveAsync(PerforceClient, StreamName, PreflightChange, CancellationToken);
+				await Repository.UnshelveAsync(PerforceClient, StreamName, preflightChange, cancellationToken);
 			}
 		}
 
-
 		/// <summary>
 		/// Strip the domain suffix to get the host name
 		/// </summary>
-		/// <param name="HostName">Hostname with optional domain</param>
+		/// <param name="hostName">Hostname with optional domain</param>
 		/// <returns>Normalized host name</returns>
-		static string GetNormalizedHostName(string HostName)
+		static string GetNormalizedHostName(string hostName)
 		{
-			return Regex.Replace(HostName, @"\..*$", "").ToUpperInvariant();
+			return Regex.Replace(hostName, @"\..*$", "").ToUpperInvariant();
 		}
 
 		/// <summary>
 		/// Revert all files in a workspace
 		/// </summary>
-		public static async Task RevertAllChangesAsync(IPerforceConnection Perforce, ILogger Logger, CancellationToken CancellationToken)
+		public static async Task RevertAllChangesAsync(IPerforceConnection perforce, ILogger logger, CancellationToken cancellationToken)
 		{
 			// Make sure the client name is set
-			if(Perforce.Settings.ClientName == null)
+			if(perforce.Settings.ClientName == null)
 			{
 				throw new ArgumentException("RevertAllChangesAsync() requires PerforceConnection with client");
 			}
 
 			// check if there are any open files, and revert them
-			List<OpenedRecord> Files = await Perforce.OpenedAsync(OpenedOptions.None, -1, null, null, 1, new[] { "//..." }, CancellationToken).ToListAsync(CancellationToken);
-			if (Files.Count > 0)
+			List<OpenedRecord> files = await perforce.OpenedAsync(OpenedOptions.None, -1, null, null, 1, new[] { "//..." }, cancellationToken).ToListAsync(cancellationToken);
+			if (files.Count > 0)
 			{
-				await Perforce.RevertAsync(-1, null, RevertOptions.KeepWorkspaceFiles, new[] { "//..." }, CancellationToken);
+				await perforce.RevertAsync(-1, null, RevertOptions.KeepWorkspaceFiles, new[] { "//..." }, cancellationToken);
 			}
 
 			// enumerate all the pending changelists
-			List<ChangesRecord> PendingChanges = await Perforce.GetChangesAsync(ChangesOptions.None, Perforce.Settings.ClientName, -1, ChangeStatus.Pending, null, FileSpecList.Empty, CancellationToken);
-			foreach (ChangesRecord PendingChange in PendingChanges)
+			List<ChangesRecord> pendingChanges = await perforce.GetChangesAsync(ChangesOptions.None, perforce.Settings.ClientName, -1, ChangeStatus.Pending, null, FileSpecList.Empty, cancellationToken);
+			foreach (ChangesRecord pendingChange in pendingChanges)
 			{
 				// delete any shelved files if there are any
-				List<DescribeRecord> Records = await Perforce.DescribeAsync(DescribeOptions.Shelved, -1, new int[] { PendingChange.Number }, CancellationToken);
-				if (Records.Count > 0 && Records[0].Files.Count > 0)
+				List<DescribeRecord> records = await perforce.DescribeAsync(DescribeOptions.Shelved, -1, new int[] { pendingChange.Number }, cancellationToken);
+				if (records.Count > 0 && records[0].Files.Count > 0)
 				{
-					Logger.LogInformation($"Deleting shelved files in changelist {PendingChange.Number}");
-					await Perforce.DeleteChangeAsync(DeleteChangeOptions.None, PendingChange.Number, CancellationToken);
+					logger.LogInformation($"Deleting shelved files in changelist {pendingChange.Number}");
+					await perforce.DeleteChangeAsync(DeleteChangeOptions.None, pendingChange.Number, cancellationToken);
 				}
 
 				// delete the changelist
-				Logger.LogInformation($"Deleting changelist {PendingChange.Number}");
-				await Perforce.DeleteChangeAsync(DeleteChangeOptions.None, PendingChange.Number, CancellationToken);
+				logger.LogInformation($"Deleting changelist {pendingChange.Number}");
+				await perforce.DeleteChangeAsync(DeleteChangeOptions.None, pendingChange.Number, cancellationToken);
 			}
 		}
 	}
diff --git a/Engine/Source/Programs/Horde/Horde.Build/Utilities/ProtoExtensions.cs b/Engine/Source/Programs/Horde/Horde.Build/Utilities/ProtoExtensions.cs
--- a/Engine/Source/Programs/Horde/Horde.Build/Utilities/ProtoExtensions.cs
+++ b/Engine/Source/Programs/Horde/Horde.Build/Utilities/ProtoExtensions.cs
@@ -14,28 +14,28 @@ namespace HordeCommon
 		/// <summary>
 		/// Constructor
 		/// </summary>
-		/// <param name="Hash"></param>
-		public IoHashWrapper(IoHash Hash)
+		/// <param name="hash"></param>
+		public IoHashWrapper(IoHash hash)
 		{
-			this.Hash = ByteString.CopyFrom(Hash.ToByteArray());
+			Hash = ByteString.CopyFrom(hash.ToByteArray());
 		}
 
 		/// <summary>
 		/// Convert from an IoHashWrapper to an IoHash
 		/// </summary>
-		/// <param name="Hash"></param>
-		public static implicit operator IoHash(IoHashWrapper Hash)
+		/// <param name="hash"></param>
+		public static implicit operator IoHash(IoHashWrapper hash)
 		{
-			return new IoHash(Hash.Hash.ToByteArray());
+			return new IoHash(hash.Hash.ToByteArray());
 		}
 
 		/// <summary>
 		/// Convert from an IoHash to an IoHashWrapper
 		/// </summary>
-		/// <param name="Hash"></param>
-		public static implicit operator IoHashWrapper(IoHash Hash)
+		/// <param name="hash"></param>
+		public static implicit operator IoHashWrapper(IoHash hash)
 		{
-			return new IoHashWrapper(Hash);
+			return new IoHashWrapper(hash);
 		}
 	}
 
@@ -47,25 +47,25 @@ namespace HordeCommon
 		/// <summary>
 		/// Constructor
 		/// </summary>
-		/// <param name="RefId"></param>
-		public RefIdWrapper(RefId RefId)
+		/// <param name="refId"></param>
+		public RefIdWrapper(RefId refId)
 		{
-			this.Hash = ByteString.CopyFrom(RefId.Hash.ToByteArray());
+			Hash = ByteString.CopyFrom(refId.Hash.ToByteArray());
 		}
 
 		/// <summary>
 		/// Convert from an IoHashWrapper to an IoHash
 		/// </summary>
-		/// <param name="RefId"></param>
-		public static implicit operator RefId(RefIdWrapper RefId) => RefId.AsRefId();
+		/// <param name="refId"></param>
+		public static implicit operator RefId(RefIdWrapper refId) => refId.AsRefId();
 
 		/// <summary>
 		/// Convert from an IoHash to an IoHashWrapper
 		/// </summary>
-		/// <param name="RefId"></param>
-		public static implicit operator RefIdWrapper(RefId RefId)
+		/// <param name="refId"></param>
+		public static implicit operator RefIdWrapper(RefId refId)
 		{
-			return new RefIdWrapper(RefId);
+			return new RefIdWrapper(refId);
 		}
 
 		/// <summary>
