Horde: Remove unnecessary usings from Horde.Agent.

diff --git a/Engine/Source/Programs/Horde/Horde.Agent/AgentSettings.cs b/Engine/Source/Programs/Horde/Horde.Agent/AgentSettings.cs
--- a/Engine/Source/Programs/Horde/Horde.Agent/AgentSettings.cs
+++ b/Engine/Source/Programs/Horde/Horde.Agent/AgentSettings.cs
@@ -1,20 +1,12 @@
 // Copyright Epic Games, Inc. All Rights Reserved.
 
 using EpicGames.Core;
-using EpicGames.Horde.Compute;
 using EpicGames.Horde.Storage;
-using HordeCommon.Rpc;
 using Microsoft.Extensions.Configuration;
 using System;
 using System.Collections.Generic;
 using System.ComponentModel.DataAnnotations;
-using System.IO;
 using System.Linq;
-using System.Reflection;
-using System.Runtime.InteropServices;
-using System.Text;
-using System.Text.Json;
-using System.Threading.Tasks;
 
 namespace Horde.Agent
 {
diff --git a/Engine/Source/Programs/Horde/Horde.Agent/Commands/Bundles/CreateCommand.cs b/Engine/Source/Programs/Horde/Horde.Agent/Commands/Bundles/CreateCommand.cs
--- a/Engine/Source/Programs/Horde/Horde.Agent/Commands/Bundles/CreateCommand.cs
+++ b/Engine/Source/Programs/Horde/Horde.Agent/Commands/Bundles/CreateCommand.cs
@@ -11,9 +11,6 @@ using Microsoft.Extensions.Configuration;
 using Microsoft.Extensions.DependencyInjection;
 using Microsoft.Extensions.Logging;
 using System;
-using System.Collections.Generic;
-using System.IO;
-using System.Linq;
 using System.Threading.Tasks;
 
 namespace Horde.Agent.Commands.Bundles
diff --git a/Engine/Source/Programs/Horde/Horde.Agent/Commands/Parse/ParseCommand.cs b/Engine/Source/Programs/Horde/Horde.Agent/Commands/Parse/ParseCommand.cs
--- a/Engine/Source/Programs/Horde/Horde.Agent/Commands/Parse/ParseCommand.cs
+++ b/Engine/Source/Programs/Horde/Horde.Agent/Commands/Parse/ParseCommand.cs
@@ -7,7 +7,6 @@ using System;
 using System.Collections.Generic;
 using System.ComponentModel;
 using System.IO;
-using System.Text;
 using System.Threading.Tasks;
 
 namespace Horde.Agent.Commands.Parse
diff --git a/Engine/Source/Programs/Horde/Horde.Agent/Commands/Service/InstallCommand.cs b/Engine/Source/Programs/Horde/Horde.Agent/Commands/Service/InstallCommand.cs
--- a/Engine/Source/Programs/Horde/Horde.Agent/Commands/Service/InstallCommand.cs
+++ b/Engine/Source/Programs/Horde/Horde.Agent/Commands/Service/InstallCommand.cs
@@ -3,7 +3,6 @@
 using Horde.Agent.Utility;
 using Microsoft.Extensions.Logging;
 using System;
-using System.Collections.Generic;
 using System.Reflection;
 using System.Text;
 using System.Threading.Tasks;
diff --git a/Engine/Source/Programs/Horde/Horde.Agent/Commands/Service/RunCommand.cs b/Engine/Source/Programs/Horde/Horde.Agent/Commands/Service/RunCommand.cs
--- a/Engine/Source/Programs/Horde/Horde.Agent/Commands/Service/RunCommand.cs
+++ b/Engine/Source/Programs/Horde/Horde.Agent/Commands/Service/RunCommand.cs
@@ -1,6 +1,5 @@
 // Copyright Epic Games, Inc. All Rights Reserved.
 
-using Horde.Agent.Commands;
 using Horde.Agent.Services;
 using Horde.Agent.Utility;
 using Microsoft.Extensions.Configuration;
@@ -10,22 +9,14 @@ using Microsoft.Extensions.Logging;
 using Polly;
 using System;
 using System.Collections.Generic;
-using System.Linq;
 using System.Net.Http;
-using System.Net.Security;
-using System.Security.Cryptography.X509Certificates;
-using System.Text;
 using System.Threading.Tasks;
 using EpicGames.Core;
-using System.IO;
-using System.Reflection;
 using OpenTracing;
 using OpenTracing.Util;
 using Datadog.Trace.Configuration;
 using Datadog.Trace;
 using EpicGames.Horde.Storage;
-using EpicGames.Horde.Storage.Impl;
-using EpicGames.Horde.Compute;
 
 namespace Horde.Agent.Modes.Service
 {
diff --git a/Engine/Source/Programs/Horde/Horde.Agent/Commands/Service/UninstallCommand.cs b/Engine/Source/Programs/Horde/Horde.Agent/Commands/Service/UninstallCommand.cs
--- a/Engine/Source/Programs/Horde/Horde.Agent/Commands/Service/UninstallCommand.cs
+++ b/Engine/Source/Programs/Horde/Horde.Agent/Commands/Service/UninstallCommand.cs
@@ -4,9 +4,6 @@ using EpicGames.Core;
 using Horde.Agent.Utility;
 using Microsoft.Extensions.Logging;
 using System;
-using System.Collections.Generic;
-using System.Reflection;
-using System.Text;
 using System.Threading.Tasks;
 
 namespace Horde.Agent.Commands.Service
diff --git a/Engine/Source/Programs/Horde/Horde.Agent/Commands/Service/UpgradeCommand.cs b/Engine/Source/Programs/Horde/Horde.Agent/Commands/Service/UpgradeCommand.cs
--- a/Engine/Source/Programs/Horde/Horde.Agent/Commands/Service/UpgradeCommand.cs
+++ b/Engine/Source/Programs/Horde/Horde.Agent/Commands/Service/UpgradeCommand.cs
@@ -1,7 +1,5 @@
 // Copyright Epic Games, Inc. All Rights Reserved.
 
-using Horde.Agent.Services;
-using Horde.Agent.Utility;
 using Microsoft.Extensions.Logging;
 using System;
 using System.Collections.Generic;
@@ -15,8 +13,6 @@ using System.Text;
 using System.Threading.Tasks;
 using EpicGames.Core;
 using System.Runtime.InteropServices;
-using System.Text.RegularExpressions;
-using System.Threading;
 
 namespace Horde.Agent.Commands
 {
diff --git a/Engine/Source/Programs/Horde/Horde.Agent/Commands/Utilities/CapsCommand.cs b/Engine/Source/Programs/Horde/Horde.Agent/Commands/Utilities/CapsCommand.cs
--- a/Engine/Source/Programs/Horde/Horde.Agent/Commands/Utilities/CapsCommand.cs
+++ b/Engine/Source/Programs/Horde/Horde.Agent/Commands/Utilities/CapsCommand.cs
@@ -1,19 +1,6 @@
 // Copyright Epic Games, Inc. All Rights Reserved.
 
-using Grpc.Net.Client;
-using Horde.Agent.Utility;
-using HordeCommon.Rpc;
-using Microsoft.Extensions.Configuration;
 using Microsoft.Extensions.Logging;
-using System;
-using System.Collections.Generic;
-using System.IO;
-using System.IO.Compression;
-using System.Net;
-using System.Net.Http;
-using System.Net.Http.Headers;
-using System.Reflection;
-using System.Text;
 using System.Threading.Tasks;
 using EpicGames.Core;
 using HordeCommon.Rpc.Messages;
diff --git a/Engine/Source/Programs/Horde/Horde.Agent/Commands/Utilities/CreateCertCommand.cs b/Engine/Source/Programs/Horde/Horde.Agent/Commands/Utilities/CreateCertCommand.cs
--- a/Engine/Source/Programs/Horde/Horde.Agent/Commands/Utilities/CreateCertCommand.cs
+++ b/Engine/Source/Programs/Horde/Horde.Agent/Commands/Utilities/CreateCertCommand.cs
@@ -1,17 +1,10 @@
 // Copyright Epic Games, Inc. All Rights Reserved.
 
-using HordeCommon;
 using Microsoft.Extensions.Configuration;
 using Microsoft.Extensions.Logging;
-using System;
-using System.Collections.Generic;
 using System.IO;
-using System.Net;
 using System.Reflection;
-using System.Security.Cryptography;
 using System.Security.Cryptography.X509Certificates;
-using System.Text;
-using System.Text.Json;
 using System.Text.RegularExpressions;
 using System.Threading.Tasks;
 using EpicGames.Core;
diff --git a/Engine/Source/Programs/Horde/Horde.Agent/Commands/Utilities/PublishCommand.cs b/Engine/Source/Programs/Horde/Horde.Agent/Commands/Utilities/PublishCommand.cs
--- a/Engine/Source/Programs/Horde/Horde.Agent/Commands/Utilities/PublishCommand.cs
+++ b/Engine/Source/Programs/Horde/Horde.Agent/Commands/Utilities/PublishCommand.cs
@@ -6,14 +6,10 @@ using HordeCommon.Rpc;
 using Microsoft.Extensions.Configuration;
 using Microsoft.Extensions.Logging;
 using System;
-using System.Collections.Generic;
 using System.IO;
 using System.IO.Compression;
-using System.Net;
 using System.Net.Http;
 using System.Net.Http.Headers;
-using System.Reflection;
-using System.Text;
 using System.Threading.Tasks;
 using EpicGames.Core;
 
diff --git a/Engine/Source/Programs/Horde/Horde.Agent/Commands/Workspace/CleanCommand.cs b/Engine/Source/Programs/Horde/Horde.Agent/Commands/Workspace/CleanCommand.cs
--- a/Engine/Source/Programs/Horde/Horde.Agent/Commands/Workspace/CleanCommand.cs
+++ b/Engine/Source/Programs/Horde/Horde.Agent/Commands/Workspace/CleanCommand.cs
@@ -1,10 +1,6 @@
 // Copyright Epic Games, Inc. All Rights Reserved.
 
-using System;
-using System.Collections.Generic;
 using System.ComponentModel;
-using System.Linq;
-using System.Text;
 using System.Threading;
 using System.Threading.Tasks;
 using Microsoft.Extensions.Logging;
diff --git a/Engine/Source/Programs/Horde/Horde.Agent/Commands/Workspace/ClearCommand.cs b/Engine/Source/Programs/Horde/Horde.Agent/Commands/Workspace/ClearCommand.cs
--- a/Engine/Source/Programs/Horde/Horde.Agent/Commands/Workspace/ClearCommand.cs
+++ b/Engine/Source/Programs/Horde/Horde.Agent/Commands/Workspace/ClearCommand.cs
@@ -4,10 +4,6 @@ using EpicGames.Core;
 using EpicGames.Perforce;
 using EpicGames.Perforce.Managed;
 using Microsoft.Extensions.Logging;
-using System;
-using System.Collections.Generic;
-using System.Linq;
-using System.Text;
 using System.Threading;
 using System.Threading.Tasks;
 
diff --git a/Engine/Source/Programs/Horde/Horde.Agent/Commands/Workspace/DumpCommand.cs b/Engine/Source/Programs/Horde/Horde.Agent/Commands/Workspace/DumpCommand.cs
--- a/Engine/Source/Programs/Horde/Horde.Agent/Commands/Workspace/DumpCommand.cs
+++ b/Engine/Source/Programs/Horde/Horde.Agent/Commands/Workspace/DumpCommand.cs
@@ -4,10 +4,6 @@ using EpicGames.Core;
 using EpicGames.Perforce;
 using EpicGames.Perforce.Managed;
 using Microsoft.Extensions.Logging;
-using System;
-using System.Collections.Generic;
-using System.Linq;
-using System.Text;
 using System.Threading.Tasks;
 
 namespace Horde.Agent.Commands.Workspace
diff --git a/Engine/Source/Programs/Horde/Horde.Agent/Commands/Workspace/PopulateCacheCommand.cs b/Engine/Source/Programs/Horde/Horde.Agent/Commands/Workspace/PopulateCacheCommand.cs
--- a/Engine/Source/Programs/Horde/Horde.Agent/Commands/Workspace/PopulateCacheCommand.cs
+++ b/Engine/Source/Programs/Horde/Horde.Agent/Commands/Workspace/PopulateCacheCommand.cs
@@ -2,11 +2,8 @@
 
 using EpicGames.Perforce.Managed;
 using Microsoft.Extensions.Logging;
-using System;
 using System.Collections.Generic;
 using System.ComponentModel;
-using System.Linq;
-using System.Text;
 using System.Threading;
 using System.Threading.Tasks;
 using EpicGames.Core;
diff --git a/Engine/Source/Programs/Horde/Horde.Agent/Commands/Workspace/PurgeCacheCommand.cs b/Engine/Source/Programs/Horde/Horde.Agent/Commands/Workspace/PurgeCacheCommand.cs
--- a/Engine/Source/Programs/Horde/Horde.Agent/Commands/Workspace/PurgeCacheCommand.cs
+++ b/Engine/Source/Programs/Horde/Horde.Agent/Commands/Workspace/PurgeCacheCommand.cs
@@ -2,10 +2,6 @@
 
 using EpicGames.Perforce.Managed;
 using Microsoft.Extensions.Logging;
-using System;
-using System.Collections.Generic;
-using System.Linq;
-using System.Text;
 using System.Threading;
 using System.Threading.Tasks;
 using EpicGames.Core;
diff --git a/Engine/Source/Programs/Horde/Horde.Agent/Commands/Workspace/RepairCommand.cs b/Engine/Source/Programs/Horde/Horde.Agent/Commands/Workspace/RepairCommand.cs
--- a/Engine/Source/Programs/Horde/Horde.Agent/Commands/Workspace/RepairCommand.cs
+++ b/Engine/Source/Programs/Horde/Horde.Agent/Commands/Workspace/RepairCommand.cs
@@ -4,10 +4,6 @@ using EpicGames.Core;
 using EpicGames.Perforce;
 using EpicGames.Perforce.Managed;
 using Microsoft.Extensions.Logging;
-using System;
-using System.Collections.Generic;
-using System.Linq;
-using System.Text;
 using System.Threading;
 using System.Threading.Tasks;
 
diff --git a/Engine/Source/Programs/Horde/Horde.Agent/Commands/Workspace/RevertCommand.cs b/Engine/Source/Programs/Horde/Horde.Agent/Commands/Workspace/RevertCommand.cs
--- a/Engine/Source/Programs/Horde/Horde.Agent/Commands/Workspace/RevertCommand.cs
+++ b/Engine/Source/Programs/Horde/Horde.Agent/Commands/Workspace/RevertCommand.cs
@@ -2,11 +2,7 @@
 
 using EpicGames.Perforce.Managed;
 using Microsoft.Extensions.Logging;
-using System;
-using System.Collections.Generic;
 using System.ComponentModel;
-using System.Linq;
-using System.Text;
 using System.Threading;
 using System.Threading.Tasks;
 using EpicGames.Core;
diff --git a/Engine/Source/Programs/Horde/Horde.Agent/Commands/Workspace/SetupCommand.cs b/Engine/Source/Programs/Horde/Horde.Agent/Commands/Workspace/SetupCommand.cs
--- a/Engine/Source/Programs/Horde/Horde.Agent/Commands/Workspace/SetupCommand.cs
+++ b/Engine/Source/Programs/Horde/Horde.Agent/Commands/Workspace/SetupCommand.cs
@@ -1,10 +1,6 @@
 // Copyright Epic Games, Inc. All Rights Reserved.
 
-using System;
-using System.Collections.Generic;
 using System.ComponentModel;
-using System.Linq;
-using System.Text;
 using System.Threading;
 using System.Threading.Tasks;
 using EpicGames.Perforce.Managed;
diff --git a/Engine/Source/Programs/Horde/Horde.Agent/Commands/Workspace/SizesCommand.cs b/Engine/Source/Programs/Horde/Horde.Agent/Commands/Workspace/SizesCommand.cs
--- a/Engine/Source/Programs/Horde/Horde.Agent/Commands/Workspace/SizesCommand.cs
+++ b/Engine/Source/Programs/Horde/Horde.Agent/Commands/Workspace/SizesCommand.cs
@@ -1,10 +1,7 @@
 // Copyright Epic Games, Inc. All Rights Reserved.
 
-using System;
 using System.Collections.Generic;
 using System.ComponentModel;
-using System.Linq;
-using System.Text;
 using System.Threading;
 using System.Threading.Tasks;
 using Microsoft.Extensions.Logging;
diff --git a/Engine/Source/Programs/Horde/Horde.Agent/Commands/Workspace/StatusCommand.cs b/Engine/Source/Programs/Horde/Horde.Agent/Commands/Workspace/StatusCommand.cs
--- a/Engine/Source/Programs/Horde/Horde.Agent/Commands/Workspace/StatusCommand.cs
+++ b/Engine/Source/Programs/Horde/Horde.Agent/Commands/Workspace/StatusCommand.cs
@@ -4,10 +4,6 @@ using EpicGames.Core;
 using EpicGames.Perforce;
 using EpicGames.Perforce.Managed;
 using Microsoft.Extensions.Logging;
-using System;
-using System.Collections.Generic;
-using System.Linq;
-using System.Text;
 using System.Threading.Tasks;
 
 namespace Horde.Agent.Commands.Workspace
diff --git a/Engine/Source/Programs/Horde/Horde.Agent/Commands/Workspace/SyncCommand.cs b/Engine/Source/Programs/Horde/Horde.Agent/Commands/Workspace/SyncCommand.cs
--- a/Engine/Source/Programs/Horde/Horde.Agent/Commands/Workspace/SyncCommand.cs
+++ b/Engine/Source/Programs/Horde/Horde.Agent/Commands/Workspace/SyncCommand.cs
@@ -2,11 +2,8 @@
 
 using EpicGames.Perforce.Managed;
 using Microsoft.Extensions.Logging;
-using System;
 using System.Collections.Generic;
 using System.ComponentModel;
-using System.Linq;
-using System.Text;
 using System.Threading;
 using System.Threading.Tasks;
 using EpicGames.Core;
diff --git a/Engine/Source/Programs/Horde/Horde.Agent/Commands/Workspace/WorkspaceCommand.cs b/Engine/Source/Programs/Horde/Horde.Agent/Commands/Workspace/WorkspaceCommand.cs
--- a/Engine/Source/Programs/Horde/Horde.Agent/Commands/Workspace/WorkspaceCommand.cs
+++ b/Engine/Source/Programs/Horde/Horde.Agent/Commands/Workspace/WorkspaceCommand.cs
@@ -1,14 +1,10 @@
 // Copyright Epic Games, Inc. All Rights Reserved.
 
-using Horde.Agent.Utility;
 using Microsoft.Extensions.Logging;
 using System;
 using System.Collections.Generic;
 using System.ComponentModel;
-using System.Diagnostics;
-using System.IO;
 using System.Linq;
-using System.Text;
 using System.Threading;
 using System.Threading.Tasks;
 using EpicGames.Core;
diff --git a/Engine/Source/Programs/Horde/Horde.Agent/Execution/BuildGraphExecutor.cs b/Engine/Source/Programs/Horde/Horde.Agent/Execution/BuildGraphExecutor.cs
--- a/Engine/Source/Programs/Horde/Horde.Agent/Execution/BuildGraphExecutor.cs
+++ b/Engine/Source/Programs/Horde/Horde.Agent/Execution/BuildGraphExecutor.cs
@@ -13,16 +13,13 @@ using OpenTracing.Util;
 using System;
 using System.Collections.Generic;
 using System.Diagnostics;
-using System.Diagnostics.CodeAnalysis;
 using System.IO;
 using System.Linq;
-using System.Net.Security;
 using System.Runtime.InteropServices;
 using System.Security.Principal;
 using System.Text;
 using System.Text.Json;
 using System.Text.Json.Serialization;
-using System.Text.RegularExpressions;
 using System.Threading;
 using System.Threading.Tasks;
 
diff --git a/Engine/Source/Programs/Horde/Horde.Agent/Execution/ComputeTaskExecutor.cs b/Engine/Source/Programs/Horde/Horde.Agent/Execution/ComputeTaskExecutor.cs
--- a/Engine/Source/Programs/Horde/Horde.Agent/Execution/ComputeTaskExecutor.cs
+++ b/Engine/Source/Programs/Horde/Horde.Agent/Execution/ComputeTaskExecutor.cs
@@ -1,11 +1,8 @@
 // Copyright Epic Games, Inc. All Rights Reserved.
 
 using EpicGames.Core;
-using EpicGames.Horde;
 using EpicGames.Horde.Storage;
-using Google.Protobuf;
 using Grpc.Core;
-using Microsoft.Extensions.Options;
 using System;
 using System.Collections.Generic;
 using System.IO;
@@ -14,17 +11,10 @@ using System.Threading;
 using System.Threading.Tasks;
 using System.Diagnostics;
 using System.Text;
-using Google.Protobuf.WellKnownTypes;
-using Microsoft.Extensions.Hosting;
-using System.Collections.Concurrent;
-using Grpc.Net.Client;
-using Horde.Agent.Utility;
 using HordeCommon.Rpc.Tasks;
 using Microsoft.Extensions.Logging;
-using HordeCommon.Rpc;
 using EpicGames.Serialization;
 using EpicGames.Horde.Compute;
-using EpicGames.Horde.Common;
 
 namespace Horde.Agent.Execution
 {
diff --git a/Engine/Source/Programs/Horde/Horde.Agent/Execution/Interfaces/IExecutor.cs b/Engine/Source/Programs/Horde/Horde.Agent/Execution/Interfaces/IExecutor.cs
--- a/Engine/Source/Programs/Horde/Horde.Agent/Execution/Interfaces/IExecutor.cs
+++ b/Engine/Source/Programs/Horde/Horde.Agent/Execution/Interfaces/IExecutor.cs
@@ -1,13 +1,8 @@
 // Copyright Epic Games, Inc. All Rights Reserved.
 
-using Horde.Agent.Parser.Interfaces;
-using Horde.Agent.Utility;
 using HordeCommon;
 using HordeCommon.Rpc;
 using Microsoft.Extensions.Logging;
-using System;
-using System.Collections.Generic;
-using System.Text;
 using System.Threading;
 using System.Threading.Tasks;
 
diff --git a/Engine/Source/Programs/Horde/Horde.Agent/Execution/LocalExecutor.cs b/Engine/Source/Programs/Horde/Horde.Agent/Execution/LocalExecutor.cs
--- a/Engine/Source/Programs/Horde/Horde.Agent/Execution/LocalExecutor.cs
+++ b/Engine/Source/Programs/Horde/Horde.Agent/Execution/LocalExecutor.cs
@@ -2,12 +2,10 @@
 
 using EpicGames.Core;
 using Horde.Agent.Utility;
-using HordeCommon;
 using HordeCommon.Rpc;
 using Microsoft.Extensions.Logging;
 using System;
 using System.Collections.Generic;
-using System.Text;
 using System.Threading;
 using System.Threading.Tasks;
 
diff --git a/Engine/Source/Programs/Horde/Horde.Agent/Execution/PerforceExecutor.cs b/Engine/Source/Programs/Horde/Horde.Agent/Execution/PerforceExecutor.cs
--- a/Engine/Source/Programs/Horde/Horde.Agent/Execution/PerforceExecutor.cs
+++ b/Engine/Source/Programs/Horde/Horde.Agent/Execution/PerforceExecutor.cs
@@ -1,32 +1,20 @@
 // Copyright Epic Games, Inc. All Rights Reserved.
 
 using EpicGames.Perforce.Managed;
-using Horde.Agent.Execution.Interfaces;
-using Horde.Agent.Parser;
-using Horde.Agent.Parser.Interfaces;
 using Horde.Agent.Utility;
-using HordeCommon;
 using HordeCommon.Rpc;
 using HordeCommon.Rpc.Messages;
-using HordeCommon.Rpc.Tasks;
 using Microsoft.Extensions.Logging;
 using System;
 using System.Collections.Generic;
-using System.ComponentModel.DataAnnotations;
-using System.Diagnostics;
 using System.IO;
 using System.Linq;
-using System.Net.Http;
 using System.Runtime.InteropServices;
-using System.Text;
-using System.Text.Json;
-using System.Text.Json.Serialization;
 using System.Text.RegularExpressions;
 using System.Threading;
 using System.Threading.Tasks;
 using EpicGames.Core;
 using EpicGames.Perforce;
-using Google.Protobuf.WellKnownTypes;
 using OpenTracing;
 using OpenTracing.Util;
 
diff --git a/Engine/Source/Programs/Horde/Horde.Agent/Execution/TestExecutor.cs b/Engine/Source/Programs/Horde/Horde.Agent/Execution/TestExecutor.cs
--- a/Engine/Source/Programs/Horde/Horde.Agent/Execution/TestExecutor.cs
+++ b/Engine/Source/Programs/Horde/Horde.Agent/Execution/TestExecutor.cs
@@ -2,18 +2,12 @@
 
 using System;
 using System.Collections.Generic;
-using System.IO;
 using System.Linq;
-using System.Net.Http;
 using System.Reflection;
-using System.Text;
-using System.Text.Json;
 using System.Threading;
 using System.Threading.Tasks;
 using EpicGames.Core;
-using Horde.Agent.Execution.Interfaces;
 using Horde.Agent.Parser;
-using Horde.Agent.Parser.Interfaces;
 using Horde.Agent.Utility;
 using HordeCommon;
 using HordeCommon.Rpc;
diff --git a/Engine/Source/Programs/Horde/Horde.Agent/Logging.cs b/Engine/Source/Programs/Horde/Horde.Agent/Logging.cs
--- a/Engine/Source/Programs/Horde/Horde.Agent/Logging.cs
+++ b/Engine/Source/Programs/Horde/Horde.Agent/Logging.cs
@@ -8,17 +8,10 @@ using Serilog.Extensions.Logging;
 using Serilog.Formatting.Json;
 using Serilog.Sinks.SystemConsole.Themes;
 using System;
-using System.IO;
 using System.Runtime.InteropServices;
 using EpicGames.Core;
-using System.Collections.Generic;
-using System.Threading;
 using OpenTracing.Util;
 using OpenTracing;
-using Serilog.Formatting;
-using System.Text.Json;
-using System.Buffers;
-using MessageTemplate = EpicGames.Core.MessageTemplate;
 
 namespace Horde.Agent
 {
diff --git a/Engine/Source/Programs/Horde/Horde.Agent/Parser/Interfaces/ILogContext.cs b/Engine/Source/Programs/Horde/Horde.Agent/Parser/Interfaces/ILogContext.cs
--- a/Engine/Source/Programs/Horde/Horde.Agent/Parser/Interfaces/ILogContext.cs
+++ b/Engine/Source/Programs/Horde/Horde.Agent/Parser/Interfaces/ILogContext.cs
@@ -1,10 +1,6 @@
 // Copyright Epic Games, Inc. All Rights Reserved.
 
 using EpicGames.Core;
-using Microsoft.Extensions.Logging;
-using System;
-using System.Collections.Generic;
-using System.Text;
 
 namespace Horde.Agent.Parser.Interfaces
 {
diff --git a/Engine/Source/Programs/Horde/Horde.Agent/Parser/LogEventMarkup.cs b/Engine/Source/Programs/Horde/Horde.Agent/Parser/LogEventMarkup.cs
--- a/Engine/Source/Programs/Horde/Horde.Agent/Parser/LogEventMarkup.cs
+++ b/Engine/Source/Programs/Horde/Horde.Agent/Parser/LogEventMarkup.cs
@@ -2,12 +2,8 @@
 
 using EpicGames.Core;
 using Horde.Agent.Parser.Interfaces;
-using System;
 using System.Collections.Generic;
 using System.IO;
-using System.Reflection;
-using System.Text;
-using System.Text.Json;
 using System.Text.RegularExpressions;
 
 namespace Horde.Agent.Parser
diff --git a/Engine/Source/Programs/Horde/Horde.Agent/Parser/LogParser.cs b/Engine/Source/Programs/Horde/Horde.Agent/Parser/LogParser.cs
--- a/Engine/Source/Programs/Horde/Horde.Agent/Parser/LogParser.cs
+++ b/Engine/Source/Programs/Horde/Horde.Agent/Parser/LogParser.cs
@@ -2,21 +2,12 @@
 
 using EpicGames.Core;
 using Horde.Agent.Parser.Interfaces;
-using Horde.Agent.Utility;
-using HordeCommon;
 using Microsoft.Extensions.Logging;
 using System;
-using System.Buffers.Text;
-using System.Collections.Concurrent;
 using System.Collections.Generic;
-using System.IO;
 using System.Linq;
 using System.Reflection;
-using System.Text;
-using System.Text.Json;
 using System.Text.RegularExpressions;
-using System.Threading;
-using System.Threading.Tasks;
 
 namespace Horde.Agent.Parser
 {
diff --git a/Engine/Source/Programs/Horde/Horde.Agent/Parser/Matchers/CompileEventMatcher.cs b/Engine/Source/Programs/Horde/Horde.Agent/Parser/Matchers/CompileEventMatcher.cs
--- a/Engine/Source/Programs/Horde/Horde.Agent/Parser/Matchers/CompileEventMatcher.cs
+++ b/Engine/Source/Programs/Horde/Horde.Agent/Parser/Matchers/CompileEventMatcher.cs
@@ -1,18 +1,12 @@
 // Copyright Epic Games, Inc. All Rights Reserved.
 
 using EpicGames.Core;
-using Horde.Agent.Parser;
 using Horde.Agent.Parser.Interfaces;
-using HordeCommon;
 using Microsoft.Extensions.Logging;
 using System;
-using System.Collections.Generic;
 using System.Diagnostics.CodeAnalysis;
-using System.IO;
 using System.Linq;
-using System.Text;
 using System.Text.RegularExpressions;
-using System.Threading.Tasks;
 
 namespace Horde.Agent.Parser.Matchers
 {
diff --git a/Engine/Source/Programs/Horde/Horde.Agent/Parser/Matchers/ContentEventMatcher.cs b/Engine/Source/Programs/Horde/Horde.Agent/Parser/Matchers/ContentEventMatcher.cs
--- a/Engine/Source/Programs/Horde/Horde.Agent/Parser/Matchers/ContentEventMatcher.cs
+++ b/Engine/Source/Programs/Horde/Horde.Agent/Parser/Matchers/ContentEventMatcher.cs
@@ -1,12 +1,8 @@
 // Copyright Epic Games, Inc. All Rights Reserved.
 
 using EpicGames.Core;
-using Horde.Agent.Parser.Interfaces;
-using HordeCommon;
 using Microsoft.Extensions.Logging;
 using System;
-using System.Collections.Generic;
-using System.Text;
 using System.Text.RegularExpressions;
 
 namespace Horde.Agent.Parser.Matchers
diff --git a/Engine/Source/Programs/Horde/Horde.Agent/Parser/Matchers/CrashEventMatcher.cs b/Engine/Source/Programs/Horde/Horde.Agent/Parser/Matchers/CrashEventMatcher.cs
--- a/Engine/Source/Programs/Horde/Horde.Agent/Parser/Matchers/CrashEventMatcher.cs
+++ b/Engine/Source/Programs/Horde/Horde.Agent/Parser/Matchers/CrashEventMatcher.cs
@@ -1,16 +1,8 @@
 // Copyright Epic Games, Inc. All Rights Reserved.
 
 using EpicGames.Core;
-using Horde.Agent.Parser;
-using Horde.Agent.Parser.Interfaces;
-using HordeCommon;
 using Microsoft.Extensions.Logging;
-using System;
-using System.Collections.Generic;
-using System.Linq;
-using System.Text;
 using System.Text.RegularExpressions;
-using System.Threading.Tasks;
 
 namespace Horde.Agent.Parser.Matchers
 {
diff --git a/Engine/Source/Programs/Horde/Horde.Agent/Parser/Matchers/ExceptionEventMatcher.cs b/Engine/Source/Programs/Horde/Horde.Agent/Parser/Matchers/ExceptionEventMatcher.cs
--- a/Engine/Source/Programs/Horde/Horde.Agent/Parser/Matchers/ExceptionEventMatcher.cs
+++ b/Engine/Source/Programs/Horde/Horde.Agent/Parser/Matchers/ExceptionEventMatcher.cs
@@ -1,15 +1,7 @@
 // Copyright Epic Games, Inc. All Rights Reserved.
 
 using EpicGames.Core;
-using Horde.Agent.Parser;
-using Horde.Agent.Parser.Interfaces;
-using HordeCommon;
 using Microsoft.Extensions.Logging;
-using System;
-using System.Collections.Generic;
-using System.Linq;
-using System.Text;
-using System.Threading.Tasks;
 
 namespace Horde.Agent.Parser.Matchers
 {
diff --git a/Engine/Source/Programs/Horde/Horde.Agent/Parser/Matchers/ExitCodeEventMatcher.cs b/Engine/Source/Programs/Horde/Horde.Agent/Parser/Matchers/ExitCodeEventMatcher.cs
--- a/Engine/Source/Programs/Horde/Horde.Agent/Parser/Matchers/ExitCodeEventMatcher.cs
+++ b/Engine/Source/Programs/Horde/Horde.Agent/Parser/Matchers/ExitCodeEventMatcher.cs
@@ -1,8 +1,6 @@
 // Copyright Epic Games, Inc. All Rights Reserved.
 
 using EpicGames.Core;
-using Horde.Agent.Parser.Interfaces;
-using HordeCommon;
 using Microsoft.Extensions.Logging;
 
 namespace Horde.Agent.Parser.Matchers
diff --git a/Engine/Source/Programs/Horde/Horde.Agent/Parser/Matchers/GauntletEventMatcher.cs b/Engine/Source/Programs/Horde/Horde.Agent/Parser/Matchers/GauntletEventMatcher.cs
--- a/Engine/Source/Programs/Horde/Horde.Agent/Parser/Matchers/GauntletEventMatcher.cs
+++ b/Engine/Source/Programs/Horde/Horde.Agent/Parser/Matchers/GauntletEventMatcher.cs
@@ -1,12 +1,7 @@
 // Copyright Epic Games, Inc. All Rights Reserved.
 
 using EpicGames.Core;
-using Horde.Agent.Parser.Interfaces;
-using HordeCommon;
 using Microsoft.Extensions.Logging;
-using System;
-using System.Collections.Generic;
-using System.Text;
 using System.Text.RegularExpressions;
 
 namespace Horde.Agent.Parser.Matchers
diff --git a/Engine/Source/Programs/Horde/Horde.Agent/Parser/Matchers/GenericEventMatcher.cs b/Engine/Source/Programs/Horde/Horde.Agent/Parser/Matchers/GenericEventMatcher.cs
--- a/Engine/Source/Programs/Horde/Horde.Agent/Parser/Matchers/GenericEventMatcher.cs
+++ b/Engine/Source/Programs/Horde/Horde.Agent/Parser/Matchers/GenericEventMatcher.cs
@@ -1,16 +1,9 @@
 // Copyright Epic Games, Inc. All Rights Reserved.
 
 using EpicGames.Core;
-using Horde.Agent.Parser;
-using Horde.Agent.Parser.Interfaces;
-using HordeCommon;
 using Microsoft.Extensions.Logging;
 using System;
-using System.Collections.Generic;
-using System.Linq;
-using System.Text;
 using System.Text.RegularExpressions;
-using System.Threading.Tasks;
 
 namespace Horde.Agent.Parser.Matchers
 {
diff --git a/Engine/Source/Programs/Horde/Horde.Agent/Parser/Matchers/GradleEventMatcher.cs b/Engine/Source/Programs/Horde/Horde.Agent/Parser/Matchers/GradleEventMatcher.cs
--- a/Engine/Source/Programs/Horde/Horde.Agent/Parser/Matchers/GradleEventMatcher.cs
+++ b/Engine/Source/Programs/Horde/Horde.Agent/Parser/Matchers/GradleEventMatcher.cs
@@ -1,8 +1,6 @@
 // Copyright Epic Games, Inc. All Rights Reserved.
 
 using EpicGames.Core;
-using Horde.Agent.Parser.Interfaces;
-using HordeCommon;
 using Microsoft.Extensions.Logging;
 using System.Text.RegularExpressions;
 
diff --git a/Engine/Source/Programs/Horde/Horde.Agent/Parser/Matchers/LinkEventMatcher.cs b/Engine/Source/Programs/Horde/Horde.Agent/Parser/Matchers/LinkEventMatcher.cs
--- a/Engine/Source/Programs/Horde/Horde.Agent/Parser/Matchers/LinkEventMatcher.cs
+++ b/Engine/Source/Programs/Horde/Horde.Agent/Parser/Matchers/LinkEventMatcher.cs
@@ -1,16 +1,9 @@
 // Copyright Epic Games, Inc. All Rights Reserved.
 
 using EpicGames.Core;
-using Horde.Agent.Parser;
-using Horde.Agent.Parser.Interfaces;
-using HordeCommon;
 using Microsoft.Extensions.Logging;
 using System;
-using System.Collections.Generic;
-using System.Linq;
-using System.Text;
 using System.Text.RegularExpressions;
-using System.Threading.Tasks;
 
 namespace Horde.Agent.Parser.Matchers
 {
diff --git a/Engine/Source/Programs/Horde/Horde.Agent/Parser/Matchers/LogChannelEventMatcher.cs b/Engine/Source/Programs/Horde/Horde.Agent/Parser/Matchers/LogChannelEventMatcher.cs
--- a/Engine/Source/Programs/Horde/Horde.Agent/Parser/Matchers/LogChannelEventMatcher.cs
+++ b/Engine/Source/Programs/Horde/Horde.Agent/Parser/Matchers/LogChannelEventMatcher.cs
@@ -1,16 +1,8 @@
 // Copyright Epic Games, Inc. All Rights Reserved.
 
 using EpicGames.Core;
-using Horde.Agent.Parser;
-using Horde.Agent.Parser.Interfaces;
-using HordeCommon;
 using Microsoft.Extensions.Logging;
-using System;
-using System.Collections.Generic;
-using System.Linq;
-using System.Text;
 using System.Text.RegularExpressions;
-using System.Threading.Tasks;
 
 namespace Horde.Agent.Parser.Matchers
 {
diff --git a/Engine/Source/Programs/Horde/Horde.Agent/Parser/Matchers/MicrosoftEventMatcher.cs b/Engine/Source/Programs/Horde/Horde.Agent/Parser/Matchers/MicrosoftEventMatcher.cs
--- a/Engine/Source/Programs/Horde/Horde.Agent/Parser/Matchers/MicrosoftEventMatcher.cs
+++ b/Engine/Source/Programs/Horde/Horde.Agent/Parser/Matchers/MicrosoftEventMatcher.cs
@@ -2,11 +2,7 @@
 
 using EpicGames.Core;
 using Horde.Agent.Parser.Interfaces;
-using HordeCommon;
 using Microsoft.Extensions.Logging;
-using System;
-using System.Collections.Generic;
-using System.Text;
 using System.Text.RegularExpressions;
 
 namespace Horde.Agent.Parser.Matchers
diff --git a/Engine/Source/Programs/Horde/Horde.Agent/Parser/Matchers/SourceFileLineEventMatcher.cs b/Engine/Source/Programs/Horde/Horde.Agent/Parser/Matchers/SourceFileLineEventMatcher.cs
--- a/Engine/Source/Programs/Horde/Horde.Agent/Parser/Matchers/SourceFileLineEventMatcher.cs
+++ b/Engine/Source/Programs/Horde/Horde.Agent/Parser/Matchers/SourceFileLineEventMatcher.cs
@@ -2,12 +2,8 @@
 
 using EpicGames.Core;
 using Horde.Agent.Parser.Interfaces;
-using HordeCommon;
 using Microsoft.Extensions.Logging;
 using System;
-using System.Collections.Generic;
-using System.Runtime.CompilerServices;
-using System.Text;
 using System.Text.RegularExpressions;
 
 namespace Horde.Agent.Parser.Matchers
diff --git a/Engine/Source/Programs/Horde/Horde.Agent/Parser/Matchers/SystemicEventMatcher.cs b/Engine/Source/Programs/Horde/Horde.Agent/Parser/Matchers/SystemicEventMatcher.cs
--- a/Engine/Source/Programs/Horde/Horde.Agent/Parser/Matchers/SystemicEventMatcher.cs
+++ b/Engine/Source/Programs/Horde/Horde.Agent/Parser/Matchers/SystemicEventMatcher.cs
@@ -1,16 +1,7 @@
 // Copyright Epic Games, Inc. All Rights Reserved.
 
 using EpicGames.Core;
-using Horde.Agent.Parser;
-using Horde.Agent.Parser.Interfaces;
-using HordeCommon;
 using Microsoft.Extensions.Logging;
-using System;
-using System.Collections.Generic;
-using System.Linq;
-using System.Text;
-using System.Text.RegularExpressions;
-using System.Threading.Tasks;
 
 namespace Horde.Agent.Parser.Matchers
 {
diff --git a/Engine/Source/Programs/Horde/Horde.Agent/Parser/Matchers/XoreaxEventMatcher.cs b/Engine/Source/Programs/Horde/Horde.Agent/Parser/Matchers/XoreaxEventMatcher.cs
--- a/Engine/Source/Programs/Horde/Horde.Agent/Parser/Matchers/XoreaxEventMatcher.cs
+++ b/Engine/Source/Programs/Horde/Horde.Agent/Parser/Matchers/XoreaxEventMatcher.cs
@@ -1,16 +1,8 @@
 // Copyright Epic Games, Inc. All Rights Reserved.
 
 using EpicGames.Core;
-using Horde.Agent.Parser;
-using Horde.Agent.Parser.Interfaces;
-using HordeCommon;
 using Microsoft.Extensions.Logging;
-using System;
-using System.Collections.Generic;
-using System.Linq;
-using System.Text;
 using System.Text.RegularExpressions;
-using System.Threading.Tasks;
 
 namespace Horde.Agent.Parser.Matchers
 {
diff --git a/Engine/Source/Programs/Horde/Horde.Agent/Program.cs b/Engine/Source/Programs/Horde/Horde.Agent/Program.cs
--- a/Engine/Source/Programs/Horde/Horde.Agent/Program.cs
+++ b/Engine/Source/Programs/Horde/Horde.Agent/Program.cs
@@ -1,34 +1,14 @@
 // Copyright Epic Games, Inc. All Rights Reserved.
 
-using Horde.Agent.Commands;
-using Horde.Agent.Services;
-using Horde.Agent.Utility;
-using Microsoft.Extensions.Configuration;
 using Microsoft.Extensions.DependencyInjection;
-using Microsoft.Extensions.Hosting;
-using Microsoft.Extensions.Hosting.WindowsServices;
 using Microsoft.Extensions.Logging;
-using Polly;
 using System;
-using System.Collections.Generic;
 using System.Diagnostics;
 using System.IO;
-using System.Linq;
-using System.Management;
-using System.Net;
-using System.Net.Http.Headers;
 using System.Reflection;
 using System.Runtime.InteropServices;
-using System.ServiceProcess;
-using System.Text.Json;
-using System.Threading;
 using System.Threading.Tasks;
 using EpicGames.Core;
-using EpicGames.Horde.Storage;
-using EpicGames.Serialization;
-using System.Text;
-using EpicGames.Perforce;
-using System.Security.Cryptography;
 
 namespace Horde.Agent
 {
diff --git a/Engine/Source/Programs/Horde/Horde.Agent/Services/GrpcService.cs b/Engine/Source/Programs/Horde/Horde.Agent/Services/GrpcService.cs
--- a/Engine/Source/Programs/Horde/Horde.Agent/Services/GrpcService.cs
+++ b/Engine/Source/Programs/Horde/Horde.Agent/Services/GrpcService.cs
@@ -8,12 +8,8 @@ using Microsoft.Extensions.Options;
 using Polly;
 using Polly.Extensions.Http;
 using System;
-using System.Collections.Generic;
 using System.Net.Http;
 using System.Net.Http.Headers;
-using System.Text;
-using Grpc.Core;
-using HordeCommon.Rpc;
 
 namespace Horde.Agent.Services
 {
diff --git a/Engine/Source/Programs/Horde/Horde.Agent/Services/WorkerService.cs b/Engine/Source/Programs/Horde/Horde.Agent/Services/WorkerService.cs
--- a/Engine/Source/Programs/Horde/Horde.Agent/Services/WorkerService.cs
+++ b/Engine/Source/Programs/Horde/Horde.Agent/Services/WorkerService.cs
@@ -29,19 +29,15 @@ using System.IO;
 using System.IO.Compression;
 using System.Linq;
 using System.Management;
-using System.Net.Http.Headers;
 using System.Reflection;
-using System.Runtime.CompilerServices;
 using System.Runtime.InteropServices;
 using System.Text;
 using System.Threading;
 using System.Threading.Tasks;
 using System.Xml;
-using Status = Grpc.Core.Status;
 using EpicGames.Horde.Storage;
 using EpicGames.Horde.Compute;
 using System.Net;
-using EpicGames.Horde.Storage.Impl;
 using System.Text.RegularExpressions;
 
 namespace Horde.Agent.Services
@@ -1023,7 +1019,7 @@ namespace Horde.Agent.Services
 			// Execute the batch
 			try
 			{
-				await ExecuteBatchAsync(rpcClient, agentId, leaseId, executeTask, batch, logger, cancellationToken);
+				await ExecuteBatchAsync(rpcClient, leaseId, executeTask, batch, logger, cancellationToken);
 			}
 			catch (Exception ex)
 			{
@@ -1055,14 +1051,13 @@ namespace Horde.Agent.Services
 		/// Executes a batch
 		/// </summary>
 		/// <param name="rpcClient">RPC client for communicating with the server</param>
-		/// <param name="agentId">The current agent id</param>
 		/// <param name="leaseId">The current lease id</param>
 		/// <param name="executeTask">The task to execute</param>
 		/// <param name="batch">The batch to execute</param>
 		/// <param name="batchLogger">Output log for the batch</param>
 		/// <param name="cancellationToken">Cancellation token to abort the batch</param>
 		/// <returns>Async task</returns>
-		async Task ExecuteBatchAsync(IRpcConnection rpcClient, string agentId, string leaseId, ExecuteJobTask executeTask, BeginBatchResponse batch, ILogger batchLogger, CancellationToken cancellationToken)
+		async Task ExecuteBatchAsync(IRpcConnection rpcClient, string leaseId, ExecuteJobTask executeTask, BeginBatchResponse batch, ILogger batchLogger, CancellationToken cancellationToken)
 		{
 			batchLogger.LogInformation("Executing batch {BatchId} using {Executor} executor", executeTask.BatchId, _settings.Executor.ToString());
 			TerminateProcesses(batchLogger, cancellationToken);
diff --git a/Engine/Source/Programs/Horde/Horde.Agent/Tracing.cs b/Engine/Source/Programs/Horde/Horde.Agent/Tracing.cs
--- a/Engine/Source/Programs/Horde/Horde.Agent/Tracing.cs
+++ b/Engine/Source/Programs/Horde/Horde.Agent/Tracing.cs
@@ -1,9 +1,6 @@
 // Copyright Epic Games, Inc. All Rights Reserved.
 
 using OpenTracing;
-using System;
-using System.Collections.Generic;
-using System.Text;
 
 namespace Horde.Agent
 {
diff --git a/Engine/Source/Programs/Horde/Horde.Agent/Utilities/HordeRpc.cs b/Engine/Source/Programs/Horde/Horde.Agent/Utilities/HordeRpc.cs
--- a/Engine/Source/Programs/Horde/Horde.Agent/Utilities/HordeRpc.cs
+++ b/Engine/Source/Programs/Horde/Horde.Agent/Utilities/HordeRpc.cs
@@ -2,13 +2,9 @@
 
 using Google.Protobuf;
 using Google.Protobuf.Collections;
-using Google.Protobuf.WellKnownTypes;
-using HordeCommon;
-using System;
 using System.Collections.Generic;
 using System.Diagnostics.CodeAnalysis;
 using System.Linq;
-using System.Text;
 
 namespace HordeCommon.Rpc
 {
diff --git a/Engine/Source/Programs/Horde/Horde.Agent/Utilities/ProtoMethods.cs b/Engine/Source/Programs/Horde/Horde.Agent/Utilities/ProtoMethods.cs
--- a/Engine/Source/Programs/Horde/Horde.Agent/Utilities/ProtoMethods.cs
+++ b/Engine/Source/Programs/Horde/Horde.Agent/Utilities/ProtoMethods.cs
@@ -1,11 +1,6 @@
 // Copyright Epic Games, Inc. All Rights Reserved.
 
-using EpicGames.Horde.Common;
 using EpicGames.Horde.Compute;
-using EpicGames.Serialization;
-using System;
-using System.Collections.Generic;
-using System.Text;
 
 namespace HordeCommon.Rpc.Tasks
 {
diff --git a/Engine/Source/Programs/Horde/Horde.Agent/Utility/ArtifactUploader.cs b/Engine/Source/Programs/Horde/Horde.Agent/Utility/ArtifactUploader.cs
--- a/Engine/Source/Programs/Horde/Horde.Agent/Utility/ArtifactUploader.cs
+++ b/Engine/Source/Programs/Horde/Horde.Agent/Utility/ArtifactUploader.cs
@@ -2,16 +2,11 @@
 
 using EpicGames.Core;
 using Grpc.Core;
-using HordeCommon;
 using HordeCommon.Rpc;
 using Microsoft.Extensions.Logging;
 using System;
 using System.Collections.Generic;
 using System.IO;
-using System.Net.Http;
-using System.Net.Http.Headers;
-using System.Runtime.Serialization;
-using System.Text;
 using System.Threading;
 using System.Threading.Tasks;
 
diff --git a/Engine/Source/Programs/Horde/Horde.Agent/Utility/CertificateHelper.cs b/Engine/Source/Programs/Horde/Horde.Agent/Utility/CertificateHelper.cs
--- a/Engine/Source/Programs/Horde/Horde.Agent/Utility/CertificateHelper.cs
+++ b/Engine/Source/Programs/Horde/Horde.Agent/Utility/CertificateHelper.cs
@@ -2,8 +2,6 @@
 
 using Microsoft.Extensions.Logging;
 using System;
-using System.Collections.Generic;
-using System.Linq;
 using System.Net.Http;
 using System.Net.Security;
 using System.Security.Cryptography.X509Certificates;
diff --git a/Engine/Source/Programs/Horde/Horde.Agent/Utility/JsonLogger.cs b/Engine/Source/Programs/Horde/Horde.Agent/Utility/JsonLogger.cs
--- a/Engine/Source/Programs/Horde/Horde.Agent/Utility/JsonLogger.cs
+++ b/Engine/Source/Programs/Horde/Horde.Agent/Utility/JsonLogger.cs
@@ -1,16 +1,8 @@
 // Copyright Epic Games, Inc. All Rights Reserved.
 
 using EpicGames.Core;
-using Horde.Agent.Parser;
-using HordeCommon;
 using Microsoft.Extensions.Logging;
 using System;
-using System.Buffers;
-using System.Collections.Generic;
-using System.Globalization;
-using System.Linq;
-using System.Text;
-using System.Text.Json;
 
 namespace Horde.Agent.Utility
 {
diff --git a/Engine/Source/Programs/Horde/Horde.Agent/Utility/JsonRpcLogger.cs b/Engine/Source/Programs/Horde/Horde.Agent/Utility/JsonRpcLogger.cs
--- a/Engine/Source/Programs/Horde/Horde.Agent/Utility/JsonRpcLogger.cs
+++ b/Engine/Source/Programs/Horde/Horde.Agent/Utility/JsonRpcLogger.cs
@@ -1,17 +1,12 @@
 // Copyright Epic Games, Inc. All Rights Reserved.
 
-using Horde.Agent.Parser.Interfaces;
 using Horde.Agent.Utility;
 using HordeCommon;
 using HordeCommon.Rpc;
 using Microsoft.Extensions.Logging;
 using System;
 using System.Buffers;
-using System.Collections.Concurrent;
 using System.Collections.Generic;
-using System.Globalization;
-using System.Linq;
-using System.Net.Http;
 using System.Text;
 using System.Text.Json;
 using System.Threading;
diff --git a/Engine/Source/Programs/Horde/Horde.Agent/Utility/NetworkShare.cs b/Engine/Source/Programs/Horde/Horde.Agent/Utility/NetworkShare.cs
--- a/Engine/Source/Programs/Horde/Horde.Agent/Utility/NetworkShare.cs
+++ b/Engine/Source/Programs/Horde/Horde.Agent/Utility/NetworkShare.cs
@@ -1,11 +1,9 @@
 // Copyright Epic Games, Inc. All Rights Reserved.
 
 using System;
-using System.Collections.Generic;
 using System.ComponentModel;
 using System.Diagnostics.CodeAnalysis;
 using System.Runtime.InteropServices;
-using System.Text;
 
 namespace Horde.Agent.Utility
 {
diff --git a/Engine/Source/Programs/Horde/Horde.Agent/Utility/RpcConnection.cs b/Engine/Source/Programs/Horde/Horde.Agent/Utility/RpcConnection.cs
--- a/Engine/Source/Programs/Horde/Horde.Agent/Utility/RpcConnection.cs
+++ b/Engine/Source/Programs/Horde/Horde.Agent/Utility/RpcConnection.cs
@@ -1,25 +1,17 @@
 // Copyright Epic Games, Inc. All Rights Reserved.
 
-using EpicGames.Core;
 using Grpc.Core;
-using Grpc.Core.Interceptors;
 using Grpc.Net.Client;
 using HordeCommon.Rpc;
-using Microsoft.Extensions.Hosting;
 using Microsoft.Extensions.Logging;
 using OpenTracing;
 using OpenTracing.Util;
 using System;
-using System.Collections.Concurrent;
 using System.Collections.Generic;
 using System.Diagnostics;
-using System.IO;
 using System.Linq;
 using System.Net;
-using System.Net.Http;
-using System.Net.Http.Headers;
 using System.Runtime.CompilerServices;
-using System.Runtime.InteropServices;
 using System.Threading;
 using System.Threading.Tasks;
 
diff --git a/Engine/Source/Programs/Horde/Horde.Agent/Utility/Shutdown.cs b/Engine/Source/Programs/Horde/Horde.Agent/Utility/Shutdown.cs
--- a/Engine/Source/Programs/Horde/Horde.Agent/Utility/Shutdown.cs
+++ b/Engine/Source/Programs/Horde/Horde.Agent/Utility/Shutdown.cs
@@ -1,14 +1,10 @@
 // Copyright Epic Games, Inc. All Rights Reserved.
 
 using Microsoft.Extensions.Logging;
-using Microsoft.Win32.SafeHandles;
 using System;
-using System.Collections.Generic;
 using System.Diagnostics;
 using System.Diagnostics.CodeAnalysis;
-using System.Net.NetworkInformation;
 using System.Runtime.InteropServices;
-using System.Text;
 
 namespace Horde.Agent.Utility
 {
diff --git a/Engine/Source/Programs/Horde/Horde.Agent/Utility/WindowsServiceManager.cs b/Engine/Source/Programs/Horde/Horde.Agent/Utility/WindowsServiceManager.cs
--- a/Engine/Source/Programs/Horde/Horde.Agent/Utility/WindowsServiceManager.cs
+++ b/Engine/Source/Programs/Horde/Horde.Agent/Utility/WindowsServiceManager.cs
@@ -1,15 +1,11 @@
 // Copyright Epic Games, Inc. All Rights Reserved.
 
 using Microsoft.Win32.SafeHandles;
-using Serilog;
 using System;
-using System.Collections.Generic;
 using System.ComponentModel;
 using System.Diagnostics;
 using System.Diagnostics.CodeAnalysis;
-using System.Reflection;
 using System.Runtime.InteropServices;
-using System.Text;
 using System.Threading;
 
 namespace Horde.Agent.Utility
diff --git a/Engine/Source/Programs/Horde/Horde.Agent/Utility/WorkspaceInfo.cs b/Engine/Source/Programs/Horde/Horde.Agent/Utility/WorkspaceInfo.cs
--- a/Engine/Source/Programs/Horde/Horde.Agent/Utility/WorkspaceInfo.cs
+++ b/Engine/Source/Programs/Horde/Horde.Agent/Utility/WorkspaceInfo.cs
@@ -5,10 +5,6 @@ using HordeCommon.Rpc.Messages;
 using Microsoft.Extensions.Logging;
 using System;
 using System.Collections.Generic;
-using System.Diagnostics;
-using System.IO;
-using System.Runtime.InteropServices;
-using System.Text;
 using System.Text.RegularExpressions;
 using System.Threading;
 using System.Threading.Tasks;
