Horde Storage - Updated logging for 0 byte files, trying to track down what causes this issue.

diff --git a/Engine/Source/Programs/Horde/HordeStorage/Horde.Storage/Implementation/Blob/FileSystemStore.cs b/Engine/Source/Programs/Horde/HordeStorage/Horde.Storage/Implementation/Blob/FileSystemStore.cs
--- a/Engine/Source/Programs/Horde/HordeStorage/Horde.Storage/Implementation/Blob/FileSystemStore.cs
+++ b/Engine/Source/Programs/Horde/HordeStorage/Horde.Storage/Implementation/Blob/FileSystemStore.cs
@@ -65,7 +65,7 @@ namespace Horde.Storage.Implementation
             }
 
             if (content.Length == 0)
-                _logger.Warning("0 byte file written as {Id} {Namespace}", blobIdentifier, ns);
+                _logger.Warning("0 byte file written as {Blob} {Namespace} {Method}", blobIdentifier, ns, "ReadOnlyMemory");
 
             UpdateLastWriteTime(filePath.FullName, DateTime.UnixEpoch);
             return blobIdentifier;
@@ -83,8 +83,9 @@ namespace Horde.Storage.Implementation
                 await content.CopyToAsync(fs, cancellationToken);
                 await fs.FlushAsync(cancellationToken);
             }
+            filePath.Refresh();
             if (filePath.Length == 0)
-                _logger.Warning("0 byte file written as {Id} {Namespace}", blobIdentifier, ns);
+                _logger.Warning("0 byte file written as {Blob} {Namespace} {Method}", blobIdentifier, ns, "Stream");
 
             UpdateLastWriteTime(filePath.FullName, DateTime.UnixEpoch);
             return blobIdentifier;
@@ -103,7 +104,7 @@ namespace Horde.Storage.Implementation
             }
 
             if (content.Length == 0)
-                _logger.Warning("0 byte file written as {Id} {Namespace}", blobIdentifier, ns);
+                _logger.Warning("0 byte file written as {Blob} {Namespace} {Method}", blobIdentifier, ns, "Array");
 
             UpdateLastWriteTime(filePath.FullName, DateTime.UnixEpoch);
             return blobIdentifier;
