Horde Storage - Added progress logging to the blob index consistency check

diff --git a/Engine/Source/Programs/Horde/HordeStorage/Horde.Storage/Implementation/Blob/BlobIndexConsistencyCheckService.cs b/Engine/Source/Programs/Horde/HordeStorage/Horde.Storage/Implementation/Blob/BlobIndexConsistencyCheckService.cs
--- a/Engine/Source/Programs/Horde/HordeStorage/Horde.Storage/Implementation/Blob/BlobIndexConsistencyCheckService.cs
+++ b/Engine/Source/Programs/Horde/HordeStorage/Horde.Storage/Implementation/Blob/BlobIndexConsistencyCheckService.cs
@@ -72,6 +72,9 @@ namespace Horde.Storage.Implementation
             {
                 Interlocked.Increment(ref countOfBlobsChecked);
 
+                if (countOfBlobsChecked % 100 == 0)
+                    _logger.Information("Consistency check running on blob index, count of blobs processed so far: {CountOfBlobs}", countOfBlobsChecked);
+
                 if (!blobInfo.Regions.Any())
                 {
                     Interlocked.Increment(ref countOfIncorrectBlobsFound);
@@ -87,11 +90,13 @@ namespace Horde.Storage.Implementation
                     }
                     else
                     {
-                        // this blob doesn't exist anywhere so we just cleanup the blob index
-                        _logger.Warning("Blob {Blob} in namespace {Namespace} was removed from the blob index as it didnt exist anywhere", blobInfo.BlobIdentifier, blobInfo.Namespace);
-
                         if (_settings.CurrentValue.AllowDeletesInBlobIndex)
+                        {
+                            // this blob doesn't exist anywhere so we just cleanup the blob index
+                            _logger.Warning("Blob {Blob} in namespace {Namespace} was removed from the blob index as it didnt exist anywhere", blobInfo.BlobIdentifier, blobInfo.Namespace);
+
                             await _blobIndex.RemoveBlobFromIndex(blobInfo.Namespace, blobInfo.BlobIdentifier);
+                        }
                     }
                 }
                 
