Horde Storage - Catch and log exceptions during blob index consistency checks, as we want to continue running every blob no matter what happens.

diff --git a/Engine/Source/Programs/Horde/HordeStorage/Horde.Storage/Implementation/Blob/BlobIndexConsistencyCheckService.cs b/Engine/Source/Programs/Horde/HordeStorage/Horde.Storage/Implementation/Blob/BlobIndexConsistencyCheckService.cs
--- a/Engine/Source/Programs/Horde/HordeStorage/Horde.Storage/Implementation/Blob/BlobIndexConsistencyCheckService.cs
+++ b/Engine/Source/Programs/Horde/HordeStorage/Horde.Storage/Implementation/Blob/BlobIndexConsistencyCheckService.cs
@@ -75,51 +75,58 @@ namespace Horde.Storage.Implementation
                 if (countOfBlobsChecked % 100 == 0)
                     _logger.Information("Consistency check running on blob index, count of blobs processed so far: {CountOfBlobs}", countOfBlobsChecked);
 
-                if (!blobInfo.Regions.Any())
+                try
                 {
-                    Interlocked.Increment(ref countOfIncorrectBlobsFound);
-                    _logger.Warning("Blob {Blob} in namespace {Namespace} is not tracked to exist in any region", blobInfo.BlobIdentifier, blobInfo.Namespace);
-
-                    if (await _blobService.ExistsInRootStore(blobInfo.Namespace, blobInfo.BlobIdentifier))
+                    if (!blobInfo.Regions.Any())
                     {
-                        _logger.Warning("Blob {Blob} in namespace {Namespace} was found to exist in our blob store so re-adding it to the index", blobInfo.BlobIdentifier, blobInfo.Namespace);
+                        Interlocked.Increment(ref countOfIncorrectBlobsFound);
+                        _logger.Warning("Blob {Blob} in namespace {Namespace} is not tracked to exist in any region", blobInfo.BlobIdentifier, blobInfo.Namespace);
 
-                        // we did have it in our blob store so we adjust the index
-                        await _blobIndex.AddBlobToIndex(blobInfo.Namespace, blobInfo.BlobIdentifier);
-                        continue;
-                    }
-                    else
-                    {
-                        if (_settings.CurrentValue.AllowDeletesInBlobIndex)
+                        if (await _blobService.ExistsInRootStore(blobInfo.Namespace, blobInfo.BlobIdentifier))
                         {
-                            // this blob doesn't exist anywhere so we just cleanup the blob index
-                            _logger.Warning("Blob {Blob} in namespace {Namespace} was removed from the blob index as it didnt exist anywhere", blobInfo.BlobIdentifier, blobInfo.Namespace);
+                            _logger.Warning("Blob {Blob} in namespace {Namespace} was found to exist in our blob store so re-adding it to the index", blobInfo.BlobIdentifier, blobInfo.Namespace);
 
-                            await _blobIndex.RemoveBlobFromIndex(blobInfo.Namespace, blobInfo.BlobIdentifier);
+                            // we did have it in our blob store so we adjust the index
+                            await _blobIndex.AddBlobToIndex(blobInfo.Namespace, blobInfo.BlobIdentifier);
+                            continue;
                         }
-                    }
-                }
-                
-                if (blobInfo.Regions.Contains(currentRegion))
-                {
-                    if (!await _blobService.ExistsInRootStore(blobInfo.Namespace, blobInfo.BlobIdentifier))
-                    {
-                        Interlocked.Increment(ref countOfIncorrectBlobsFound);
-                        _logger.Warning("Blob {Blob} in namespace {Namespace} did not exist in root store but is tracked as doing so in the blob index. Attempting to replicate it.", blobInfo.BlobIdentifier, blobInfo.Namespace);
-                        try
+                        else
                         {
-                            BlobContents _ = await _blobService.ReplicateObject(blobInfo.Namespace, blobInfo.BlobIdentifier, force: true);
+                            if (_settings.CurrentValue.AllowDeletesInBlobIndex)
+                            {
+                                // this blob doesn't exist anywhere so we just cleanup the blob index
+                                _logger.Warning("Blob {Blob} in namespace {Namespace} was removed from the blob index as it didnt exist anywhere", blobInfo.BlobIdentifier, blobInfo.Namespace);
+
+                                await _blobIndex.RemoveBlobFromIndex(blobInfo.Namespace, blobInfo.BlobIdentifier);
+                            }
                         }
-                        catch (BlobReplicationException e)
+                    }
+                    
+                    if (blobInfo.Regions.Contains(currentRegion))
+                    {
+                        if (!await _blobService.ExistsInRootStore(blobInfo.Namespace, blobInfo.BlobIdentifier))
                         {
-                            _logger.Error(e, "Failed to replicate Blob {Blob} in namespace {Namespace}. Unable to repair the blob index", blobInfo.BlobIdentifier, blobInfo.Namespace);
-
-                            // we update the blob index to accurately reflect that we do not have the blob, this is not good though as it means a upload that we thought happened now lacks content
-                            if (_settings.CurrentValue.AllowDeletesInBlobIndex)
-                                await _blobIndex.RemoveBlobFromRegion(blobInfo.Namespace, blobInfo.BlobIdentifier);
+                            Interlocked.Increment(ref countOfIncorrectBlobsFound);
+                            _logger.Warning("Blob {Blob} in namespace {Namespace} did not exist in root store but is tracked as doing so in the blob index. Attempting to replicate it.", blobInfo.BlobIdentifier, blobInfo.Namespace);
+                            try
+                            {
+                                BlobContents _ = await _blobService.ReplicateObject(blobInfo.Namespace, blobInfo.BlobIdentifier, force: true);
+                            }
+                            catch (BlobReplicationException e)
+                            {
+                                _logger.Error(e, "Failed to replicate Blob {Blob} in namespace {Namespace}. Unable to repair the blob index", blobInfo.BlobIdentifier, blobInfo.Namespace);
+
+                                // we update the blob index to accurately reflect that we do not have the blob, this is not good though as it means a upload that we thought happened now lacks content
+                                if (_settings.CurrentValue.AllowDeletesInBlobIndex)
+                                    await _blobIndex.RemoveBlobFromRegion(blobInfo.Namespace, blobInfo.BlobIdentifier);
+                            }
                         }
                     }
-
+                }
+                catch (Exception e)
+                {
+                    // TODO: We should likely catch some common exceptions here, like UnknownNamespaceException as there are likely quite a few cases were this isn't even an error
+                    _logger.Error(e, "Exception when doing blob index consistency check for {Blob} in namespace {Namespace}", blobInfo.BlobIdentifier, blobInfo.Namespace);
                 }
             }
 
