Fix issue when selecting history item does not load the correspongding result with single view mode
Fix meta display in single view mode

#rnx

diff --git a/Engine/Source/Programs/Horde/HordeDashboard/src/Plugins/AutomatedTestSession/models/AutomatedTestSession.ts b/Engine/Source/Programs/Horde/HordeDashboard/src/Plugins/AutomatedTestSession/models/AutomatedTestSession.ts
--- a/Engine/Source/Programs/Horde/HordeDashboard/src/Plugins/AutomatedTestSession/models/AutomatedTestSession.ts
+++ b/Engine/Source/Programs/Horde/HordeDashboard/src/Plugins/AutomatedTestSession/models/AutomatedTestSession.ts
@@ -927,7 +927,7 @@ export class TestSessionCollection {
                 await when(() => testcase.version !== previousVersion || this.loadingCount === 0);
                 test = testcase.Metas.get(metauid);
             }
-            if(test && change && test.Session.Testdata.change !== change) {
+            if(change && test?.Session.Testdata.change !== change) {
                 // Need to get the target change from the history
                 test = this.getLoadedTestHistoryByMeta(uid)?.get(metauid)?.find(
                     (item) => item.Session.Testdata.change === change
