Horde: Fix accounting for active conform leases.

- If a conform lease cannot be allocated due to the maximum number of conform leases already being allocated, and a conform was requested manually, we MUST drain the current lease list rather than allowing the agent to service other work.
- If a lease is cancelled due to another lease being assigned to an agent, the lease was not being removed from the conform list.
- Add a missing entry in the agent log to note a requested full conform.

diff --git a/Engine/Source/Programs/Horde/Horde.Build/Agents/Impl/Tasks/ConformTaskSource.cs b/Engine/Source/Programs/Horde/Horde.Build/Agents/Impl/Tasks/ConformTaskSource.cs
--- a/Engine/Source/Programs/Horde/Horde.Build/Agents/Impl/Tasks/ConformTaskSource.cs
+++ b/Engine/Source/Programs/Horde/Horde.Build/Agents/Impl/Tasks/ConformTaskSource.cs
@@ -165,18 +165,16 @@ namespace Horde.Build.Tasks.Impl
 				if (await GetWorkspacesAsync(Agent, Task.Workspaces))
 				{
 					LeaseId LeaseId = LeaseId.GenerateNewId();
-					if (!await AllocateConformLeaseAsync(Agent.Id, Task.Workspaces, LeaseId))
+					if (await AllocateConformLeaseAsync(Agent.Id, Task.Workspaces, LeaseId))
 					{
-						return null;
-					}
-
-					ILogFile Log = await LogService.CreateLogFileAsync(JobId.Empty, Agent.SessionId, LogType.Json);
-					Task.LogId = Log.Id.ToString();
-					Task.RemoveUntrackedFiles = Agent.RequestFullConform;
+						ILogFile Log = await LogService.CreateLogFileAsync(JobId.Empty, Agent.SessionId, LogType.Json);
+						Task.LogId = Log.Id.ToString();
+						Task.RemoveUntrackedFiles = Agent.RequestFullConform;
 
-					byte[] Payload = Any.Pack(Task).ToByteArray();
+						byte[] Payload = Any.Pack(Task).ToByteArray();
 
-					return new AgentLease(LeaseId, "Updating workspaces", null, null, Log.Id, LeaseState.Pending, null, true, Payload);
+						return new AgentLease(LeaseId, "Updating workspaces", null, null, Log.Id, LeaseState.Pending, null, true, Payload);
+					}
 				}
 			}
 
@@ -188,6 +186,12 @@ namespace Horde.Build.Tasks.Impl
 			return null;
 		}
 
+		/// <inheritdoc/>
+		public override Task CancelLeaseAsync(IAgent Agent, LeaseId LeaseId, ConformTask Payload)
+		{
+			return ReleaseConformLeaseAsync(LeaseId);
+		}
+
 		/// <inheritdoc/>
 		public async Task<bool> GetWorkspacesAsync(IAgent Agent, IList<HordeCommon.Rpc.Messages.AgentWorkspace> Workspaces)
 		{
diff --git a/Engine/Source/Programs/Horde/Horde.Build/Controllers/AgentsController.cs b/Engine/Source/Programs/Horde/Horde.Build/Controllers/AgentsController.cs
--- a/Engine/Source/Programs/Horde/Horde.Build/Controllers/AgentsController.cs
+++ b/Engine/Source/Programs/Horde/Horde.Build/Controllers/AgentsController.cs
@@ -1,4 +1,4 @@
-﻿// Copyright Epic Games, Inc. All Rights Reserved.
+// Copyright Epic Games, Inc. All Rights Reserved.
 
 using Horde.Build.Acls;
 using Horde.Build.Api;
@@ -192,7 +192,11 @@ namespace Horde.Build.Controllers
 				}
 				if (Agent.RequestConform != NewAgent.RequestConform)
 				{
-					Logger.LogInformation("Setting changed: RequestConfig = {State} ({UserName})", NewAgent.RequestConform, UserName);
+					Logger.LogInformation("Setting changed: RequestConform = {State} ({UserName})", NewAgent.RequestConform, UserName);
+				}
+				if (Agent.RequestFullConform != NewAgent.RequestFullConform)
+				{
+					Logger.LogInformation("Setting changed: RequestFullConform = {State} ({UserName})", NewAgent.RequestFullConform, UserName);
 				}
 				if (Agent.RequestRestart != NewAgent.RequestRestart)
 				{
