Horde: Include a "pendingFullConform" flag in agent responses. Will need to be hooked up to status column on dashboard similarly to "pendingConform".

diff --git a/Engine/Source/Programs/Horde/Horde.Build/Api/Agents.cs b/Engine/Source/Programs/Horde/Horde.Build/Api/Agents.cs
--- a/Engine/Source/Programs/Horde/Horde.Build/Api/Agents.cs
+++ b/Engine/Source/Programs/Horde/Horde.Build/Api/Agents.cs
@@ -366,6 +366,11 @@ namespace Horde.Build.Api
 		/// </summary>
 		public bool PendingConform { get; set; }
 
+		/// <summary>
+		/// Whether a full conform job is pending
+		/// </summary>
+		public bool PendingFullConform { get; set; }
+
 		/// <summary>
 		/// Whether a restart is pending
 		/// </summary>
@@ -470,6 +475,7 @@ namespace Horde.Build.Api
 			this.Online = Agent.IsSessionValid(DateTime.UtcNow);
 			this.Deleted = Agent.Deleted;
 			this.PendingConform = Agent.RequestConform;
+			this.PendingFullConform = Agent.RequestFullConform;
 			this.PendingRestart = Agent.RequestRestart;
 			this.PendingShutdown = Agent.RequestShutdown;
 			this.LastShutdownReason = Agent.LastShutdownReason ?? "Unknown";
