Horde: Add support for pending full conforms
#rnx
#skipci

diff --git a/Engine/Source/Programs/Horde/HordeDashboard/src/backend/Api.ts b/Engine/Source/Programs/Horde/HordeDashboard/src/backend/Api.ts
--- a/Engine/Source/Programs/Horde/HordeDashboard/src/backend/Api.ts
+++ b/Engine/Source/Programs/Horde/HordeDashboard/src/backend/Api.ts
@@ -662,6 +662,9 @@ export type GetAgentResponse = {
 	/** Whether agent is pending conform */
 	pendingConform: boolean;
 
+	/** Whether a full conform job is pending */
+	pendingFullConform: boolean;
+
 	/** Conform attempt count */
 	conformAttemptCount?: number;
 
