Horde: Remove unused parameter.

diff --git a/Engine/Source/Programs/Horde/Horde.Agent/Services/WorkerService.cs b/Engine/Source/Programs/Horde/Horde.Agent/Services/WorkerService.cs
--- a/Engine/Source/Programs/Horde/Horde.Agent/Services/WorkerService.cs
+++ b/Engine/Source/Programs/Horde/Horde.Agent/Services/WorkerService.cs
@@ -773,7 +773,7 @@ namespace Horde.Agent.Services
 			if (leaseInfo.Lease.Payload.TryUnpack(out jobTask))
 			{
 				GlobalTracer.Instance.ActiveSpan?.SetTag("task", "Job");
-				Task<LeaseResult> Handler(ILogger newLogger) => ExecuteJobAsync(rpcConnection, agentId, leaseInfo.Lease.Id, jobTask, newLogger, leaseInfo.CancellationTokenSource.Token);
+				Task<LeaseResult> Handler(ILogger newLogger) => ExecuteJobAsync(rpcConnection, leaseInfo.Lease.Id, jobTask, newLogger, leaseInfo.CancellationTokenSource.Token);
 				return await HandleLeasePayloadWithLogAsync(rpcConnection, jobTask.LogId, jobTask.JobId, jobTask.BatchId, Handler, leaseInfo.CancellationTokenSource.Token);
 			}
 
@@ -1000,13 +1000,12 @@ namespace Horde.Agent.Services
 		/// Execute part of a job
 		/// </summary>
 		/// <param name="rpcClient">RPC client for communicating with the server</param>
-		/// <param name="agentId">The current agent id</param>
 		/// <param name="leaseId">The current lease id</param>
 		/// <param name="executeTask">The task to execute</param>
 		/// <param name="logger">The logger to use for this lease</param>
 		/// <param name="cancellationToken">Token used to cancel the operation</param>
 		/// <returns>Async task</returns>
-		internal async Task<LeaseResult> ExecuteJobAsync(IRpcConnection rpcClient, string agentId, string leaseId, ExecuteJobTask executeTask, ILogger logger, CancellationToken cancellationToken)
+		internal async Task<LeaseResult> ExecuteJobAsync(IRpcConnection rpcClient, string leaseId, ExecuteJobTask executeTask, ILogger logger, CancellationToken cancellationToken)
 		{
 			logger.LogInformation("Executing job \"{JobName}\", jobId {JobId}, batchId {BatchId}, leaseId {LeaseId}", executeTask.JobName, executeTask.JobId, executeTask.BatchId, leaseId);
 			GlobalTracer.Instance.ActiveSpan?.SetTag("jobId", executeTask.JobId.ToString());
