Horde Storage - Improved log message when doing on demand replication

diff --git a/Engine/Source/Programs/Horde/HordeStorage/Horde.Storage/Implementation/Blob/IBlobService.cs b/Engine/Source/Programs/Horde/HordeStorage/Horde.Storage/Implementation/Blob/IBlobService.cs
--- a/Engine/Source/Programs/Horde/HordeStorage/Horde.Storage/Implementation/Blob/IBlobService.cs
+++ b/Engine/Source/Programs/Horde/HordeStorage/Horde.Storage/Implementation/Blob/IBlobService.cs
@@ -270,7 +270,7 @@ public class BlobService : IBlobService
         if (!blobInfo.Regions.Any())
             throw new BlobReplicationException(ns, blob, "Blob not found in any region");
 
-        _logger.Information("On-demand replicating blob {Blob} in Namespace {Ns}", blob, ns);
+        _logger.Information("On-demand replicating blob {Blob} in Namespace {Namespace}", blob, ns);
         SortedList<int, string> possiblePeers = new SortedList<int, string>(_peerStatusService.GetPeersByLatency(blobInfo.Regions.ToList()));
 
         bool replicated = false;
