Horde Storage - Improvements to the blob index consistency checks, improving some error messages, adding code to handle some scenarios and generally avoiding some spaminess in the log output

diff --git a/Engine/Source/Programs/Horde/HordeStorage/Horde.Storage/Implementation/Blob/BlobIndexConsistencyCheckService.cs b/Engine/Source/Programs/Horde/HordeStorage/Horde.Storage/Implementation/Blob/BlobIndexConsistencyCheckService.cs
--- a/Engine/Source/Programs/Horde/HordeStorage/Horde.Storage/Implementation/Blob/BlobIndexConsistencyCheckService.cs
+++ b/Engine/Source/Programs/Horde/HordeStorage/Horde.Storage/Implementation/Blob/BlobIndexConsistencyCheckService.cs
@@ -107,25 +107,52 @@ namespace Horde.Storage.Implementation
                         if (!await _blobService.ExistsInRootStore(blobInfo.Namespace, blobInfo.BlobIdentifier))
                         {
                             Interlocked.Increment(ref countOfIncorrectBlobsFound);
-                            _logger.Warning("Blob {Blob} in namespace {Namespace} did not exist in root store but is tracked as doing so in the blob index. Attempting to replicate it.", blobInfo.BlobIdentifier, blobInfo.Namespace);
-                            try
+                            
+                            if (blobInfo.Regions.Count > 1)
                             {
-                                BlobContents _ = await _blobService.ReplicateObject(blobInfo.Namespace, blobInfo.BlobIdentifier, force: true);
+                                _logger.Warning("Blob {Blob} in namespace {Namespace} did not exist in root store but is tracked as doing so in the blob index. Attempting to replicate it.", blobInfo.BlobIdentifier, blobInfo.Namespace);
+
+                                try
+                                {
+                                    BlobContents _ = await _blobService.ReplicateObject(blobInfo.Namespace, blobInfo.BlobIdentifier, force: true);
+                                }
+                                catch (BlobReplicationException e)
+                                {
+                                    _logger.Error(e, "Failed to replicate Blob {Blob} in namespace {Namespace}. Unable to repair the blob index", blobInfo.BlobIdentifier, blobInfo.Namespace);
+
+                                    // we update the blob index to accurately reflect that we do not have the blob, this is not good though as it means a upload that we thought happened now lacks content
+                                    if (_settings.CurrentValue.AllowDeletesInBlobIndex)
+                                        await _blobIndex.RemoveBlobFromRegion(blobInfo.Namespace, blobInfo.BlobIdentifier);
+                                }
                             }
-                            catch (BlobReplicationException e)
+                            else
                             {
-                                _logger.Error(e, "Failed to replicate Blob {Blob} in namespace {Namespace}. Unable to repair the blob index", blobInfo.BlobIdentifier, blobInfo.Namespace);
+                                // if the blob only exists in the current region there is no point in attempting to replicate it
+                                _logger.Warning("Blob {Blob} in namespace {Namespace} did not exist in root store but is tracked as doing so in the blob index. Does not exist anywhere else so unable to replicate it.", blobInfo.BlobIdentifier, blobInfo.Namespace);
 
-                                // we update the blob index to accurately reflect that we do not have the blob, this is not good though as it means a upload that we thought happened now lacks content
                                 if (_settings.CurrentValue.AllowDeletesInBlobIndex)
+                                {
+                                    // this blob can not be repaired so we just delete it from the blob index
+                                    _logger.Warning("Blob {Blob} in namespace {Namespace} can not be repaired so removing existence from current region.", blobInfo.BlobIdentifier, blobInfo.Namespace);
+
                                     await _blobIndex.RemoveBlobFromRegion(blobInfo.Namespace, blobInfo.BlobIdentifier);
+                                }
                             }
                         }
                     }
                 }
+                catch (UnknownNamespaceException)
+                {
+                    if (_settings.CurrentValue.AllowDeletesInBlobIndex)
+                    {
+                        _logger.Warning("Blob {Blob} in namespace {Namespace} is of a unknown namespace, removing.", blobInfo.BlobIdentifier, blobInfo.Namespace);
+
+                        // for entries that are of a unknown namespace we simply remove them
+                        await _blobIndex.RemoveBlobFromIndex(blobInfo.Namespace, blobInfo.BlobIdentifier);
+                    }
+                }
                 catch (Exception e)
                 {
-                    // TODO: We should likely catch some common exceptions here, like UnknownNamespaceException as there are likely quite a few cases were this isn't even an error
                     _logger.Error(e, "Exception when doing blob index consistency check for {Blob} in namespace {Namespace}", blobInfo.BlobIdentifier, blobInfo.Namespace);
                 }
             }
