Horde Storage - Improved logging when on demand replicating blobs to see which regions are being considered.

diff --git a/Engine/Source/Programs/Horde/HordeStorage/Horde.Storage/Implementation/Blob/IBlobService.cs b/Engine/Source/Programs/Horde/HordeStorage/Horde.Storage/Implementation/Blob/IBlobService.cs
--- a/Engine/Source/Programs/Horde/HordeStorage/Horde.Storage/Implementation/Blob/IBlobService.cs
+++ b/Engine/Source/Programs/Horde/HordeStorage/Horde.Storage/Implementation/Blob/IBlobService.cs
@@ -279,6 +279,9 @@ public class BlobService : IBlobService
             IPeerStatusService.PeerStatus? peerStatus = _peerStatusService.GetPeerStatus(region);
             if (peerStatus == null)
                 throw new Exception($"Failed to find peer {region}");
+
+            _logger.Information("Attempting to replicate blob {Blob} in Namespace {Namespace} from {Region}", blob, ns, region);
+
             PeerEndpoints peerEndpoint = peerStatus.Endpoints.First();
             HttpClient httpClient = _httpClientFactory.CreateClient();
             HttpRequestMessage blobRequest = BuildHttpRequest(HttpMethod.Get, $"{peerEndpoint.Url}/api/v1/blobs/{ns}/{blob}");
