Horde Storage - Further polish to the logging for the blob index consistency service

diff --git a/Engine/Source/Programs/Horde/HordeStorage/Horde.Storage/Implementation/Blob/BlobIndexConsistencyCheckService.cs b/Engine/Source/Programs/Horde/HordeStorage/Horde.Storage/Implementation/Blob/BlobIndexConsistencyCheckService.cs
--- a/Engine/Source/Programs/Horde/HordeStorage/Horde.Storage/Implementation/Blob/BlobIndexConsistencyCheckService.cs
+++ b/Engine/Source/Programs/Horde/HordeStorage/Horde.Storage/Implementation/Blob/BlobIndexConsistencyCheckService.cs
@@ -122,7 +122,10 @@ namespace Horde.Storage.Implementation
 
                                     // we update the blob index to accurately reflect that we do not have the blob, this is not good though as it means a upload that we thought happened now lacks content
                                     if (_settings.CurrentValue.AllowDeletesInBlobIndex)
+                                    {
+                                        _logger.Warning("Updating blob index to remove Blob {Blob} in namespace {Namespace} as we failed to repair it.", blobInfo.BlobIdentifier, blobInfo.Namespace);
                                         await _blobIndex.RemoveBlobFromRegion(blobInfo.Namespace, blobInfo.BlobIdentifier);
+                                    }
                                 }
                             }
                             else
