Horde: Build fix.

diff --git a/Engine/Source/Programs/Horde/Horde.Agent.Tests/WorkerServiceTest.cs b/Engine/Source/Programs/Horde/Horde.Agent.Tests/WorkerServiceTest.cs
--- a/Engine/Source/Programs/Horde/Horde.Agent.Tests/WorkerServiceTest.cs
+++ b/Engine/Source/Programs/Horde/Horde.Agent.Tests/WorkerServiceTest.cs
@@ -116,7 +116,7 @@ namespace Horde.Agent.Tests
 
 			WorkerService ws = GetWorkerService((a, b, c) => executor);
 			ws._stepAbortPollInterval = TimeSpan.FromMilliseconds(1);
-			LeaseOutcome outcome = (await ws.ExecuteJobAsync(rpcConnection, "agentId1", "leaseId1", executeJobTask,
+			LeaseOutcome outcome = (await ws.ExecuteJobAsync(rpcConnection, "leaseId1", executeJobTask,
 				_workerLogger, token)).Outcome;
 
 			Assert.AreEqual(LeaseOutcome.Success, outcome);
