Horde: Update Horde.Build.Tests to use standard coding conventions, and fix a number of issues caught by static analysis.

diff --git a/Engine/Source/Programs/Horde/Horde.Build.Tests/AgentControllerTest.cs b/Engine/Source/Programs/Horde/Horde.Build.Tests/AgentControllerTest.cs
--- a/Engine/Source/Programs/Horde/Horde.Build.Tests/AgentControllerTest.cs
+++ b/Engine/Source/Programs/Horde/Horde.Build.Tests/AgentControllerTest.cs
@@ -2,34 +2,32 @@
 
 using System.Threading.Tasks;
 using Horde.Build.Api;
-using Microsoft.VisualStudio.TestTools.UnitTesting;
-using Horde.Build.Controllers;
 using Horde.Build.Models;
 using Microsoft.AspNetCore.Mvc;
+using Microsoft.VisualStudio.TestTools.UnitTesting;
+
+namespace Horde.Build.Tests;
 
-namespace Horde.Build.Tests
+[TestClass]
+public class AgentControllerDbTest : TestSetup
 {
-	[TestClass]
-    public class AgentControllerDbTest : TestSetup
-    {
-        [TestMethod]
-        public async Task UpdateAgent()
-        {
-			Fixture Fixture = await CreateFixtureAsync();
-	        IAgent FixtureAgent = Fixture.Agent1;
+	[TestMethod]
+	public async Task UpdateAgent()
+	{
+		Fixture fixture = await CreateFixtureAsync();
+		IAgent fixtureAgent = fixture.Agent1;
+
+		ActionResult<object> obj = await AgentsController.GetAgentAsync(fixtureAgent.Id);
+		GetAgentResponse getRes = (obj.Value as GetAgentResponse)!;
+		Assert.AreEqual(fixture!.Agent1Name.ToUpper(), getRes.Name);
+		Assert.IsNull(getRes.Comment);
+
+		UpdateAgentRequest updateReq = new();
+		updateReq.Comment = "foo bar baz";
+		await AgentsController.UpdateAgentAsync(fixtureAgent.Id, updateReq);
 
-	        ActionResult<object> Obj = await AgentsController.GetAgentAsync(FixtureAgent.Id);
-	        GetAgentResponse GetRes = (Obj.Value as GetAgentResponse)!;
-	        Assert.AreEqual(Fixture!.Agent1Name.ToUpper(), GetRes.Name);
-	        Assert.IsNull(GetRes.Comment);
-	        
-	        UpdateAgentRequest UpdateReq = new UpdateAgentRequest();
-	        UpdateReq.Comment = "foo bar baz";
-	        await AgentsController.UpdateAgentAsync(FixtureAgent.Id, UpdateReq);
-	        
-	        Obj = await AgentsController.GetAgentAsync(FixtureAgent.Id);
-	        GetRes = (Obj.Value as GetAgentResponse)!;
-	        Assert.AreEqual("foo bar baz", GetRes.Comment);
-        }
-    }
+		obj = await AgentsController.GetAgentAsync(fixtureAgent.Id);
+		getRes = (obj.Value as GetAgentResponse)!;
+		Assert.AreEqual("foo bar baz", getRes.Comment);
+	}
 }
\ No newline at end of file
diff --git a/Engine/Source/Programs/Horde/Horde.Build.Tests/AgentServiceTest.cs b/Engine/Source/Programs/Horde/Horde.Build.Tests/AgentServiceTest.cs
--- a/Engine/Source/Programs/Horde/Horde.Build.Tests/AgentServiceTest.cs
+++ b/Engine/Source/Programs/Horde/Horde.Build.Tests/AgentServiceTest.cs
@@ -4,58 +4,53 @@ using System;
 using System.Collections.Generic;
 using System.Security.Claims;
 using System.Threading.Tasks;
-using HordeCommon;
 using Horde.Build.Api;
-using Microsoft.VisualStudio.TestTools.UnitTesting;
 using Horde.Build.Models;
-using Horde.Build.Services;
-using Microsoft.AspNetCore.Mvc;
-using AgentSoftwareVersion = Horde.Build.Utilities.StringId<Horde.Build.Collections.IAgentSoftwareCollection>;
 using Horde.Build.Utilities;
-using MongoDB.Bson;
+using HordeCommon;
+using Microsoft.AspNetCore.Mvc;
+using Microsoft.VisualStudio.TestTools.UnitTesting;
+
+namespace Horde.Build.Tests;
 
-namespace Horde.Build.Tests
+/// <summary>
+///     Testing the agent service
+/// </summary>
+[TestClass]
+public class AgentServiceTest : TestSetup
 {
-    /// <summary>
-    /// Testing the agent service
-    /// </summary>
-	[TestClass]
-    public class AgentServiceTest : TestSetup
-    {
-        [TestMethod]
-        public async Task GetJobs()
-        {
-			Fixture Fixture = await CreateFixtureAsync();
-
-			ActionResult<List<object>> Res = await JobsController.FindJobsAsync();
-	        Assert.AreEqual(2, Res.Value!.Count);
-	        Assert.AreEqual("hello2", (Res.Value[0] as GetJobResponse)!.Name);
-	        Assert.AreEqual("hello1", (Res.Value[1] as GetJobResponse)!.Name);
-	        
-	        Res = await JobsController.FindJobsAsync(IncludePreflight: false);
-	        Assert.AreEqual(1, Res.Value!.Count);
-	        Assert.AreEqual("hello2", (Res.Value[0] as GetJobResponse)!.Name);
-        }
-        
-        [TestMethod]
-        public async Task CreateSessionTest()
-        {
-			Fixture Fixture = await CreateFixtureAsync();
-
-	        IAgent Agent = await AgentService.CreateSessionAsync(Fixture.Agent1, AgentStatus.Ok, new List<string>(), new Dictionary<string, int>(),
-		        "test");
-	        
-	        Assert.IsTrue(AgentService.AuthorizeSession(Agent, GetUser(Agent)));
-	        await Clock.AdvanceAsync(TimeSpan.FromMinutes(20));
-	        Assert.IsFalse(AgentService.AuthorizeSession(Agent, GetUser(Agent)));
-        }
-
-        private ClaimsPrincipal GetUser(IAgent Agent)
-        {
-	        return new ClaimsPrincipal(new ClaimsIdentity(new List<Claim>
-	        {
-		        new Claim(HordeClaimTypes.AgentSessionId, Agent.SessionId.ToString()!),
-	        }, "TestAuthType"));
-        }
-    }
+	[TestMethod]
+	public async Task GetJobs()
+	{
+		Fixture fixture = await CreateFixtureAsync();
+
+		ActionResult<List<object>> res = await JobsController.FindJobsAsync();
+		Assert.AreEqual(2, res.Value!.Count);
+		Assert.AreEqual("hello2", (res.Value[0] as GetJobResponse)!.Name);
+		Assert.AreEqual("hello1", (res.Value[1] as GetJobResponse)!.Name);
+
+		res = await JobsController.FindJobsAsync(IncludePreflight: false);
+		Assert.AreEqual(1, res.Value!.Count);
+		Assert.AreEqual("hello2", (res.Value[0] as GetJobResponse)!.Name);
+	}
+
+	[TestMethod]
+	public async Task CreateSessionTest()
+	{
+		Fixture fixture = await CreateFixtureAsync();
+
+		IAgent agent = await AgentService.CreateSessionAsync(fixture.Agent1, AgentStatus.Ok, new List<string>(),
+			new Dictionary<string, int>(),
+			"test");
+
+		Assert.IsTrue(AgentService.AuthorizeSession(agent, GetUser(agent)));
+		await Clock.AdvanceAsync(TimeSpan.FromMinutes(20));
+		Assert.IsFalse(AgentService.AuthorizeSession(agent, GetUser(agent)));
+	}
+
+	private static ClaimsPrincipal GetUser(IAgent agent)
+	{
+		return new ClaimsPrincipal(new ClaimsIdentity(
+			new List<Claim> { new(HordeClaimTypes.AgentSessionId, agent.SessionId.ToString()!) }, "TestAuthType"));
+	}
 }
\ No newline at end of file
diff --git a/Engine/Source/Programs/Horde/Horde.Build.Tests/ControllerIntegrationTest.cs b/Engine/Source/Programs/Horde/Horde.Build.Tests/ControllerIntegrationTest.cs
--- a/Engine/Source/Programs/Horde/Horde.Build.Tests/ControllerIntegrationTest.cs
+++ b/Engine/Source/Programs/Horde/Horde.Build.Tests/ControllerIntegrationTest.cs
@@ -3,88 +3,90 @@
 using System;
 using System.Collections.Generic;
 using System.Net.Http;
-using System.Net.Http.Headers;
-using System.Runtime.CompilerServices;
 using System.Threading.Tasks;
-using Horde.Build;
 using Horde.Build.Collections;
 using Horde.Build.Collections.Impl;
 using Horde.Build.Jobs;
 using Horde.Build.Services;
-using Microsoft.AspNetCore.Builder;
+using Horde.Build.Tests.Stubs.Services;
 using Microsoft.AspNetCore.Hosting;
 using Microsoft.AspNetCore.Mvc.Testing;
 using Microsoft.AspNetCore.TestHost;
 using Microsoft.Extensions.Configuration;
 using Microsoft.Extensions.DependencyInjection;
-using Microsoft.Extensions.Hosting;
 
-namespace Horde.Build.Tests
+namespace Horde.Build.Tests;
+
+public class TestWebApplicationFactory<TStartup> : WebApplicationFactory<TStartup> where TStartup : class
 {
-	public class TestWebApplicationFactory<TStartup> : WebApplicationFactory<TStartup> where TStartup : class
-    {
-		MongoDbInstance MongoDbInstance;
+	private readonly MongoDbInstance _mongoDbInstance;
+
+	public TestWebApplicationFactory(MongoDbInstance mongoDbInstance)
+	{
+		_mongoDbInstance = mongoDbInstance;
+	}
 
-		public TestWebApplicationFactory(MongoDbInstance MongoDbInstance)
+	protected override void ConfigureWebHost(IWebHostBuilder builder)
+	{
+		Dictionary<string, string?> dict = new()
 		{
-			this.MongoDbInstance = MongoDbInstance;
-		}
+			{ "Horde:DatabaseConnectionString", _mongoDbInstance.ConnectionString },
+			{ "Horde:DatabaseName", _mongoDbInstance.DatabaseName },
+			{ "Horde:LogServiceWriteCacheType", "inmemory" },
+			{ "Horde:DisableAuth", "true" },
+			{ "Horde:OidcAuthority", null },
+			{ "Horde:OidcClientId", null }
+		};
+		builder.ConfigureAppConfiguration((hostingContext, config) => { config.AddInMemoryCollection(dict); });
+		builder.ConfigureTestServices(collection => collection.AddSingleton<IPerforceService, PerforceServiceStub>());
+	}
+}
 
-		protected override void ConfigureWebHost(IWebHostBuilder builder)
-        {
-            var Dict = new Dictionary<string, string?>
-            {
-                {"Horde:DatabaseConnectionString", MongoDbInstance.ConnectionString},
-                {"Horde:DatabaseName", MongoDbInstance.DatabaseName},
-                {"Horde:LogServiceWriteCacheType", "inmemory"},
-                {"Horde:DisableAuth", "true"},
-                {"Horde:OidcAuthority", null},
-                {"Horde:OidcClientId", null},
-            };
-            builder.ConfigureAppConfiguration((hostingContext, config) => { config.AddInMemoryCollection(Dict); });
-			builder.ConfigureTestServices(Collection => Collection.AddSingleton<IPerforceService, Stubs.Services.PerforceServiceStub>());
-        }
-    }
+public class ControllerIntegrationTest : IDisposable
+{
+	private readonly Lazy<Task<Fixture>> _fixture;
 
-    public class ControllerIntegrationTest : IDisposable
-    {
-		protected MongoDbInstance MongoDbInstance { get; }
-		TestWebApplicationFactory<Startup> _factory { get; }
-		protected HttpClient client { get; }
-		private Lazy<Task<Fixture>> _fixture;
+	public ControllerIntegrationTest()
+	{
+		MongoDbInstance = new MongoDbInstance();
+		Factory = new TestWebApplicationFactory<Startup>(MongoDbInstance);
+		Client = Factory.CreateClient();
 
-		public ControllerIntegrationTest()
-		{
-			MongoDbInstance = new MongoDbInstance();
-			_factory = new TestWebApplicationFactory<Startup>(MongoDbInstance);
-			client = _factory.CreateClient();
+		_fixture = new Lazy<Task<Fixture>>(Task.Run(() => CreateFixture()));
+	}
 
-			_fixture = new Lazy<Task<Fixture>>(Task.Run(() => CreateFixture()));
-		}
+	protected MongoDbInstance MongoDbInstance { get; }
+	private TestWebApplicationFactory<Startup> Factory { get; }
+	protected HttpClient Client { get; }
 
-		public void Dispose()
-		{
-			MongoDbInstance.Dispose();
-		}
+	public void Dispose()
+	{
+		Dispose(true);
+		GC.SuppressFinalize(this);
+	}
 
-		public Task<Fixture> GetFixture()
-		{
-			return _fixture.Value;
-		}
+	protected virtual void Dispose(bool disposing)
+	{
+		MongoDbInstance.Dispose();
+	}
 
-		async Task<Fixture> CreateFixture()
-		{
-			IServiceProvider Services = _factory.Services;
-			DatabaseService DatabaseService = Services.GetRequiredService<DatabaseService>();
-			ITemplateCollection TemplateService = Services.GetRequiredService<ITemplateCollection>();
-			JobService JobService = Services.GetRequiredService<JobService>();
-			IArtifactCollection ArtifactCollection = Services.GetRequiredService<IArtifactCollection>();
-			StreamService StreamService = Services.GetRequiredService<StreamService>();
-			AgentService AgentService = Services.GetRequiredService<AgentService>();
-			IPerforceService PerforceService = Services.GetRequiredService<IPerforceService>();
-			GraphCollection GraphCollection = new GraphCollection(DatabaseService);
+	public Task<Fixture> GetFixture()
+	{
+		return _fixture.Value;
+	}
+
+	private async Task<Fixture> CreateFixture()
+	{
+		IServiceProvider services = Factory.Services;
+		DatabaseService databaseService = services.GetRequiredService<DatabaseService>();
+		ITemplateCollection templateService = services.GetRequiredService<ITemplateCollection>();
+		JobService jobService = services.GetRequiredService<JobService>();
+		IArtifactCollection artifactCollection = services.GetRequiredService<IArtifactCollection>();
+		StreamService streamService = services.GetRequiredService<StreamService>();
+		AgentService agentService = services.GetRequiredService<AgentService>();
+		GraphCollection graphCollection = new(databaseService);
 
-			return await Fixture.Create(GraphCollection, TemplateService, JobService, ArtifactCollection, StreamService, AgentService, PerforceService);
-        }
-    }
+		return await Fixture.Create(graphCollection, templateService, jobService, artifactCollection, streamService,
+			agentService);
+	}
 }
\ No newline at end of file
diff --git a/Engine/Source/Programs/Horde/Horde.Build.Tests/DatabaseIntegrationTest.cs b/Engine/Source/Programs/Horde/Horde.Build.Tests/DatabaseIntegrationTest.cs
--- a/Engine/Source/Programs/Horde/Horde.Build.Tests/DatabaseIntegrationTest.cs
+++ b/Engine/Source/Programs/Horde/Horde.Build.Tests/DatabaseIntegrationTest.cs
@@ -3,10 +3,9 @@
 using System;
 using System.Collections.Generic;
 using System.Net;
-using System.Threading.Tasks;
-using Horde.Build;
 using Horde.Build.Services;
 using Horde.Build.Utilities;
+using Microsoft.Extensions.DependencyInjection;
 using Microsoft.Extensions.Logging;
 using Microsoft.Extensions.Options;
 using Microsoft.VisualStudio.TestTools.UnitTesting;
@@ -15,74 +14,74 @@ using StackExchange.Redis;
 
 namespace Horde.Build.Tests
 {
-    // Stub for fulfilling IOptionsMonitor interface during testing
-    public class TestOptionsMonitor<T> : IOptionsMonitor<T>, IDisposable
+	// Stub for fulfilling IOptionsMonitor interface during testing
+	public sealed class TestOptionsMonitor<T> : IOptionsMonitor<T>
         where T : class, new()
     {
-        public TestOptionsMonitor(T CurrentValue)
+		sealed class Disposable : IDisposable
+		{
+			public void Dispose() { }
+		}
+
+		public TestOptionsMonitor(T currentValue)
         {
-            this.CurrentValue = CurrentValue;
+            CurrentValue = currentValue;
         }
 
-        public T Get(string Name)
+        public T Get(string name)
         {
             return CurrentValue;
         }
 
         public IDisposable OnChange(Action<T, string> listener)
         {
-            return this;
+			return new Disposable();
         }
 
         public T CurrentValue { get; }
-        
-        public void Dispose()
-        {
-	        // Dummy stub to satisfy return value of OnChange 
-        }
     }
 
-	public class MongoDbInstance : IDisposable
+	public sealed class MongoDbInstance : IDisposable
 	{
 		public string DatabaseName { get; }
 		public string ConnectionString { get; }
 		MongoClient Client { get; }
 
-		private static object LockObject = new object();
-		private static MongoDbRunnerLocal? _mongoDbRunner;
-		private static int NextDatabaseIndex = 1;
+		private static readonly object s_lockObject = new object();
+		private static MongoDbRunnerLocal? s_mongoDbRunner;
+		private static int s_nextDatabaseIndex = 1;
 		public const string MongoDbDatabaseNamePrefix = "HordeServerTest_";
 
 		public MongoDbInstance()
 		{
-			int DatabaseIndex;
-			lock (LockObject)
+			int databaseIndex;
+			lock (s_lockObject)
 			{
-				if (_mongoDbRunner == null)
+				if (s_mongoDbRunner == null)
 				{
 					// One-time setup per test run to avoid overhead of starting the external MongoDB process
 					Startup.ConfigureMongoDbClient();
-					_mongoDbRunner = new MongoDbRunnerLocal();
-					_mongoDbRunner.Start();
+					s_mongoDbRunner = new MongoDbRunnerLocal();
+					s_mongoDbRunner.Start();
 
 					// Drop all the previous databases
-					MongoClientSettings MongoSettings = MongoClientSettings.FromConnectionString(_mongoDbRunner.GetConnectionString());
-					MongoClient Client = new MongoClient(MongoSettings);
+					MongoClientSettings mongoSettings = MongoClientSettings.FromConnectionString(s_mongoDbRunner.GetConnectionString());
+					MongoClient client = new MongoClient(mongoSettings);
 
-					List<string> DropDatabaseNames = Client.ListDatabaseNames().ToList();
-					foreach (string DropDatabaseName in DropDatabaseNames)
+					List<string> dropDatabaseNames = client.ListDatabaseNames().ToList();
+					foreach (string dropDatabaseName in dropDatabaseNames)
 					{
-						if (DropDatabaseName.StartsWith(MongoDbDatabaseNamePrefix, StringComparison.Ordinal))
+						if (dropDatabaseName.StartsWith(MongoDbDatabaseNamePrefix, StringComparison.Ordinal))
 						{
-							Client.DropDatabase(DropDatabaseName);
+							client.DropDatabase(dropDatabaseName);
 						}
 					}
 				}
-				DatabaseIndex = NextDatabaseIndex++;
+				databaseIndex = s_nextDatabaseIndex++;
 			}
 
-			DatabaseName = $"{MongoDbDatabaseNamePrefix}{DatabaseIndex}";
-			ConnectionString = $"{_mongoDbRunner.GetConnectionString()}/{DatabaseName}";
+			DatabaseName = $"{MongoDbDatabaseNamePrefix}{databaseIndex}";
+			ConnectionString = $"{s_mongoDbRunner.GetConnectionString()}/{DatabaseName}";
 			Client = new MongoClient(MongoClientSettings.FromConnectionString(ConnectionString));
 		}
 
@@ -92,104 +91,140 @@ namespace Horde.Build.Tests
 		}
 	}
 
-	public class DatabaseIntegrationTest : IDisposable
+	public class ServiceTest : IDisposable
+	{
+		private ServiceProvider? _serviceProvider = null;
+
+		public IServiceProvider ServiceProvider
+		{
+			get
+			{
+				if (_serviceProvider == null)
+				{
+					IServiceCollection services = new ServiceCollection();
+					ConfigureServices(services);
+
+					_serviceProvider = services.BuildServiceProvider();
+				}
+				return _serviceProvider;
+			}
+		}
+
+		public void Dispose()
+		{
+			Dispose(true);
+			GC.SuppressFinalize(this);
+		}
+
+		protected virtual void Dispose(bool disposing)
+		{
+			_serviceProvider?.Dispose();
+		}
+
+		protected virtual void ConfigureSettings(ServerSettings settings)
+		{
+		}
+
+		protected virtual void ConfigureServices(IServiceCollection services)
+		{
+		}
+	}
+
+	public class DatabaseIntegrationTest : ServiceTest
     {
-		private static object LockObject = new object();
-		private MongoDbInstance? MongoDbInstance;
-		private DatabaseService? DatabaseService;
-		
-        private static RedisRunner? _redisRunner;
-        private static IDatabase? _redisDb;
-        private static RedisConnectionPool? _redisConnectionPool;
-        private static ConnectionMultiplexer? _conMux;
+		private static readonly object s_lockObject = new object();
+
+		private MongoDbInstance? _mongoDbInstance;
+		private DatabaseService? _databaseService;
+		private readonly LoggerFactory _loggerFactory = new LoggerFactory();
+
+		private static RedisRunner? s_redisRunner;
+        private IDatabase? _redisDb;
+        private RedisConnectionPool? _redisConnectionPool;
+        private ConnectionMultiplexer? _conMux;
         public const int RedisDbNum = 15;
 
 		public DatabaseIntegrationTest()
 		{
 		}
 
-		public void Dispose()
+		protected override void ConfigureServices(IServiceCollection services)
 		{
-			MongoDbInstance?.Dispose();
-			DatabaseService?.Dispose();
+			services.AddSingleton(GetDatabaseServiceSingleton());
+			services.AddSingleton(GetRedisDatabase());
 		}
 
-		public DatabaseService GetDatabaseService()
+		protected override void Dispose(bool disposing)
+		{
+			base.Dispose(disposing);
+
+			_mongoDbInstance?.Dispose();
+			_databaseService?.Dispose();
+			_loggerFactory.Dispose();
+
+			_conMux?.Dispose();
+		}
+
+		public DatabaseService GetDatabaseServiceSingleton()
         {
-			lock(LockObject)
+			lock(s_lockObject)
 			{
-				if (DatabaseService == null)
+				if (_databaseService == null)
 				{
-					MongoDbInstance = new MongoDbInstance();
-
-					ServerSettings Ss = new ServerSettings();
-					Ss.DatabaseName = MongoDbInstance.DatabaseName;
-					Ss.DatabaseConnectionString = MongoDbInstance.ConnectionString;
+					_mongoDbInstance = new MongoDbInstance();
 
-					ILoggerFactory LoggerFactory = new LoggerFactory();
+					ServerSettings ss = new ServerSettings();
+					ss.DatabaseName = _mongoDbInstance.DatabaseName;
+					ss.DatabaseConnectionString = _mongoDbInstance.ConnectionString;
 
-					DatabaseService = new DatabaseService(Options.Create(Ss), LoggerFactory);
+					_databaseService = new DatabaseService(Options.Create(ss), _loggerFactory);
 				}
 			}
-			return DatabaseService;
+			return _databaseService;
         }
 
-        public static IDatabase GetRedisDatabase()
-        {
-			lock (LockObject)
+		private static RedisRunner GetRedisRunner()
+		{
+			lock (s_lockObject)
 			{
-				if (_redisDb == null)
+				if (s_redisRunner == null)
 				{
 					// One-time setup per test run to avoid overhead of starting the external Redis process
-					_redisRunner = new RedisRunner();
-					_redisRunner.Start();
+					s_redisRunner = new RedisRunner();
+					s_redisRunner.Start();
+				}
+			}
+			return s_redisRunner;
+		}
 
-					(string Host, int Port) = _redisRunner.GetListenAddress();
-					_conMux = ConnectionMultiplexer.Connect($"{Host}:{Port},allowAdmin=true");
-					ClearRedisDatabase();
+		public IDatabase GetRedisDatabase()
+        {
+			RedisRunner redisRunner = GetRedisRunner();
+			(string host, int port) = redisRunner.GetListenAddress();
 
-					_redisDb = _conMux.GetDatabase(RedisDbNum);
-				}
+			_conMux = ConnectionMultiplexer.Connect($"{host}:{port},allowAdmin=true");
+			foreach (EndPoint endpoint in _conMux.GetEndPoints())
+			{
+				_conMux.GetServer(endpoint).FlushDatabase(RedisDbNum);
 			}
 
+			_redisDb = _conMux.GetDatabase(RedisDbNum);
 	        return _redisDb;
         }
         
-        public static RedisConnectionPool GetRedisConnectionPool()
+        public RedisConnectionPool GetRedisConnectionPool()
         {
-	        if (_redisConnectionPool != null)
-	        {
-		        return _redisConnectionPool;
-	        }
-	        
-	        // Will initialize the Redis runner
-	        GetRedisDatabase();
-
-	        if (_redisRunner == null)
-	        {
-		        throw new Exception("Redis runner should have been initialized");
-	        }
-
-	        (string Host, int Port) = _redisRunner.GetListenAddress();
-	        _redisConnectionPool = new RedisConnectionPool(3, $"{Host}:{Port},allowAdmin=true", RedisDbNum);
-	        return _redisConnectionPool;
-        }
+			RedisRunner redisRunner = GetRedisRunner();
+			(string host, int port) = redisRunner.GetListenAddress();
 
-        public static void ClearRedisDatabase()
-        {
-	        if (_conMux != null)
-	        {
-		        foreach (EndPoint Endpoint in _conMux.GetEndPoints())
-		        {
-			        _conMux.GetServer(Endpoint).FlushDatabase(RedisDbNum);
-		        }    
-	        }
+	        _redisConnectionPool = new RedisConnectionPool(3, $"{host}:{port},allowAdmin=true", RedisDbNum);
+	        return _redisConnectionPool;
         }
 
-		public static T Deref<T>(T? Item)
+		public static T Deref<T>(T? item)
 		{
-			Assert.IsNotNull(Item);
-			return Item!;
+			Assert.IsNotNull(item);
+			return item!;
 		}
 	}
 }
\ No newline at end of file
diff --git a/Engine/Source/Programs/Horde/Horde.Build.Tests/DatabaseRunner.cs b/Engine/Source/Programs/Horde/Horde.Build.Tests/DatabaseRunner.cs
--- a/Engine/Source/Programs/Horde/Horde.Build.Tests/DatabaseRunner.cs
+++ b/Engine/Source/Programs/Horde/Horde.Build.Tests/DatabaseRunner.cs
@@ -1,7 +1,6 @@
 // Copyright Epic Games, Inc. All Rights Reserved.
 
 using System;
-using System.Collections.Generic;
 using System.Diagnostics;
 using System.IO;
 using System.Linq;
@@ -10,214 +9,215 @@ using System.Net.NetworkInformation;
 using System.Net.Sockets;
 using System.Reflection;
 using System.Runtime.InteropServices;
-using Castle.Core.Internal;
 using EpicGames.Core;
 
-namespace Horde.Build.Tests
+namespace Horde.Build.Tests;
+
+public abstract class DatabaseRunner
 {
-	public abstract class DatabaseRunner
+	private readonly string _binName;
+	private readonly int _defaultPort;
+	private readonly string _name;
+	private readonly bool _printStdErr;
+	private readonly bool _printStdOut;
+	private readonly bool _reuseProcess;
+	protected string TempDir { get; }
+	private Process? _proc;
+
+	protected DatabaseRunner(string name, string binName, int defaultPort, bool reuseProcess, bool printStdOut = false,
+		bool printStdErr = true)
 	{
-		private readonly string Name;
-		private readonly string BinName;
-		private readonly int DefaultPort;
-		private readonly bool ReuseProcess;
-		private readonly bool PrintStdOut;
-		private readonly bool PrintStdErr;
-		protected readonly string TempDir;
-		private Process? Proc;
-		protected int Port { get; private set; } = -1;
-
-		protected DatabaseRunner(string Name, string BinName, int DefaultPort, bool ReuseProcess, bool PrintStdOut = false, bool PrintStdErr = true)
-		{
-			this.Name = Name;
-			this.BinName = BinName;
-			this.DefaultPort = DefaultPort;
-			this.ReuseProcess = ReuseProcess;
-			this.PrintStdOut = PrintStdOut;
-			this.PrintStdErr = PrintStdErr;
-			TempDir = GetTemporaryDirectory();
-		}
-
-		protected abstract string GetArguments();
-		
-		public void Start()
-		{
-			if (ReuseProcess && !IsPortAvailable(DefaultPort))
-			{
-				Console.WriteLine($"Re-using already running {Name} process!");
-				Port = DefaultPort;
-				return;
-			}
-
-			if (!RuntimeInformation.IsOSPlatform(OSPlatform.Windows))
-			{
-				Console.WriteLine($"Unable to find a running {Name} process to use during testing!");
-				Console.WriteLine("This is required on any non-Windows as the runner can only start Windows binaries!");
-				Console.WriteLine($"Please ensure {BinName} is running on the default port {DefaultPort}.");
-				throw new Exception("Failed finding process to re-use! See stdout for info.");
-			}
-			
-			if (Proc != null)
-			{
-				return;
-			}
+		this._name = name;
+		this._binName = binName;
+		this._defaultPort = defaultPort;
+		this._reuseProcess = reuseProcess;
+		this._printStdOut = printStdOut;
+		this._printStdErr = printStdErr;
+		TempDir = GetTemporaryDirectory();
+	}
 
-			Port = GetAvailablePort();
-			
-			Process P = new Process();
-			if (PrintStdOut)
-			{
-				P.OutputDataReceived += (_, Args) => Console.WriteLine("{0} stdout: {1}", Name, Args.Data);	
-			}
-			if (PrintStdErr)
-			{
-				P.ErrorDataReceived += (_, Args) => Console.WriteLine("{0} stderr: {1}", Name, Args.Data);
-			}
-				
-			P.StartInfo.FileName = GetBinaryPath();
-			P.StartInfo.WorkingDirectory = TempDir;
-			P.StartInfo.Arguments = GetArguments();
-			P.StartInfo.UseShellExecute = false;
-			P.StartInfo.CreateNoWindow = true;
-			P.StartInfo.RedirectStandardOutput = true;
-			P.StartInfo.RedirectStandardError = true;
-
-			if (!P.Start())
-			{
-				throw new Exception("Process start failed!");
-			}
-			P.BeginOutputReadLine();
-			P.BeginErrorReadLine();
+	protected int Port { get; private set; } = -1;
 
-			// Try detect when main .NET process exits and kill the runner
-			AppDomain.CurrentDomain.ProcessExit += (Sender, EventArgs) =>
-			{
-				Console.WriteLine("Main process exiting!");
-				Stop();
-			};
+	protected abstract string GetArguments();
 
-			Proc = P;
+	public void Start()
+	{
+		if (_reuseProcess && !IsPortAvailable(_defaultPort))
+		{
+			Console.WriteLine($"Re-using already running {_name} process!");
+			Port = _defaultPort;
+			return;
 		}
-		
-		public void Stop()
+
+		if (!RuntimeInformation.IsOSPlatform(OSPlatform.Windows))
 		{
-			if (Proc != null)
-			{
-				Proc.Kill(true);
-				
-				// Waiting for exit blocks excessively even though the kill was sent. Anti-virus interfering?
-				// Process eventually shuts down but takes 2-3 min in Redis case. ReuseProcess flags circumvents this.
-				//Proc.WaitForExit();
-				Proc = null;
-				DeleteDirectory(TempDir);
-			}
+			Console.WriteLine($"Unable to find a running {_name} process to use during testing!");
+			Console.WriteLine("This is required on any non-Windows as the runner can only start Windows binaries!");
+			Console.WriteLine($"Please ensure {_binName} is running on the default port {_defaultPort}.");
+			throw new Exception("Failed finding process to re-use! See stdout for info.");
 		}
 
-		public (string Host, int Port) GetListenAddress()
+		if (_proc != null)
 		{
-			return ("localhost", Port);
+			return;
 		}
-		
-		private string GetBinaryPath()
+
+		Port = GetAvailablePort();
+
+		Process p = new();
+		if (_printStdOut)
 		{
-			FileReference File = new FileReference(new Uri(Assembly.GetExecutingAssembly().Location).LocalPath);
-			FileReference BinPath = FileReference.Combine(File.Directory, BinName);
-			return BinPath.FullName;
+			p.OutputDataReceived += (_, args) => Console.WriteLine("{0} stdout: {1}", _name, args.Data);
 		}
-		
-		private string GetTemporaryDirectory()
+
+		if (_printStdErr)
 		{
-			string Temp = Path.Join(Path.GetTempPath(), $"horde-{Name}-"  + Path.GetRandomFileName());
-			Directory.CreateDirectory(Temp);
-			return Temp;
+			p.ErrorDataReceived += (_, args) => Console.WriteLine("{0} stderr: {1}", _name, args.Data);
 		}
-		
-		private static int GetAvailablePort()
+
+		p.StartInfo.FileName = GetBinaryPath();
+		p.StartInfo.WorkingDirectory = TempDir;
+		p.StartInfo.Arguments = GetArguments();
+		p.StartInfo.UseShellExecute = false;
+		p.StartInfo.CreateNoWindow = true;
+		p.StartInfo.RedirectStandardOutput = true;
+		p.StartInfo.RedirectStandardError = true;
+
+		if (!p.Start())
 		{
-			TcpListener Listener = new TcpListener(IPAddress.Loopback, 0);
-			Listener.Start();
-			int Port = ((IPEndPoint)Listener.LocalEndpoint).Port;
-			Listener.Stop();
-			return Port;
+			throw new Exception("Process start failed!");
 		}
 
-		private static bool IsPortAvailable(int Port)
+		p.BeginOutputReadLine();
+		p.BeginErrorReadLine();
+
+		// Try detect when main .NET process exits and kill the runner
+		AppDomain.CurrentDomain.ProcessExit += (sender, eventArgs) =>
 		{
-			if (RuntimeInformation.IsOSPlatform(OSPlatform.Windows))
-			{
-				IPGlobalProperties IpGlobalProperties = IPGlobalProperties.GetIPGlobalProperties();
+			Console.WriteLine("Main process exiting!");
+			Stop();
+		};
 
-				IPEndPoint[] Listeners = IpGlobalProperties.GetActiveTcpListeners();
-				if (Listeners.Any(x => x.Port == Port))
-				{
-					return false;
-				}
+		_proc = p;
+	}
 
-				return true;
-			}
-			else
-			{
-				TcpListener? ListenerAny = null;
-				TcpListener? ListenerLoopback = null;
-				try
-				{
-					ListenerAny = new TcpListener(IPAddress.Loopback, Port);
-					ListenerAny.Start();
-					ListenerLoopback = new TcpListener(IPAddress.Any, Port);
-					ListenerLoopback.Start();
-					return true;
-				}
-				catch (SocketException)
-				{
-				}
-				finally
-				{
-					ListenerAny?.Stop();
-					ListenerLoopback?.Stop();
-				}
+	public void Stop()
+	{
+		if (_proc != null)
+		{
+			_proc.Kill(true);
 
-				return false;
-			}
+			// Waiting for exit blocks excessively even though the kill was sent. Anti-virus interfering?
+			// Process eventually shuts down but takes 2-3 min in Redis case. ReuseProcess flags circumvents this.
+			//Proc.WaitForExit();
+			_proc = null;
+			DeleteDirectory(TempDir);
 		}
-		
-		private static void DeleteDirectory(string Path)
+	}
+
+	public (string Host, int Port) GetListenAddress()
+	{
+		return ("localhost", Port);
+	}
+
+	private string GetBinaryPath()
+	{
+		FileReference file = new(new Uri(Assembly.GetExecutingAssembly().Location).LocalPath);
+		FileReference binPath = FileReference.Combine(file.Directory, _binName);
+		return binPath.FullName;
+	}
+
+	private string GetTemporaryDirectory()
+	{
+		string temp = Path.Join(Path.GetTempPath(), $"horde-{_name}-" + Path.GetRandomFileName());
+		Directory.CreateDirectory(temp);
+		return temp;
+	}
+
+	private static int GetAvailablePort()
+	{
+		TcpListener listener = new(IPAddress.Loopback, 0);
+		listener.Start();
+		int port = ((IPEndPoint)listener.LocalEndpoint).Port;
+		listener.Stop();
+		return port;
+	}
+
+	private static bool IsPortAvailable(int port)
+	{
+		if (RuntimeInformation.IsOSPlatform(OSPlatform.Windows))
 		{
-			DirectoryInfo Dir = new DirectoryInfo(Path) { Attributes = FileAttributes.Normal };
-			foreach (var Info in Dir.GetFileSystemInfos("*", SearchOption.AllDirectories))
+			IPGlobalProperties ipGlobalProperties = IPGlobalProperties.GetIPGlobalProperties();
+
+			IPEndPoint[] listeners = ipGlobalProperties.GetActiveTcpListeners();
+			if (listeners.Any(x => x.Port == port))
 			{
-				Info.Attributes = FileAttributes.Normal;
+				return false;
 			}
-			Dir.Delete(true);
+
+			return true;
 		}
-	}
 
-	public class MongoDbRunnerLocal : DatabaseRunner
-	{
-		public MongoDbRunnerLocal() : base("mongodb", "ThirdParty/Mongo/mongod.exe", 27017, true)
+		TcpListener? listenerAny = null;
+		TcpListener? listenerLoopback = null;
+		try
 		{
+			listenerAny = new TcpListener(IPAddress.Loopback, port);
+			listenerAny.Start();
+			listenerLoopback = new TcpListener(IPAddress.Any, port);
+			listenerLoopback.Start();
+			return true;
 		}
-
-		protected override string GetArguments()
+		catch (SocketException)
 		{
-			return $"--dbpath {TempDir} --noauth --quiet --port {Port}";
 		}
-
-		public string GetConnectionString()
+		finally
 		{
-			(var Host, int ListenPort) = GetListenAddress();
-			return $"mongodb://{Host}:{ListenPort}";
+			listenerAny?.Stop();
+			listenerLoopback?.Stop();
 		}
+
+		return false;
 	}
-	
-	public class RedisRunner : DatabaseRunner
+
+	private static void DeleteDirectory(string path)
 	{
-		public RedisRunner() : base("redis", "ThirdParty/Redis/redis-server.exe", 6379, true)
+		DirectoryInfo dir = new(path) { Attributes = FileAttributes.Normal };
+		foreach (FileSystemInfo info in dir.GetFileSystemInfos("*", SearchOption.AllDirectories))
 		{
+			info.Attributes = FileAttributes.Normal;
 		}
 
-		protected override string GetArguments()
-		{
-			return $"--port {Port} --save \"\" --appendonly no";
-		}
+		dir.Delete(true);
+	}
+}
+
+public class MongoDbRunnerLocal : DatabaseRunner
+{
+	public MongoDbRunnerLocal() : base("mongodb", "ThirdParty/Mongo/mongod.exe", 27017, true)
+	{
+	}
+
+	protected override string GetArguments()
+	{
+		return $"--dbpath {TempDir} --noauth --quiet --port {Port}";
+	}
+
+	public string GetConnectionString()
+	{
+		(string host, int listenPort) = GetListenAddress();
+		return $"mongodb://{host}:{listenPort}";
+	}
+}
+
+public class RedisRunner : DatabaseRunner
+{
+	public RedisRunner() : base("redis", "ThirdParty/Redis/redis-server.exe", 6379, true)
+	{
+	}
+
+	protected override string GetArguments()
+	{
+		return $"--port {Port} --save \"\" --appendonly no";
 	}
 }
\ No newline at end of file
diff --git a/Engine/Source/Programs/Horde/Horde.Build.Tests/DatabaseRunnerTest.cs b/Engine/Source/Programs/Horde/Horde.Build.Tests/DatabaseRunnerTest.cs
--- a/Engine/Source/Programs/Horde/Horde.Build.Tests/DatabaseRunnerTest.cs
+++ b/Engine/Source/Programs/Horde/Horde.Build.Tests/DatabaseRunnerTest.cs
@@ -11,19 +11,19 @@ namespace Horde.Build.Tests
         [TestMethod]
         public void RunMongoDbTest()
         {
-	        MongoDbRunnerLocal Runner = new MongoDbRunnerLocal();
-	        Runner.Start();
+	        MongoDbRunnerLocal runner = new MongoDbRunnerLocal();
+	        runner.Start();
             Thread.Sleep(100);
-            Runner.Stop();
+            runner.Stop();
         }
         
         [TestMethod]
         public void RunRedisTest()
         {
-	        RedisRunner Runner = new RedisRunner();
-	        Runner.Start();
+	        RedisRunner runner = new RedisRunner();
+	        runner.Start();
 	        Thread.Sleep(100);
-	        Runner.Stop();
+	        runner.Stop();
         }
     }
 }
\ No newline at end of file
diff --git a/Engine/Source/Programs/Horde/Horde.Build.Tests/DebugSettings.cs b/Engine/Source/Programs/Horde/Horde.Build.Tests/DebugSettings.cs
--- a/Engine/Source/Programs/Horde/Horde.Build.Tests/DebugSettings.cs
+++ b/Engine/Source/Programs/Horde/Horde.Build.Tests/DebugSettings.cs
@@ -7,7 +7,7 @@ namespace Horde.Build.Tests
     ///
     /// Put into a separate file to avoid accidental commit of credentials.
     /// </summary>
-	public class DebugSettings
+	public static class DebugSettings
     {
 	    public static readonly string DbUsername = "username-not-set";
 	    public static readonly string DbPassword = "password-not-set";
diff --git a/Engine/Source/Programs/Horde/Horde.Build.Tests/DebugTest.cs b/Engine/Source/Programs/Horde/Horde.Build.Tests/DebugTest.cs
--- a/Engine/Source/Programs/Horde/Horde.Build.Tests/DebugTest.cs
+++ b/Engine/Source/Programs/Horde/Horde.Build.Tests/DebugTest.cs
@@ -1,20 +1,5 @@
 // Copyright Epic Games, Inc. All Rights Reserved.
 
-using System;
-using System.Collections.Generic;
-using System.Threading.Tasks;
-using Castle.Core.Internal;
-using Horde.Build;
-using Horde.Build.Models;
-using Microsoft.VisualStudio.TestTools.UnitTesting;
-using Horde.Build.Services;
-using Horde.Build.Services.Impl;
-using Horde.Build.Utilities;
-using Microsoft.Extensions.DependencyInjection;
-using Microsoft.Extensions.Logging;
-using Microsoft.Extensions.Logging.Console;
-using MongoDB.Bson;
-
 #if false
 using PoolId = Horde.Build.Utilities.StringId<Horde.Build.Models.IPool>;
 using Microsoft.Extensions.Options;
diff --git a/Engine/Source/Programs/Horde/Horde.Build.Tests/Fixture.cs b/Engine/Source/Programs/Horde/Horde.Build.Tests/Fixture.cs
--- a/Engine/Source/Programs/Horde/Horde.Build.Tests/Fixture.cs
+++ b/Engine/Source/Programs/Horde/Horde.Build.Tests/Fixture.cs
@@ -9,8 +9,6 @@ using HordeCommon;
 using Horde.Build.Collections;
 using Horde.Build.Models;
 using Horde.Build.Services;
-using Horde.Build.Tests.Stubs.Services;
-using MongoDB.Bson;
 using StreamId = Horde.Build.Utilities.StringId<Horde.Build.Models.IStream>;
 using TemplateRefId = Horde.Build.Utilities.StringId<Horde.Build.Models.TemplateRef>;
 using ProjectId = Horde.Build.Utilities.StringId<Horde.Build.Models.IProject>;
@@ -35,57 +33,57 @@ namespace Horde.Build.Tests
 		public string Job1ArtifactData { get; private set; } = null!;
 		public IAgent Agent1 { get; private set; } = null!;
 		public string Agent1Name { get; private set; } = null!;
-		public string PoolName = "TestingPool";
+		public const string PoolName = "TestingPool";
 
-		public static async Task<Fixture> Create(IGraphCollection GraphCollection, ITemplateCollection TemplateCollection, JobService JobService, IArtifactCollection ArtifactCollection, StreamService StreamService, AgentService AgentService, IPerforceService PerforceService)
+		public static async Task<Fixture> Create(IGraphCollection graphCollection, ITemplateCollection templateCollection, JobService jobService, IArtifactCollection artifactCollection, StreamService streamService, AgentService agentService)
 		{
-			Fixture _fixture = new Fixture();
-			await _fixture.Populate(GraphCollection, TemplateCollection, JobService, ArtifactCollection, StreamService, AgentService, PerforceService);
+			Fixture fixture = new Fixture();
+			await fixture.Populate(graphCollection, templateCollection, jobService, artifactCollection, streamService, agentService);
 
 //			(PerforceService as PerforceServiceStub)?.AddChange("//UE5/Main", 112233, "leet.coder", "Did stuff", new []{"file.cpp"});
 //			(PerforceService as PerforceServiceStub)?.AddChange("//UE5/Main", 1111, "swarm", "A shelved CL here", new []{"renderer.cpp"});
 		
-			return _fixture;
+			return fixture;
 		}
 
-		private async Task Populate(IGraphCollection GraphCollection, ITemplateCollection TemplateCollection, JobService JobService, IArtifactCollection ArtifactCollection, StreamService StreamService, AgentService AgentService, IPerforceService PerforceService)
+		private async Task Populate(IGraphCollection graphCollection, ITemplateCollection templateCollection, JobService jobService, IArtifactCollection artifactCollection, StreamService streamService, AgentService agentService)
 		{
-			var Fg = new FixtureGraph();
-			Fg.Id = ContentHash.Empty;
-			Fg.Schema = 1122;
-			Fg.Groups = new List<INodeGroup>();
-			Fg.Aggregates = new List<IAggregate>();
-			Fg.Labels = new List<ILabel>();
+			FixtureGraph fg = new FixtureGraph();
+			fg.Id = ContentHash.Empty;
+			fg.Schema = 1122;
+			fg.Groups = new List<INodeGroup>();
+			fg.Aggregates = new List<IAggregate>();
+			fg.Labels = new List<ILabel>();
 
-			Template = await TemplateCollection.AddAsync("Test template");
-			Graph = await GraphCollection.AddAsync(Template);
+			Template = await templateCollection.AddAsync("Test template");
+			Graph = await graphCollection.AddAsync(Template);
 
 			TemplateRefId1 = new TemplateRefId("template1");
 			TemplateRefId2 = new TemplateRefId("template2");
 
-			List<CreateTemplateRefRequest> Templates = new List<CreateTemplateRefRequest>();
-			Templates.Add(new CreateTemplateRefRequest { Id = TemplateRefId1.ToString(), Name = "Test Template" });
-			Templates.Add(new CreateTemplateRefRequest { Id = TemplateRefId2.ToString(), Name = "Test Template" });
+			List<CreateTemplateRefRequest> templates = new List<CreateTemplateRefRequest>();
+			templates.Add(new CreateTemplateRefRequest { Id = TemplateRefId1.ToString(), Name = "Test Template" });
+			templates.Add(new CreateTemplateRefRequest { Id = TemplateRefId2.ToString(), Name = "Test Template" });
 
-			List<CreateStreamTabRequest> Tabs = new List<CreateStreamTabRequest>();
-			Tabs.Add(new CreateJobsTabRequest { Title = "foo", Templates = new List<string> { TemplateRefId1.ToString(), TemplateRefId2.ToString() } });
+			List<CreateStreamTabRequest> tabs = new List<CreateStreamTabRequest>();
+			tabs.Add(new CreateJobsTabRequest { Title = "foo", Templates = new List<string> { TemplateRefId1.ToString(), TemplateRefId2.ToString() } });
 
-			Dictionary<string, CreateAgentTypeRequest> AgentTypes = new()
+			Dictionary<string, CreateAgentTypeRequest> agentTypes = new()
 			{
 				{ "Win64", new() { Pool = PoolName} }
 			};
 
-			Stream = await StreamService.StreamCollection.GetAsync(new StreamId("ue5-main"));
-			Stream = await StreamService.StreamCollection.TryCreateOrReplaceAsync(
+			Stream = await streamService.StreamCollection.GetAsync(new StreamId("ue5-main"));
+			Stream = await streamService.StreamCollection.TryCreateOrReplaceAsync(
 				new StreamId("ue5-main"),
 				Stream,
 				"",
 				"",
 				new ProjectId("does-not-exist"),
-				new StreamConfig { Name = "//UE5/Main", Tabs = Tabs, Templates = Templates, AgentTypes = AgentTypes}
+				new StreamConfig { Name = "//UE5/Main", Tabs = tabs, Templates = templates, AgentTypes = agentTypes}
 			);
 			
-			Job1 = await JobService.CreateJobAsync(
+			Job1 = await jobService.CreateJobAsync(
 				JobId: new JobId("5f283932841e7fdbcafb6ab5"),
 				Stream: Stream!,
 				TemplateRefId: TemplateRefId1,
@@ -108,14 +106,14 @@ namespace Horde.Build.Tests
 				null,
 				Arguments: new List<string>()
 			);
-			Job1 = (await JobService.GetJobAsync(Job1.Id))!;
+			Job1 = (await jobService.GetJobAsync(Job1.Id))!;
 
-			Job2 = await JobService.CreateJobAsync(
+			Job2 = await jobService.CreateJobAsync(
 				JobId: new JobId("5f69ea1b68423e921b035106"),
 				Stream: Stream!,
 				TemplateRefId: new TemplateRefId("template-id-1"),
 				TemplateHash: ContentHash.MD5("made-up-template-hash"),
-				Graph: Fg,
+				Graph: fg,
 				Name: "hello2",
 				Change: 2000001,
 				CodeChange: 2000002,
@@ -133,14 +131,15 @@ namespace Horde.Build.Tests
 				null,
 				Arguments: new List<string>()
 			);
-			Job2 = (await JobService.GetJobAsync(Job2.Id))!;
+			Job2 = (await jobService.GetJobAsync(Job2.Id))!;
 
 			Job1ArtifactData = "For The Horde!";
-			Job1Artifact = await ArtifactCollection.CreateArtifactAsync(Job1.Id, SubResourceId.Parse("22"), "myFile.txt",
-				"text/plain", new MemoryStream(Encoding.UTF8.GetBytes(Job1ArtifactData)));
+			using MemoryStream job1ArtifactStream = new MemoryStream(Encoding.UTF8.GetBytes(Job1ArtifactData));
+			Job1Artifact = await artifactCollection.CreateArtifactAsync(Job1.Id, SubResourceId.Parse("22"), "myFile.txt",
+				"text/plain", job1ArtifactStream);
 
 			Agent1Name = "testAgent1";
-			Agent1 = await AgentService.CreateAgentAsync(Agent1Name, true, null, null);
+			Agent1 = await agentService.CreateAgentAsync(Agent1Name, true, null, null);
 		}
 
 		private class FixtureGraph : IGraph
diff --git a/Engine/Source/Programs/Horde/Horde.Build.Tests/Fleet/AutoscaleServiceV2Test.cs b/Engine/Source/Programs/Horde/Horde.Build.Tests/Fleet/AutoscaleServiceV2Test.cs
--- a/Engine/Source/Programs/Horde/Horde.Build.Tests/Fleet/AutoscaleServiceV2Test.cs
+++ b/Engine/Source/Programs/Horde/Horde.Build.Tests/Fleet/AutoscaleServiceV2Test.cs
@@ -1,15 +1,14 @@
 // Copyright Epic Games, Inc. All Rights Reserved.
 
+using StatsdClient;
 using System;
 using System.Collections.Generic;
 using System.Threading;
 using System.Threading.Tasks;
 using Horde.Build.Fleet.Autoscale;
-using Horde.Build;
 using Horde.Build.Models;
 using Horde.Build.Services;
 using Horde.Build.Utilities;
-using Horde.Build.Tests;
 using Microsoft.Extensions.DependencyInjection;
 using Microsoft.Extensions.Logging;
 using Microsoft.Extensions.Options;
@@ -18,27 +17,25 @@ using Microsoft.VisualStudio.TestTools.UnitTesting;
 namespace Horde.Build.Tests.Fleet
 {
 	using PoolId = StringId<IPool>;
-	using ProjectId = StringId<IProject>;
-	using StreamId = StringId<IStream>;
-	
+
 	public class FleetManagerSpy : IFleetManager
 	{
 		public int ExpandPoolAsyncCallCount { get; private set; }
 		public int ShrinkPoolAsyncCallCount { get; private set; }
 		
-		public Task ExpandPoolAsync(IPool Pool, IReadOnlyList<IAgent> Agents, int Count)
+		public Task ExpandPoolAsync(IPool pool, IReadOnlyList<IAgent> agents, int count)
 		{
 			ExpandPoolAsyncCallCount++;
 			return Task.CompletedTask;
 		}
 
-		public Task ShrinkPoolAsync(IPool Pool, IReadOnlyList<IAgent> Agents, int Count)
+		public Task ShrinkPoolAsync(IPool pool, IReadOnlyList<IAgent> agents, int count)
 		{
 			ShrinkPoolAsyncCallCount++;
 			return Task.CompletedTask;
 		}
 
-		public Task<int> GetNumStoppedInstancesAsync(IPool Pool)
+		public Task<int> GetNumStoppedInstancesAsync(IPool pool)
 		{
 			throw new NotImplementedException();
 		}
@@ -49,14 +46,14 @@ namespace Horde.Build.Tests.Fleet
 		public int CallCount { get; private set; }
 		public HashSet<PoolId> PoolIdsSeen { get; } = new();
 		
-		public Task<List<PoolSizeData>> CalcDesiredPoolSizesAsync(List<PoolSizeData> Pools)
+		public Task<List<PoolSizeData>> CalcDesiredPoolSizesAsync(List<PoolSizeData> pools)
 		{
 			CallCount++;
-			foreach (PoolSizeData Data in Pools)
+			foreach (PoolSizeData data in pools)
 			{
-				PoolIdsSeen.Add(Data.Pool.Id);
+				PoolIdsSeen.Add(data.Pool.Id);
 			}
-			return Task.FromResult(Pools);
+			return Task.FromResult(pools);
 		}
 
 		public string Name { get; } = "PoolSizeStrategySpy";
@@ -65,102 +62,90 @@ namespace Horde.Build.Tests.Fleet
 	[TestClass]
 	public class AutoscaleServiceV2Test : TestSetup
 	{
-		FleetManagerSpy FleetManagerSpy = new();
-		
+		readonly FleetManagerSpy _fleetManagerSpy = new();
+		readonly IDogStatsd _dogStatsD = new NoOpDogStatsd();
+
 		[TestMethod]
 		public async Task PerPoolStrategy()
 		{
-			var Service = GetAutoscaleService(FleetManagerSpy);
+			using AutoscaleServiceV2 service = GetAutoscaleService(_fleetManagerSpy);
 			
-			IPool Pool1 = await PoolService.CreatePoolAsync("bogusPoolLease1", null, true, 0, 0, SizeStrategy: PoolSizeStrategy.LeaseUtilization);
-			IPool Pool2 = await PoolService.CreatePoolAsync("bogusPoolLease2", null, true, 0, 0, SizeStrategy: null);
-			IPool Pool3 = await PoolService.CreatePoolAsync("bogusPoolJobQueue1", null, true, 0, 0, SizeStrategy: PoolSizeStrategy.JobQueue);
-			IPool Pool4 = await PoolService.CreatePoolAsync("bogusPoolJobQueue2", null, true, 0, 0, SizeStrategy: PoolSizeStrategy.JobQueue);
-			IPool Pool5 = await PoolService.CreatePoolAsync("bogusPoolNoOp", null, true, 0, 0, SizeStrategy: PoolSizeStrategy.NoOp);
-
-			PoolSizeStrategySpy LeaseUtilizationSpy = new();
-			PoolSizeStrategySpy JobQueueSpy = new();
-			PoolSizeStrategySpy NoOpSpy = new();
-			Service.OverridePoolSizeStrategiesDuringTesting(LeaseUtilizationSpy, JobQueueSpy, NoOpSpy);
-
-			await Service.TickLeaderAsync(CancellationToken.None);
+			IPool pool1 = await PoolService.CreatePoolAsync("bogusPoolLease1", null, true, 0, 0, SizeStrategy: PoolSizeStrategy.LeaseUtilization);
+			IPool pool2 = await PoolService.CreatePoolAsync("bogusPoolLease2", null, true, 0, 0, SizeStrategy: null);
+			IPool pool3 = await PoolService.CreatePoolAsync("bogusPoolJobQueue1", null, true, 0, 0, SizeStrategy: PoolSizeStrategy.JobQueue);
+			IPool pool4 = await PoolService.CreatePoolAsync("bogusPoolJobQueue2", null, true, 0, 0, SizeStrategy: PoolSizeStrategy.JobQueue);
+			IPool pool5 = await PoolService.CreatePoolAsync("bogusPoolNoOp", null, true, 0, 0, SizeStrategy: PoolSizeStrategy.NoOp);
+
+			PoolSizeStrategySpy leaseUtilizationSpy = new();
+			PoolSizeStrategySpy jobQueueSpy = new();
+			PoolSizeStrategySpy noOpSpy = new();
+			service.OverridePoolSizeStrategiesDuringTesting(leaseUtilizationSpy, jobQueueSpy, noOpSpy);
+
+			await service.TickLeaderAsync(CancellationToken.None);
 			
-			Assert.AreEqual(1, LeaseUtilizationSpy.CallCount);
-			Assert.IsTrue(LeaseUtilizationSpy.PoolIdsSeen.Contains(Pool1.Id));
-			Assert.IsTrue(LeaseUtilizationSpy.PoolIdsSeen.Contains(Pool2.Id));
+			Assert.AreEqual(1, leaseUtilizationSpy.CallCount);
+			Assert.IsTrue(leaseUtilizationSpy.PoolIdsSeen.Contains(pool1.Id));
+			Assert.IsTrue(leaseUtilizationSpy.PoolIdsSeen.Contains(pool2.Id));
 			
-			Assert.AreEqual(1, JobQueueSpy.CallCount);
-			Assert.IsTrue(JobQueueSpy.PoolIdsSeen.Contains(Pool3.Id));
-			Assert.IsTrue(JobQueueSpy.PoolIdsSeen.Contains(Pool4.Id));
+			Assert.AreEqual(1, jobQueueSpy.CallCount);
+			Assert.IsTrue(jobQueueSpy.PoolIdsSeen.Contains(pool3.Id));
+			Assert.IsTrue(jobQueueSpy.PoolIdsSeen.Contains(pool4.Id));
 			
-			Assert.AreEqual(1, NoOpSpy.CallCount);
-			Assert.IsTrue(NoOpSpy.PoolIdsSeen.Contains(Pool5.Id));
+			Assert.AreEqual(1, noOpSpy.CallCount);
+			Assert.IsTrue(noOpSpy.PoolIdsSeen.Contains(pool5.Id));
 		}
 
 		[TestMethod]
 		public async Task ScaleOutCooldown()
 		{
-			AutoscaleServiceV2 Service = GetAutoscaleService(FleetManagerSpy);
-			IPool Pool = await PoolService.CreatePoolAsync("testPool", null, true, 0, 0, SizeStrategy: PoolSizeStrategy.NoOp);
+			using AutoscaleServiceV2 service = GetAutoscaleService(_fleetManagerSpy);
+			IPool pool = await PoolService.CreatePoolAsync("testPool", null, true, 0, 0, SizeStrategy: PoolSizeStrategy.NoOp);
 
 			// First scale-out will succeed
-			await Service.ResizePools(new() { new PoolSizeData(Pool, new List<IAgent>(), 1, "Testing") });
-			Assert.AreEqual(1, FleetManagerSpy.ExpandPoolAsyncCallCount);
+			await service.ResizePools(new() { new PoolSizeData(pool, new List<IAgent>(), 1, "Testing") });
+			Assert.AreEqual(1, _fleetManagerSpy.ExpandPoolAsyncCallCount);
 			
 			// Cannot scale-out due to cool-down
-			await Service.ResizePools(new() { new PoolSizeData(Pool, new List<IAgent>(), 2, "Testing") });
-			Assert.AreEqual(1, FleetManagerSpy.ExpandPoolAsyncCallCount);
+			await service.ResizePools(new() { new PoolSizeData(pool, new List<IAgent>(), 2, "Testing") });
+			Assert.AreEqual(1, _fleetManagerSpy.ExpandPoolAsyncCallCount);
 
 			// Wait some time and then try again
 			await Clock.AdvanceAsync(TimeSpan.FromHours(2));
-			await Service.ResizePools(new() { new PoolSizeData(Pool, new List<IAgent>(), 2, "Testing") });
-			Assert.AreEqual(2, FleetManagerSpy.ExpandPoolAsyncCallCount);
+			await service.ResizePools(new() { new PoolSizeData(pool, new List<IAgent>(), 2, "Testing") });
+			Assert.AreEqual(2, _fleetManagerSpy.ExpandPoolAsyncCallCount);
 		}
 		
 		[TestMethod]
 		public async Task ScaleInCooldown()
 		{
-			AutoscaleServiceV2 Service = GetAutoscaleService(FleetManagerSpy);
-			IPool Pool = await PoolService.CreatePoolAsync("testPool", null, true, 0, 0, SizeStrategy: PoolSizeStrategy.NoOp);
-			IAgent Agent1 = await CreateAgentAsync(Pool);
-			IAgent Agent2 = await CreateAgentAsync(Pool);
+			using AutoscaleServiceV2 service = GetAutoscaleService(_fleetManagerSpy);
+			IPool pool = await PoolService.CreatePoolAsync("testPool", null, true, 0, 0, SizeStrategy: PoolSizeStrategy.NoOp);
+			IAgent agent1 = await CreateAgentAsync(pool);
+			IAgent agent2 = await CreateAgentAsync(pool);
 
 			// First scale-out will succeed
-			await Service.ResizePools(new() { new PoolSizeData(Pool, new () { Agent1, Agent2 }, 1, "Testing") });
-			Assert.AreEqual(1, FleetManagerSpy.ShrinkPoolAsyncCallCount);
+			await service.ResizePools(new() { new PoolSizeData(pool, new () { agent1, agent2 }, 1, "Testing") });
+			Assert.AreEqual(1, _fleetManagerSpy.ShrinkPoolAsyncCallCount);
 			
 			// Cannot scale-out due to cool-down
-			await Service.ResizePools(new() { new PoolSizeData(Pool, new () { Agent1 }, 0, "Testing") });
-			Assert.AreEqual(1, FleetManagerSpy.ShrinkPoolAsyncCallCount);
+			await service.ResizePools(new() { new PoolSizeData(pool, new () { agent1 }, 0, "Testing") });
+			Assert.AreEqual(1, _fleetManagerSpy.ShrinkPoolAsyncCallCount);
 
 			// Wait some time and then try again
 			await Clock.AdvanceAsync(TimeSpan.FromHours(2));
-			await Service.ResizePools(new() { new PoolSizeData(Pool, new () { Agent1 }, 0, "Testing") });
-			Assert.AreEqual(2, FleetManagerSpy.ShrinkPoolAsyncCallCount);
+			await service.ResizePools(new() { new PoolSizeData(pool, new () { agent1 }, 0, "Testing") });
+			Assert.AreEqual(2, _fleetManagerSpy.ShrinkPoolAsyncCallCount);
 		}
 
-		private async Task<List<PoolSizeData>> GetPoolSizeData()
+		private AutoscaleServiceV2 GetAutoscaleService(IFleetManager fleetManager)
 		{
-			IPool Pool1 = await PoolService.CreatePoolAsync("bogusPool1", null, true, 0, 0);
-			IAgent Agent1 = await CreateAgentAsync(Pool1);
-			IAgent Agent2 = await CreateAgentAsync(Pool1);
+			ILogger<AutoscaleServiceV2> logger = ServiceProvider.GetRequiredService<ILogger<AutoscaleServiceV2>>();
+			IOptions<ServerSettings> serverSettingsOpt = ServiceProvider.GetRequiredService<IOptions<ServerSettings>>();
 
-			return new List<PoolSizeData>
-			{
-				new (Pool1, new List<IAgent> { Agent1, Agent2 }, null)
-			};
+			LeaseUtilizationStrategy leaseUtilizationStrategy = new(AgentCollection, PoolCollection, LeaseCollection, Clock);
+			JobQueueStrategy jobQueueStrategy = new(JobCollection, GraphCollection, StreamService, Clock);
+			AutoscaleServiceV2 service = new AutoscaleServiceV2(leaseUtilizationStrategy,jobQueueStrategy, new NoOpPoolSizeStrategy(), AgentCollection, PoolCollection, fleetManager, _dogStatsD, Clock, serverSettingsOpt, logger);
+			return service;
 		}
-
-		private AutoscaleServiceV2 GetAutoscaleService(IFleetManager FleetManager)
-		{
-			ILogger<AutoscaleServiceV2> Logger = ServiceProvider.GetRequiredService<ILogger<AutoscaleServiceV2>>();
-			IOptions<ServerSettings> ServerSettingsOpt = ServiceProvider.GetRequiredService<IOptions<ServerSettings>>();
-
-			LeaseUtilizationStrategy LeaseUtilizationStrategy = new(AgentCollection, PoolCollection, LeaseCollection, Clock);
-			JobQueueStrategy JobQueueStrategy = new(JobCollection, GraphCollection, StreamService, Clock);
-			AutoscaleServiceV2 Service = new AutoscaleServiceV2(LeaseUtilizationStrategy,JobQueueStrategy, new NoOpPoolSizeStrategy(), AgentCollection, PoolCollection, FleetManager, new NoOpDogStatsd(), Clock, ServerSettingsOpt, Logger);
-			return Service;
-		}
-
 	}
 }
\ No newline at end of file
diff --git a/Engine/Source/Programs/Horde/Horde.Build.Tests/Fleet/JobQueueStrategyTest.cs b/Engine/Source/Programs/Horde/Horde.Build.Tests/Fleet/JobQueueStrategyTest.cs
--- a/Engine/Source/Programs/Horde/Horde.Build.Tests/Fleet/JobQueueStrategyTest.cs
+++ b/Engine/Source/Programs/Horde/Horde.Build.Tests/Fleet/JobQueueStrategyTest.cs
@@ -8,7 +8,6 @@ using Horde.Build.Fleet.Autoscale;
 using Horde.Build.Api;
 using Horde.Build.Models;
 using Horde.Build.Utilities;
-using Horde.Build.Tests;
 using Microsoft.VisualStudio.TestTools.UnitTesting;
 
 namespace Horde.Build.Tests.Fleet
@@ -16,18 +15,18 @@ namespace Horde.Build.Tests.Fleet
 	using PoolId = StringId<IPool>;
 	using ProjectId = StringId<IProject>;
 	using StreamId = StringId<IStream>;
-	
+
 	[TestClass]
 	public class JobQueueStrategyTest : TestSetup
 	{
 		[TestMethod]
 		public async Task GetPoolQueueSizes()
 		{
-			(JobQueueStrategy Strategy, PoolSizeData PoolSizeData) = await SetUpJobsAsync(1, 5);
-			await Clock.AdvanceAsync(Strategy.ReadyTimeThreshold + TimeSpan.FromSeconds(5));
-			Dictionary<PoolId, int> PoolQueueSizes = await Strategy.GetPoolQueueSizesAsync(Clock.UtcNow - TimeSpan.FromHours(2));
-			Assert.AreEqual(1, PoolQueueSizes.Count);
-			Assert.AreEqual(5, PoolQueueSizes[PoolSizeData.Pool.Id]);
+			(JobQueueStrategy strategy, PoolSizeData poolSizeData) = await SetUpJobsAsync(1, 5);
+			await Clock.AdvanceAsync(strategy.ReadyTimeThreshold + TimeSpan.FromSeconds(5));
+			Dictionary<PoolId, int> poolQueueSizes = await strategy.GetPoolQueueSizesAsync(Clock.UtcNow - TimeSpan.FromHours(2));
+			Assert.AreEqual(1, poolQueueSizes.Count);
+			Assert.AreEqual(5, poolQueueSizes[poolSizeData.Pool.Id]);
 		}
 		
 		[TestMethod]
@@ -78,79 +77,79 @@ namespace Horde.Build.Tests.Fleet
 			await AssertAgentCount(25, 7);
 		}
 
-		public async Task AssertAgentCount(int NumBatchesReady, int ExpectedAgentDelta, bool WaitedBeyondThreshold = true, int NumAgents = 8)
+		public async Task AssertAgentCount(int numBatchesReady, int expectedAgentDelta, bool waitedBeyondThreshold = true, int numAgents = 8)
 		{
-			(JobQueueStrategy Strategy, PoolSizeData PoolSizeData) = await SetUpJobsAsync(1, NumBatchesReady, NumAgents);
-			TimeSpan TimeToWait = WaitedBeyondThreshold
-				? Strategy.ReadyTimeThreshold + TimeSpan.FromSeconds(5)
+			(JobQueueStrategy strategy, PoolSizeData poolSizeData) = await SetUpJobsAsync(1, numBatchesReady, numAgents);
+			TimeSpan timeToWait = waitedBeyondThreshold
+				? strategy.ReadyTimeThreshold + TimeSpan.FromSeconds(5)
 				: TimeSpan.FromSeconds(15);
 			
-			await Clock.AdvanceAsync(TimeToWait);
+			await Clock.AdvanceAsync(timeToWait);
 
-			List<PoolSizeData> Result = await Strategy.CalcDesiredPoolSizesAsync(new() { PoolSizeData });
-			Assert.AreEqual(1, Result.Count);
-			Assert.AreEqual(Result[0].Agents.Count + ExpectedAgentDelta, Result[0].DesiredAgentCount);
+			List<PoolSizeData> result = await strategy.CalcDesiredPoolSizesAsync(new() { poolSizeData });
+			Assert.AreEqual(1, result.Count);
+			Assert.AreEqual(result[0].Agents.Count + expectedAgentDelta, result[0].DesiredAgentCount);
 		}
 	
 		/// <summary>
 		/// Set up a fixture for job queue tests, ensuring a certain number of job batches are in running or waiting state
 		/// </summary>
-		/// <param name="NumBatchesRunning">Num of job batches that should be in state running</param>
-		/// <param name="NumBatchesReady">Num of job batches that should be in state waiting</param>
-		private async Task<(JobQueueStrategy, PoolSizeData)> SetUpJobsAsync(int NumBatchesRunning, int NumBatchesReady, int NumAgents = 8)
+		/// <param name="numBatchesRunning">Num of job batches that should be in state running</param>
+		/// <param name="numBatchesReady">Num of job batches that should be in state waiting</param>
+		private async Task<(JobQueueStrategy, PoolSizeData)> SetUpJobsAsync(int numBatchesRunning, int numBatchesReady, int numAgents = 8)
 		{
-			IPool Pool1 = await PoolService.CreatePoolAsync("bogusPool1", null, true, 0, 0);
-			List<IAgent> Agents = new();
-			for (int i = 0; i < NumAgents; i++)
+			IPool pool1 = await PoolService.CreatePoolAsync("bogusPool1", null, true, 0, 0);
+			List<IAgent> agents = new();
+			for (int i = 0; i < numAgents; i++)
 			{
-				Agents.Add(await CreateAgentAsync(Pool1));
+				agents.Add(await CreateAgentAsync(pool1));
 			}
 			
-			PoolSizeData PoolSize = new (Pool1, Agents, null);
+			PoolSizeData poolSize = new (pool1, agents, null);
 			
-			string AgentTypeName1 = "bogusAgentType1";
-			Dictionary<string, CreateAgentTypeRequest> AgentTypes = new() { {AgentTypeName1, new() { Pool = Pool1.Name} }, };
-			IStream Stream = (await StreamCollection.TryCreateOrReplaceAsync(
+			string agentTypeName1 = "bogusAgentType1";
+			Dictionary<string, CreateAgentTypeRequest> agentTypes = new() { {agentTypeName1, new() { Pool = pool1.Name} }, };
+			IStream stream = (await StreamCollection.TryCreateOrReplaceAsync(
 				new StreamId("ue5-main"),
 				null,
 				"",
 				"",
 				new ProjectId("does-not-exist"),
-				new StreamConfig { Name = "//UE5/Main", AgentTypes = AgentTypes}
+				new StreamConfig { Name = "//UE5/Main", AgentTypes = agentTypes}
 			))!;
 
-			string NodeForAgentType1 = "bogusNodeOnAgentType1";
-			IGraph Graph = await GraphCollection.AppendAsync(null, new()
+			string nodeForAgentType1 = "bogusNodeOnAgentType1";
+			IGraph graph = await GraphCollection.AppendAsync(null, new()
 			{
-				new NewGroup(AgentTypeName1, new List<NewNode>
+				new NewGroup(agentTypeName1, new List<NewNode>
 				{
-					new (NodeForAgentType1),
+					new (nodeForAgentType1),
 				})
 			});
 
-			for (int i = 0; i < NumBatchesRunning; i++)
+			for (int i = 0; i < numBatchesRunning; i++)
 			{
-				IJob Job = await AddPlaceholderJob(Graph, Stream.Id, NodeForAgentType1);
-				await JobCollection.TryUpdateBatchAsync(Job, Graph, Job.Batches[0].Id, null, JobStepBatchState.Running, null);
+				IJob job = await AddPlaceholderJob(graph, stream.Id, nodeForAgentType1);
+				await JobCollection.TryUpdateBatchAsync(job, graph, job.Batches[0].Id, null, JobStepBatchState.Running, null);
 			}
 			
-			for (int i = 0; i < NumBatchesReady; i++)
+			for (int i = 0; i < numBatchesReady; i++)
 			{
-				IJob Job = await AddPlaceholderJob(Graph, Stream.Id, NodeForAgentType1);
-				await JobCollection.TryUpdateBatchAsync(Job, Graph, Job.Batches[0].Id, null, JobStepBatchState.Ready, null);
+				IJob job = await AddPlaceholderJob(graph, stream.Id, nodeForAgentType1);
+				await JobCollection.TryUpdateBatchAsync(job, graph, job.Batches[0].Id, null, JobStepBatchState.Ready, null);
 			}
 			
-			return (new (JobCollection, GraphCollection, StreamService, Clock), PoolSize);
+			return (new (JobCollection, GraphCollection, StreamService, Clock), poolSize);
 		}
 		
-		private async Task<IJob> AddPlaceholderJob(IGraph Graph, StreamId StreamId, string NodeNameToExecute)
+		private async Task<IJob> AddPlaceholderJob(IGraph graph, StreamId streamId, string nodeNameToExecute)
 		{
-			IJob Job = await JobCollection.AddAsync(ObjectId<IJob>.GenerateNewId(), StreamId,
-				new StringId<TemplateRef>("bogusTemplateRefId"), ContentHash.Empty, Graph, "bogusJobName",
+			IJob job = await JobCollection.AddAsync(ObjectId<IJob>.GenerateNewId(), streamId,
+				new StringId<TemplateRef>("bogusTemplateRefId"), ContentHash.Empty, graph, "bogusJobName",
 				1000, 1000, null, null, null, null, null, null, null, null, false,
-				false, null, null, new List<string> { "-Target=" + NodeNameToExecute });
+				false, null, null, new List<string> { "-Target=" + nodeNameToExecute });
 
-			return Job;
+			return job;
 		}
 	}
 }
\ No newline at end of file
diff --git a/Engine/Source/Programs/Horde/Horde.Build.Tests/Fleet/LeaseUtilizationStrategyTest.cs b/Engine/Source/Programs/Horde/Horde.Build.Tests/Fleet/LeaseUtilizationStrategyTest.cs
--- a/Engine/Source/Programs/Horde/Horde.Build.Tests/Fleet/LeaseUtilizationStrategyTest.cs
+++ b/Engine/Source/Programs/Horde/Horde.Build.Tests/Fleet/LeaseUtilizationStrategyTest.cs
@@ -10,110 +10,108 @@ using HordeCommon;
 using HordeCommon.Rpc.Tasks;
 using Horde.Build.Models;
 using Horde.Build.Utilities;
-using Horde.Build.Tests;
 using Microsoft.VisualStudio.TestTools.UnitTesting;
 
 namespace Horde.Build.Tests.Fleet
 {
-	using PoolId = StringId<IPool>;
 	using StreamId = StringId<IStream>;
-	
+
 	[TestClass]
 	public class LeaseUtilizationStrategyTest : TestSetup
 	{
-		private IAgent Agent1 = null!;
-		private IAgent Agent2 = null!;
-		private IAgent Agent3 = null!;
-		private IAgent Agent4 = null!;
-		private List<IAgent> PoolAgents = null!;
+		private IAgent _agent1 = null!;
+		private IAgent _agent2 = null!;
+		private IAgent _agent3 = null!;
+		private IAgent _agent4 = null!;
+		private List<IAgent> _poolAgents = null!;
 
-		public void CreateAgents(IPool Pool)
+		public void CreateAgents(IPool pool)
 		{
-			Agent1 = CreateAgentAsync(Pool).Result;
-			Agent2 = CreateAgentAsync(Pool).Result;
-			Agent3 = CreateAgentAsync(Pool).Result;
-			Agent4 = CreateAgentAsync(Pool).Result;
-			PoolAgents = new() { Agent1, Agent2, Agent3, Agent4 };
+			_agent1 = CreateAgentAsync(pool).Result;
+			_agent2 = CreateAgentAsync(pool).Result;
+			_agent3 = CreateAgentAsync(pool).Result;
+			_agent4 = CreateAgentAsync(pool).Result;
+			_poolAgents = new() { _agent1, _agent2, _agent3, _agent4 };
 		}
 
 		[TestMethod]
 		public async Task UtilizationFull()
 		{
-			IPool Pool = await PoolService.CreatePoolAsync("AutoscalePool1", null, true, 0, 0);
-			CreateAgents(Pool);
+			IPool pool = await PoolService.CreatePoolAsync("AutoscalePool1", null, true, 0, 0);
+			CreateAgents(pool);
 			
-			await AddPlaceholderLease(Agent1, Pool, Clock.UtcNow - TimeSpan.FromMinutes(120), TimeSpan.FromMinutes(120));
-			await AddPlaceholderLease(Agent2, Pool, Clock.UtcNow - TimeSpan.FromMinutes(120), TimeSpan.FromMinutes(120));
-			await AddPlaceholderLease(Agent3, Pool, Clock.UtcNow - TimeSpan.FromMinutes(120), TimeSpan.FromMinutes(120));
-			await AddPlaceholderLease(Agent4, Pool, Clock.UtcNow - TimeSpan.FromMinutes(120), TimeSpan.FromMinutes(120));
-			await AssertPoolSizeAsync(Pool, 4);
+			await AddPlaceholderLease(_agent1, pool, Clock.UtcNow - TimeSpan.FromMinutes(120), TimeSpan.FromMinutes(120));
+			await AddPlaceholderLease(_agent2, pool, Clock.UtcNow - TimeSpan.FromMinutes(120), TimeSpan.FromMinutes(120));
+			await AddPlaceholderLease(_agent3, pool, Clock.UtcNow - TimeSpan.FromMinutes(120), TimeSpan.FromMinutes(120));
+			await AddPlaceholderLease(_agent4, pool, Clock.UtcNow - TimeSpan.FromMinutes(120), TimeSpan.FromMinutes(120));
+			await AssertPoolSizeAsync(pool, 4);
 		}
 		
 		[TestMethod]
 		public async Task UtilizationHalf()
 		{
-			IPool Pool = await PoolService.CreatePoolAsync("AutoscalePool1", null, true, 0, 0);
-			CreateAgents(Pool);
+			IPool pool = await PoolService.CreatePoolAsync("AutoscalePool1", null, true, 0, 0);
+			CreateAgents(pool);
 			
-			await AddPlaceholderLease(Agent1, Pool, Clock.UtcNow - TimeSpan.FromMinutes(120), TimeSpan.FromMinutes(120));
-			await AddPlaceholderLease(Agent2, Pool, Clock.UtcNow - TimeSpan.FromMinutes(120), TimeSpan.FromMinutes(120));
-			await AssertPoolSizeAsync(Pool, 2);
+			await AddPlaceholderLease(_agent1, pool, Clock.UtcNow - TimeSpan.FromMinutes(120), TimeSpan.FromMinutes(120));
+			await AddPlaceholderLease(_agent2, pool, Clock.UtcNow - TimeSpan.FromMinutes(120), TimeSpan.FromMinutes(120));
+			await AssertPoolSizeAsync(pool, 2);
 		}
 		
 		[TestMethod]
 		public async Task UtilizationZero()
 		{
-			IPool Pool = await PoolService.CreatePoolAsync("AutoscalePool1", null, true, 0, 0);
-			CreateAgents(Pool);
+			IPool pool = await PoolService.CreatePoolAsync("AutoscalePool1", null, true, 0, 0);
+			CreateAgents(pool);
 
-			await AssertPoolSizeAsync(Pool, 0);
+			await AssertPoolSizeAsync(pool, 0);
 		}
 		
 		[TestMethod]
 		public async Task ReserveAgents()
 		{
-			IPool Pool = await PoolService.CreatePoolAsync("AutoscalePool1", null, true, 0, 5);
-			CreateAgents(Pool);
+			IPool pool = await PoolService.CreatePoolAsync("AutoscalePool1", null, true, 0, 5);
+			CreateAgents(pool);
 			
-			await AddPlaceholderLease(Agent1, Pool, Clock.UtcNow - TimeSpan.FromMinutes(120), TimeSpan.FromMinutes(120));
-			await AddPlaceholderLease(Agent2, Pool, Clock.UtcNow - TimeSpan.FromMinutes(120), TimeSpan.FromMinutes(120));
-			await AddPlaceholderLease(Agent3, Pool, Clock.UtcNow - TimeSpan.FromMinutes(120), TimeSpan.FromMinutes(120));
-			await AddPlaceholderLease(Agent4, Pool, Clock.UtcNow - TimeSpan.FromMinutes(120), TimeSpan.FromMinutes(120));
+			await AddPlaceholderLease(_agent1, pool, Clock.UtcNow - TimeSpan.FromMinutes(120), TimeSpan.FromMinutes(120));
+			await AddPlaceholderLease(_agent2, pool, Clock.UtcNow - TimeSpan.FromMinutes(120), TimeSpan.FromMinutes(120));
+			await AddPlaceholderLease(_agent3, pool, Clock.UtcNow - TimeSpan.FromMinutes(120), TimeSpan.FromMinutes(120));
+			await AddPlaceholderLease(_agent4, pool, Clock.UtcNow - TimeSpan.FromMinutes(120), TimeSpan.FromMinutes(120));
 			
 			// Full utilization should mean all agents plus the reserve agents
-			await AssertPoolSizeAsync(Pool, 9);
+			await AssertPoolSizeAsync(pool, 9);
 		}
 		
 		[TestMethod]
 		public async Task MinAgents()
 		{
-			IPool Pool = await PoolService.CreatePoolAsync("AutoscalePool1", null, true, 2, 0);
-			CreateAgents(Pool);
+			IPool pool = await PoolService.CreatePoolAsync("AutoscalePool1", null, true, 2, 0);
+			CreateAgents(pool);
 			
 			// Even with no utilization, expect at least the min number of agents
-			await AssertPoolSizeAsync(Pool, 2);
+			await AssertPoolSizeAsync(pool, 2);
 		}
 
-		private async Task AssertPoolSizeAsync(IPool Pool, int ExpectedNumAgents)
+		private async Task AssertPoolSizeAsync(IPool pool, int expectedNumAgents)
 		{
-			LeaseUtilizationStrategy Strategy = new (AgentCollection, PoolCollection, LeaseCollection, Clock);
-			List<PoolSizeData> Output = await Strategy.CalcDesiredPoolSizesAsync(new() { new(Pool, PoolAgents, null) });
-			Assert.AreEqual(1, Output.Count);
-			Assert.AreEqual(ExpectedNumAgents, Output[0].DesiredAgentCount);
+			LeaseUtilizationStrategy strategy = new (AgentCollection, PoolCollection, LeaseCollection, Clock);
+			List<PoolSizeData> output = await strategy.CalcDesiredPoolSizesAsync(new() { new(pool, _poolAgents, null) });
+			Assert.AreEqual(1, output.Count);
+			Assert.AreEqual(expectedNumAgents, output[0].DesiredAgentCount);
 		}
 		
-		private async Task<ILease> AddPlaceholderLease(IAgent Agent, IPool Pool, DateTime StartTime, TimeSpan Duration)
+		private async Task<ILease> AddPlaceholderLease(IAgent agent, IPool pool, DateTime startTime, TimeSpan duration)
 		{
-			Assert.IsNotNull(Agent.SessionId);
+			Assert.IsNotNull(agent.SessionId);
 			
-			ExecuteJobTask PlaceholderJobTask = new();
-			PlaceholderJobTask.JobName = "placeholderJobName";
-			byte[] Payload = Any.Pack(PlaceholderJobTask).ToByteArray();
+			ExecuteJobTask placeholderJobTask = new();
+			placeholderJobTask.JobName = "placeholderJobName";
+			byte[] payload = Any.Pack(placeholderJobTask).ToByteArray();
 
-			ILease Lease = await LeaseCollection.AddAsync(ObjectId<ILease>.GenerateNewId(), "placeholderLease", Agent.Id, Agent.SessionId!.Value, new StreamId("placeholderStream"), Pool.Id, null, StartTime, Payload);
-			bool WasModified = await LeaseCollection.TrySetOutcomeAsync(Lease.Id, StartTime + Duration, LeaseOutcome.Success, null);
-			Assert.IsTrue(WasModified);
-			return Lease;
+			ILease lease = await LeaseCollection.AddAsync(ObjectId<ILease>.GenerateNewId(), "placeholderLease", agent.Id, agent.SessionId!.Value, new StreamId("placeholderStream"), pool.Id, null, startTime, payload);
+			bool wasModified = await LeaseCollection.TrySetOutcomeAsync(lease.Id, startTime + duration, LeaseOutcome.Success, null);
+			Assert.IsTrue(wasModified);
+			return lease;
 		}
 	}
 }
\ No newline at end of file
diff --git a/Engine/Source/Programs/Horde/Horde.Build.Tests/GrpcTestHelpers.cs b/Engine/Source/Programs/Horde/Horde.Build.Tests/GrpcTestHelpers.cs
--- a/Engine/Source/Programs/Horde/Horde.Build.Tests/GrpcTestHelpers.cs
+++ b/Engine/Source/Programs/Horde/Horde.Build.Tests/GrpcTestHelpers.cs
@@ -110,9 +110,9 @@ namespace Horde.Build.Tests
 		{
 			_serverCallContext.CancellationToken.ThrowIfCancellationRequested();
 
-			if (await _channel.Reader.WaitToReadAsync())
+			if (await _channel.Reader.WaitToReadAsync(cancellationToken))
 			{
-				if (_channel.Reader.TryRead(out var message))
+				if (_channel.Reader.TryRead(out T? message))
 				{
 					Current = message;
 					return true;
diff --git a/Engine/Source/Programs/Horde/Horde.Build.Tests/Horde.Build.Tests.csproj b/Engine/Source/Programs/Horde/Horde.Build.Tests/Horde.Build.Tests.csproj
--- a/Engine/Source/Programs/Horde/Horde.Build.Tests/Horde.Build.Tests.csproj
+++ b/Engine/Source/Programs/Horde/Horde.Build.Tests/Horde.Build.Tests.csproj
@@ -4,32 +4,33 @@
     <TargetFramework>net6.0</TargetFramework>
     <Nullable>enable</Nullable>
     <IsPackable>false</IsPackable>
+    <NeutralLanguage>en-US</NeutralLanguage>
   </PropertyGroup>
 
   <ItemGroup>
-    <Compile Remove="Stubs\Collections\JobCollectionStub.cs" />
-    <Compile Remove="Stubs\Collections\StreamCollectionStub.cs" />
+    <Compile Remove="Stubs\Collections\JobCollectionStub.cs"/>
+    <Compile Remove="Stubs\Collections\StreamCollectionStub.cs"/>
   </ItemGroup>
 
   <ItemGroup>
-    <None Include="Stubs\Collections\JobCollectionStub.cs" />
-    <None Include="Stubs\Collections\StreamCollectionStub.cs" />
+    <None Include="Stubs\Collections\JobCollectionStub.cs"/>
+    <None Include="Stubs\Collections\StreamCollectionStub.cs"/>
   </ItemGroup>
 
   <ItemGroup>
-    <PackageReference Include="Microsoft.AspNetCore.Mvc.Testing" Version="3.1.6" />
-    <PackageReference Include="Microsoft.NET.Test.Sdk" Version="16.5.0" />
-    <PackageReference Include="Moq" Version="4.14.5" />
-    <PackageReference Include="MSTest.TestAdapter" Version="2.1.0" />
-    <PackageReference Include="MSTest.TestFramework" Version="2.1.0" />
-    <PackageReference Include="Microsoft.Extensions.Caching.Memory" Version="6.0.0" />
-    <PackageReference Include="coverlet.collector" Version="1.2.0" />
+    <PackageReference Include="Microsoft.AspNetCore.Mvc.Testing" Version="3.1.6"/>
+    <PackageReference Include="Microsoft.NET.Test.Sdk" Version="16.5.0"/>
+    <PackageReference Include="Moq" Version="4.14.5"/>
+    <PackageReference Include="MSTest.TestAdapter" Version="2.1.0"/>
+    <PackageReference Include="MSTest.TestFramework" Version="2.1.0"/>
+    <PackageReference Include="Microsoft.Extensions.Caching.Memory" Version="6.0.0"/>
+    <PackageReference Include="coverlet.collector" Version="1.2.0"/>
   </ItemGroup>
 
   <ItemGroup>
-    <ProjectReference Include="..\..\Shared\EpicGames.Serialization\EpicGames.Serialization.csproj" />
-    <ProjectReference Include="..\Horde.Agent\Horde.Agent.csproj" />
-    <ProjectReference Include="..\Horde.Build\Horde.Build.csproj" />
+    <ProjectReference Include="..\..\Shared\EpicGames.Serialization\EpicGames.Serialization.csproj"/>
+    <ProjectReference Include="..\Horde.Agent\Horde.Agent.csproj"/>
+    <ProjectReference Include="..\Horde.Build\Horde.Build.csproj"/>
   </ItemGroup>
 
   <ItemGroup>
@@ -48,7 +49,7 @@
   </ItemGroup>
 
   <ItemGroup>
-    <Folder Include="Resources\" />
+    <Folder Include="Resources\"/>
   </ItemGroup>
 
 </Project>
