Horde: Update Horde.Build.Tests to use standard coding conventions, and fix a number of issues caught by static analysis.

diff --git a/Engine/Source/Programs/Horde/Horde.Build.Tests/JobCollectionTests.cs b/Engine/Source/Programs/Horde/Horde.Build.Tests/JobCollectionTests.cs
--- a/Engine/Source/Programs/Horde/Horde.Build.Tests/JobCollectionTests.cs
+++ b/Engine/Source/Programs/Horde/Horde.Build.Tests/JobCollectionTests.cs
@@ -1,12 +1,8 @@
 // Copyright Epic Games, Inc. All Rights Reserved.
 
-using Horde.Build.Collections;
 using MongoDB.Bson;
 using System;
 using System.Collections.Generic;
-using System.Text;
-
-using ProjectId = Horde.Build.Utilities.StringId<Horde.Build.Models.IProject>;
 using PoolId = Horde.Build.Utilities.StringId<Horde.Build.Models.IPool>;
 using StreamId = Horde.Build.Utilities.StringId<Horde.Build.Models.IStream>;
 using TemplateRefId = Horde.Build.Utilities.StringId<Horde.Build.Models.TemplateRef>;
@@ -15,13 +11,10 @@ using Microsoft.VisualStudio.TestTools.UnitTesting;
 using EpicGames.Core;
 using Moq;
 using Horde.Build.Models;
-using Horde.Build.Services;
 using System.Threading.Tasks;
 using Horde.Build.Api;
 using System.Linq;
 using HordeCommon;
-using Horde.Build.Collections.Impl;
-using Microsoft.Extensions.Logging.Abstractions;
 using Horde.Build.Utilities;
 
 namespace Horde.Build.Tests
@@ -33,103 +26,103 @@ namespace Horde.Build.Tests
 	[TestClass]
 	public class JobCollectionTests : TestSetup
 	{
-		NewGroup AddGroup(List<NewGroup> Groups)
+		static NewGroup AddGroup(List<NewGroup> groups)
 		{
-			NewGroup Group = new NewGroup("win64", new List<NewNode>());
-			Groups.Add(Group);
-			return Group;
+			NewGroup @group = new NewGroup("win64", new List<NewNode>());
+			groups.Add(@group);
+			return @group;
 		}
 
-		NewNode AddNode(NewGroup Group, string Name, string[]? InputDependencies, Action<NewNode>? Action = null)
+		static NewNode AddNode(NewGroup @group, string name, string[]? inputDependencies, Action<NewNode>? action = null)
 		{
-			NewNode Node = new NewNode(Name, InputDependencies?.ToList(), InputDependencies?.ToList(), null, null, null, null, null, null);
-			if (Action != null)
+			NewNode node = new NewNode(name, inputDependencies?.ToList(), inputDependencies?.ToList(), null, null, null, null, null, null);
+			if (action != null)
 			{
-				Action.Invoke(Node);
+				action.Invoke(node);
 			}
-			Group.Nodes.Add(Node);
-			return Node;
+			@group.Nodes.Add(node);
+			return node;
 		}
 
-		async Task<IJob> StartBatch(IJob Job, IGraph Graph, int BatchIdx)
+		async Task<IJob> StartBatch(IJob job, IGraph graph, int batchIdx)
 		{
-			Assert.AreEqual(JobStepBatchState.Ready, Job.Batches[BatchIdx].State);
-			Job = Deref(await JobCollection.TryUpdateBatchAsync(Job, Graph, Job.Batches[BatchIdx].Id, null, JobStepBatchState.Running, null));
-			Assert.AreEqual(JobStepBatchState.Running, Job.Batches[BatchIdx].State);
-			return Job;
+			Assert.AreEqual(JobStepBatchState.Ready, job.Batches[batchIdx].State);
+			job = Deref(await JobCollection.TryUpdateBatchAsync(job, graph, job.Batches[batchIdx].Id, null, JobStepBatchState.Running, null));
+			Assert.AreEqual(JobStepBatchState.Running, job.Batches[batchIdx].State);
+			return job;
 		}
 
-		async Task<IJob> RunStep(IJob Job, IGraph Graph, int BatchIdx, int StepIdx, JobStepOutcome Outcome)
+		async Task<IJob> RunStep(IJob job, IGraph graph, int batchIdx, int stepIdx, JobStepOutcome outcome)
 		{
-			Assert.AreEqual(JobStepState.Ready, Job.Batches[BatchIdx].Steps[StepIdx].State);
-			Job = Deref(await JobCollection.TryUpdateStepAsync(Job, Graph, Job.Batches[BatchIdx].Id, Job.Batches[BatchIdx].Steps[StepIdx].Id, JobStepState.Running, JobStepOutcome.Success));
-			Assert.AreEqual(JobStepState.Running, Job.Batches[BatchIdx].Steps[StepIdx].State);
-			Job = Deref(await JobCollection.TryUpdateStepAsync(Job, Graph, Job.Batches[BatchIdx].Id, Job.Batches[BatchIdx].Steps[StepIdx].Id, JobStepState.Completed, Outcome));
-			Assert.AreEqual(JobStepState.Completed, Job.Batches[BatchIdx].Steps[StepIdx].State);
-			Assert.AreEqual(Outcome, Job.Batches[BatchIdx].Steps[StepIdx].Outcome);
-			return Job;
+			Assert.AreEqual(JobStepState.Ready, job.Batches[batchIdx].Steps[stepIdx].State);
+			job = Deref(await JobCollection.TryUpdateStepAsync(job, graph, job.Batches[batchIdx].Id, job.Batches[batchIdx].Steps[stepIdx].Id, JobStepState.Running, JobStepOutcome.Success));
+			Assert.AreEqual(JobStepState.Running, job.Batches[batchIdx].Steps[stepIdx].State);
+			job = Deref(await JobCollection.TryUpdateStepAsync(job, graph, job.Batches[batchIdx].Id, job.Batches[batchIdx].Steps[stepIdx].Id, JobStepState.Completed, outcome));
+			Assert.AreEqual(JobStepState.Completed, job.Batches[batchIdx].Steps[stepIdx].State);
+			Assert.AreEqual(outcome, job.Batches[batchIdx].Steps[stepIdx].Outcome);
+			return job;
 		}
 
 		[TestMethod]
 		public async Task TestStates()
 		{
-			Mock<ITemplate> TemplateMock = new Mock<ITemplate>(MockBehavior.Strict);
-			TemplateMock.SetupGet(x => x.InitialAgentType).Returns((string?)null);
+			Mock<ITemplate> templateMock = new Mock<ITemplate>(MockBehavior.Strict);
+			templateMock.SetupGet(x => x.InitialAgentType).Returns((string?)null);
 
-			IGraph BaseGraph = await GraphCollection.AddAsync(TemplateMock.Object);
+			IGraph baseGraph = await GraphCollection.AddAsync(templateMock.Object);
 
-			List<string> Arguments = new List<string>();
-			Arguments.Add("-Target=Publish Client");
-			Arguments.Add("-Target=Post-Publish Client");
+			List<string> arguments = new List<string>();
+			arguments.Add("-Target=Publish Client");
+			arguments.Add("-Target=Post-Publish Client");
 
-			IJob Job = await JobCollection.AddAsync(JobId.GenerateNewId(), new StreamId("ue4-main"), new TemplateRefId("test-build"), ContentHash.SHA1("hello"), BaseGraph, "Test job", 123, 123, null, null, null, null, null, null, null, null, false, false, null, null, Arguments);
+			IJob job = await JobCollection.AddAsync(JobId.GenerateNewId(), new StreamId("ue4-main"), new TemplateRefId("test-build"), ContentHash.SHA1("hello"), baseGraph, "Test job", 123, 123, null, null, null, null, null, null, null, null, false, false, null, null, arguments);
 
-			Job = await StartBatch(Job, BaseGraph, 0);
-			Job = await RunStep(Job, BaseGraph, 0, 0, JobStepOutcome.Success); // Setup Build
+			job = await StartBatch(job, baseGraph, 0);
+			job = await RunStep(job, baseGraph, 0, 0, JobStepOutcome.Success); // Setup Build
 
-			List<NewGroup> NewGroups = new List<NewGroup>();
+			List<NewGroup> newGroups = new List<NewGroup>();
 
-			NewGroup InitialGroup = AddGroup(NewGroups);
-			AddNode(InitialGroup, "Update Version Files", null);
-			AddNode(InitialGroup, "Compile Editor", new[] { "Update Version Files" });
+			NewGroup initialGroup = AddGroup(newGroups);
+			AddNode(initialGroup, "Update Version Files", null);
+			AddNode(initialGroup, "Compile Editor", new[] { "Update Version Files" });
 
-			NewGroup CompileGroup = AddGroup(NewGroups);
-			AddNode(CompileGroup, "Compile Client", new[] { "Update Version Files" });
+			NewGroup compileGroup = AddGroup(newGroups);
+			AddNode(compileGroup, "Compile Client", new[] { "Update Version Files" });
 
-			NewGroup PublishGroup = AddGroup(NewGroups);
-			AddNode(PublishGroup, "Cook Client", new[] { "Compile Editor" }, x => x.RunEarly = true);
-			AddNode(PublishGroup, "Publish Client", new[] { "Compile Client", "Cook Client" });
-			AddNode(PublishGroup, "Post-Publish Client", null, x => x.OrderDependencies = new List<string> { "Publish Client" });
+			NewGroup publishGroup = AddGroup(newGroups);
+			AddNode(publishGroup, "Cook Client", new[] { "Compile Editor" }, x => x.RunEarly = true);
+			AddNode(publishGroup, "Publish Client", new[] { "Compile Client", "Cook Client" });
+			AddNode(publishGroup, "Post-Publish Client", null, x => x.OrderDependencies = new List<string> { "Publish Client" });
 
-			IGraph Graph = await GraphCollection.AppendAsync(BaseGraph, NewGroups, null, null);
-			Job = Deref(await JobCollection.TryUpdateGraphAsync(Job, Graph));
+			IGraph graph = await GraphCollection.AppendAsync(baseGraph, newGroups, null, null);
+			job = Deref(await JobCollection.TryUpdateGraphAsync(job, graph));
 
-			Job = await StartBatch(Job, Graph, 1);
-			Job = await RunStep(Job, Graph, 1, 0, JobStepOutcome.Success); // Update Version Files
-			Job = await RunStep(Job, Graph, 1, 1, JobStepOutcome.Success); // Compile Editor
+			job = await StartBatch(job, graph, 1);
+			job = await RunStep(job, graph, 1, 0, JobStepOutcome.Success); // Update Version Files
+			job = await RunStep(job, graph, 1, 1, JobStepOutcome.Success); // Compile Editor
 
-			Job = await StartBatch(Job, Graph, 2);
-			Job = await RunStep(Job, Graph, 2, 0, JobStepOutcome.Success); // Compile Client
+			job = await StartBatch(job, graph, 2);
+			job = await RunStep(job, graph, 2, 0, JobStepOutcome.Success); // Compile Client
 
-			Job = await StartBatch(Job, Graph, 3);
-			Job = await RunStep(Job, Graph, 3, 0, JobStepOutcome.Failure); // Cook Client
-			Assert.AreEqual(JobStepState.Skipped, Job.Batches[3].Steps[1].State); // Publish Client
-			Assert.AreEqual(JobStepState.Skipped, Job.Batches[3].Steps[2].State); // Post-Publish Client
+			job = await StartBatch(job, graph, 3);
+			job = await RunStep(job, graph, 3, 0, JobStepOutcome.Failure); // Cook Client
+			Assert.AreEqual(JobStepState.Skipped, job.Batches[3].Steps[1].State); // Publish Client
+			Assert.AreEqual(JobStepState.Skipped, job.Batches[3].Steps[2].State); // Post-Publish Client
 		}
 
 		[TestMethod]
 		public async Task TryAssignLeaseTest()
 		{
-			Fixture Fixture = await CreateFixtureAsync();
+			Fixture fixture = await CreateFixtureAsync();
 
-			ObjectId<ISession> SessionId1 = new (ObjectId.GenerateNewId());
-			await JobCollection.TryAssignLeaseAsync(Fixture.Job1, 0, new PoolId("foo"), Fixture.Agent1.Id,
-				SessionId1, LeaseId.GenerateNewId(), LogId.GenerateNewId());
+			ObjectId<ISession> sessionId1 = new (ObjectId.GenerateNewId());
+			await JobCollection.TryAssignLeaseAsync(fixture.Job1, 0, new PoolId("foo"), fixture.Agent1.Id,
+				sessionId1, LeaseId.GenerateNewId(), LogId.GenerateNewId());
 			
-			ObjectId<ISession> SessionId2 = new (ObjectId.GenerateNewId());
-			IJob Job = (await JobCollection.GetAsync(Fixture.Job1.Id))!;
-			await JobCollection.TryAssignLeaseAsync(Job, 0, new PoolId("foo"), Fixture.Agent1.Id,
-				SessionId2, LeaseId.GenerateNewId(), LogId.GenerateNewId());
+			ObjectId<ISession> sessionId2 = new (ObjectId.GenerateNewId());
+			IJob job = (await JobCollection.GetAsync(fixture.Job1.Id))!;
+			await JobCollection.TryAssignLeaseAsync(job, 0, new PoolId("foo"), fixture.Agent1.Id,
+				sessionId2, LeaseId.GenerateNewId(), LogId.GenerateNewId());
 			
 			// Manually verify the log output
 		}
@@ -146,108 +139,108 @@ namespace Horde.Build.Tests
 			return LostLeaseTestInternal(false);
 		}
 
-		public async Task LostLeaseTestInternal(bool HasDependency)
+		public async Task LostLeaseTestInternal(bool hasDependency)
 		{
-			Mock<ITemplate> TemplateMock = new Mock<ITemplate>(MockBehavior.Strict);
-			TemplateMock.SetupGet(x => x.InitialAgentType).Returns((string?)null);
+			Mock<ITemplate> templateMock = new Mock<ITemplate>(MockBehavior.Strict);
+			templateMock.SetupGet(x => x.InitialAgentType).Returns((string?)null);
 
-			IGraph BaseGraph = await GraphCollection.AddAsync(TemplateMock.Object);
+			IGraph baseGraph = await GraphCollection.AddAsync(templateMock.Object);
 
-			List<string> Arguments = new List<string>();
-			Arguments.Add("-Target=Step 1");
-			Arguments.Add("-Target=Step 3");
+			List<string> arguments = new List<string>();
+			arguments.Add("-Target=Step 1");
+			arguments.Add("-Target=Step 3");
 
-			IJob Job = await JobCollection.AddAsync(JobId.GenerateNewId(), new StreamId("ue4-main"), new TemplateRefId("test-build"), ContentHash.SHA1("hello"), BaseGraph, "Test job", 123, 123, null, null, null, null, null, null, null, null, false, false, null, null, Arguments);
+			IJob job = await JobCollection.AddAsync(JobId.GenerateNewId(), new StreamId("ue4-main"), new TemplateRefId("test-build"), ContentHash.SHA1("hello"), baseGraph, "Test job", 123, 123, null, null, null, null, null, null, null, null, false, false, null, null, arguments);
 
-			Job = await StartBatch(Job, BaseGraph, 0);
-			Job = await RunStep(Job, BaseGraph, 0, 0, JobStepOutcome.Success); // Setup Build
+			job = await StartBatch(job, baseGraph, 0);
+			job = await RunStep(job, baseGraph, 0, 0, JobStepOutcome.Success); // Setup Build
 
-			List<NewGroup> NewGroups = new List<NewGroup>();
+			List<NewGroup> newGroups = new List<NewGroup>();
 
-			NewGroup InitialGroup = AddGroup(NewGroups);
-			AddNode(InitialGroup, "Step 1", null);
-			AddNode(InitialGroup, "Step 2", HasDependency? new[] { "Step 1" } : null);
-			AddNode(InitialGroup, "Step 3", new[] { "Step 2" });
+			NewGroup initialGroup = AddGroup(newGroups);
+			AddNode(initialGroup, "Step 1", null);
+			AddNode(initialGroup, "Step 2", hasDependency? new[] { "Step 1" } : null);
+			AddNode(initialGroup, "Step 3", new[] { "Step 2" });
 
-			IGraph Graph = await GraphCollection.AppendAsync(BaseGraph, NewGroups, null, null);
-			Job = Deref(await JobCollection.TryUpdateGraphAsync(Job, Graph));
+			IGraph graph = await GraphCollection.AppendAsync(baseGraph, newGroups, null, null);
+			job = Deref(await JobCollection.TryUpdateGraphAsync(job, graph));
 
-			Job = await StartBatch(Job, Graph, 1);
-			Job = await RunStep(Job, Graph, 1, 0, JobStepOutcome.Success); // Step 1
-			Job = await RunStep(Job, Graph, 1, 1, JobStepOutcome.Success); // Step 2
+			job = await StartBatch(job, graph, 1);
+			job = await RunStep(job, graph, 1, 0, JobStepOutcome.Success); // Step 1
+			job = await RunStep(job, graph, 1, 1, JobStepOutcome.Success); // Step 2
 
 			// Force an error executing the batch
-			Job = Deref(await JobCollection.TryUpdateBatchAsync(Job, Graph, Job.Batches[1].Id, null, JobStepBatchState.Complete, JobStepBatchError.Incomplete));
+			job = Deref(await JobCollection.TryUpdateBatchAsync(job, graph, job.Batches[1].Id, null, JobStepBatchState.Complete, JobStepBatchError.Incomplete));
 
 			// Check that it restarted all three nodes
-			IJob NewJob = (await JobCollection.GetAsync(Job.Id))!;
-			Assert.AreEqual(3, NewJob.Batches.Count);
-			Assert.AreEqual(1, NewJob.Batches[2].GroupIdx);
+			IJob newJob = (await JobCollection.GetAsync(job.Id))!;
+			Assert.AreEqual(3, newJob.Batches.Count);
+			Assert.AreEqual(1, newJob.Batches[2].GroupIdx);
 
-			if (HasDependency)
+			if (hasDependency)
 			{
-				Assert.AreEqual(3, NewJob.Batches[2].Steps.Count);
+				Assert.AreEqual(3, newJob.Batches[2].Steps.Count);
 
-				Assert.AreEqual(0, NewJob.Batches[2].Steps[0].NodeIdx);
-				Assert.AreEqual(1, NewJob.Batches[2].Steps[1].NodeIdx);
-				Assert.AreEqual(2, NewJob.Batches[2].Steps[2].NodeIdx);
+				Assert.AreEqual(0, newJob.Batches[2].Steps[0].NodeIdx);
+				Assert.AreEqual(1, newJob.Batches[2].Steps[1].NodeIdx);
+				Assert.AreEqual(2, newJob.Batches[2].Steps[2].NodeIdx);
 			}
 			else
 			{
-				Assert.AreEqual(2, NewJob.Batches[2].Steps.Count);
+				Assert.AreEqual(2, newJob.Batches[2].Steps.Count);
 
-				Assert.AreEqual(1, NewJob.Batches[2].Steps[0].NodeIdx);
-				Assert.AreEqual(2, NewJob.Batches[2].Steps[1].NodeIdx);
+				Assert.AreEqual(1, newJob.Batches[2].Steps[0].NodeIdx);
+				Assert.AreEqual(2, newJob.Batches[2].Steps[1].NodeIdx);
 			}
 		}
 
 		[TestMethod]
 		public async Task IncompleteBatchAsync()
 		{
-			Mock<ITemplate> TemplateMock = new Mock<ITemplate>(MockBehavior.Strict);
-			TemplateMock.SetupGet(x => x.InitialAgentType).Returns((string?)null);
-
-			IGraph BaseGraph = await GraphCollection.AddAsync(TemplateMock.Object);
-
-			List<string> Arguments = new List<string>();
-			Arguments.Add("-Target=Step 1");
-			Arguments.Add("-Target=Step 3");
-
-			IJob Job = await JobCollection.AddAsync(JobId.GenerateNewId(), new StreamId("ue4-main"), new TemplateRefId("test-build"), ContentHash.SHA1("hello"), BaseGraph, "Test job", 123, 123, null, null, null, null, null, null, null, null, false, false, null, null, Arguments);
-			Assert.AreEqual(1, Job.Batches.Count);
-
-			Job = await StartBatch(Job, BaseGraph, 0);
-			Job = Deref(await JobCollection.TryUpdateBatchAsync(Job, BaseGraph, Job.Batches[0].Id, null, JobStepBatchState.Complete, JobStepBatchError.Incomplete));
-
-			Job = (await JobCollection.GetAsync(Job.Id))!;
-			Assert.AreEqual(2, Job.Batches.Count);
-			Assert.AreEqual(JobStepBatchState.Complete, Job.Batches[0].State);
-			Assert.AreEqual(0, Job.Batches[0].Steps.Count);
-			Assert.AreEqual(JobStepBatchState.Ready, Job.Batches[1].State);
-			Assert.AreEqual(1, Job.Batches[1].Steps.Count);
-
-			Job = Deref(await JobCollection.TryUpdateBatchAsync(Job, BaseGraph, Job.Batches[1].Id, null, JobStepBatchState.Complete, JobStepBatchError.Incomplete));
-
-			Job = (await JobCollection.GetAsync(Job.Id))!;
-			Assert.AreEqual(3, Job.Batches.Count);
-			Assert.AreEqual(JobStepBatchState.Complete, Job.Batches[0].State);
-			Assert.AreEqual(0, Job.Batches[0].Steps.Count);
-			Assert.AreEqual(JobStepBatchState.Complete, Job.Batches[1].State);
-			Assert.AreEqual(0, Job.Batches[1].Steps.Count);
-			Assert.AreEqual(JobStepBatchState.Ready, Job.Batches[2].State);
-			Assert.AreEqual(1, Job.Batches[2].Steps.Count);
-
-			Job = Deref(await JobCollection.TryUpdateBatchAsync(Job, BaseGraph, Job.Batches[2].Id, null, JobStepBatchState.Complete, JobStepBatchError.Incomplete));
-
-			Job = (await JobCollection.GetAsync(Job.Id))!;
-			Assert.AreEqual(3, Job.Batches.Count);
-			Assert.AreEqual(JobStepBatchState.Complete, Job.Batches[0].State);
-			Assert.AreEqual(0, Job.Batches[0].Steps.Count);
-			Assert.AreEqual(JobStepBatchState.Complete, Job.Batches[1].State);
-			Assert.AreEqual(0, Job.Batches[1].Steps.Count);
-			Assert.AreEqual(JobStepBatchState.Complete, Job.Batches[2].State);
-			Assert.AreEqual(1, Job.Batches[2].Steps.Count);
-			Assert.AreEqual(JobStepState.Skipped, Job.Batches[2].Steps[0].State);
+			Mock<ITemplate> templateMock = new Mock<ITemplate>(MockBehavior.Strict);
+			templateMock.SetupGet(x => x.InitialAgentType).Returns((string?)null);
+
+			IGraph baseGraph = await GraphCollection.AddAsync(templateMock.Object);
+
+			List<string> arguments = new List<string>();
+			arguments.Add("-Target=Step 1");
+			arguments.Add("-Target=Step 3");
+
+			IJob job = await JobCollection.AddAsync(JobId.GenerateNewId(), new StreamId("ue4-main"), new TemplateRefId("test-build"), ContentHash.SHA1("hello"), baseGraph, "Test job", 123, 123, null, null, null, null, null, null, null, null, false, false, null, null, arguments);
+			Assert.AreEqual(1, job.Batches.Count);
+
+			job = await StartBatch(job, baseGraph, 0);
+			job = Deref(await JobCollection.TryUpdateBatchAsync(job, baseGraph, job.Batches[0].Id, null, JobStepBatchState.Complete, JobStepBatchError.Incomplete));
+
+			job = (await JobCollection.GetAsync(job.Id))!;
+			Assert.AreEqual(2, job.Batches.Count);
+			Assert.AreEqual(JobStepBatchState.Complete, job.Batches[0].State);
+			Assert.AreEqual(0, job.Batches[0].Steps.Count);
+			Assert.AreEqual(JobStepBatchState.Ready, job.Batches[1].State);
+			Assert.AreEqual(1, job.Batches[1].Steps.Count);
+
+			job = Deref(await JobCollection.TryUpdateBatchAsync(job, baseGraph, job.Batches[1].Id, null, JobStepBatchState.Complete, JobStepBatchError.Incomplete));
+
+			job = (await JobCollection.GetAsync(job.Id))!;
+			Assert.AreEqual(3, job.Batches.Count);
+			Assert.AreEqual(JobStepBatchState.Complete, job.Batches[0].State);
+			Assert.AreEqual(0, job.Batches[0].Steps.Count);
+			Assert.AreEqual(JobStepBatchState.Complete, job.Batches[1].State);
+			Assert.AreEqual(0, job.Batches[1].Steps.Count);
+			Assert.AreEqual(JobStepBatchState.Ready, job.Batches[2].State);
+			Assert.AreEqual(1, job.Batches[2].Steps.Count);
+
+			job = Deref(await JobCollection.TryUpdateBatchAsync(job, baseGraph, job.Batches[2].Id, null, JobStepBatchState.Complete, JobStepBatchError.Incomplete));
+
+			job = (await JobCollection.GetAsync(job.Id))!;
+			Assert.AreEqual(3, job.Batches.Count);
+			Assert.AreEqual(JobStepBatchState.Complete, job.Batches[0].State);
+			Assert.AreEqual(0, job.Batches[0].Steps.Count);
+			Assert.AreEqual(JobStepBatchState.Complete, job.Batches[1].State);
+			Assert.AreEqual(0, job.Batches[1].Steps.Count);
+			Assert.AreEqual(JobStepBatchState.Complete, job.Batches[2].State);
+			Assert.AreEqual(1, job.Batches[2].Steps.Count);
+			Assert.AreEqual(JobStepState.Skipped, job.Batches[2].Steps[0].State);
 		}
 	}
 }
diff --git a/Engine/Source/Programs/Horde/Horde.Build.Tests/JobServiceTests.cs b/Engine/Source/Programs/Horde/Horde.Build.Tests/JobServiceTests.cs
--- a/Engine/Source/Programs/Horde/Horde.Build.Tests/JobServiceTests.cs
+++ b/Engine/Source/Programs/Horde/Horde.Build.Tests/JobServiceTests.cs
@@ -5,7 +5,6 @@ using HordeCommon;
 using Horde.Build.Api;
 using Horde.Build.Models;
 using Microsoft.VisualStudio.TestTools.UnitTesting;
-using MongoDB.Bson;
 using System.Collections.Generic;
 using System.Text;
 using System.Threading.Tasks;
@@ -26,87 +25,87 @@ namespace Horde.Build.Tests
 		[TestMethod]
 		public async Task TestChainedJobs()
 		{
-			ProjectId ProjectId = new ProjectId("ue5");
-			IProject? Project = await ProjectService.Collection.AddOrUpdateAsync(ProjectId, "", "", 0, new ProjectConfig { Name = "UE5" });
-			Assert.IsNotNull(Project);
+			ProjectId projectId = new ProjectId("ue5");
+			IProject? project = await ProjectService.Collection.AddOrUpdateAsync(projectId, "", "", 0, new ProjectConfig { Name = "UE5" });
+			Assert.IsNotNull(project);
 
-			ITemplate Template = await TemplateCollection.AddAsync("Test template");
-			IGraph Graph = await GraphCollection.AddAsync(Template);
+			ITemplate template = await TemplateCollection.AddAsync("Test template");
+			IGraph graph = await GraphCollection.AddAsync(template);
 
-			TemplateRefId TemplateRefId1 = new TemplateRefId("template1");
-			TemplateRefId TemplateRefId2 = new TemplateRefId("template2");
+			TemplateRefId templateRefId1 = new TemplateRefId("template1");
+			TemplateRefId templateRefId2 = new TemplateRefId("template2");
 
-			StreamConfig StreamConfig = new StreamConfig();
-			StreamConfig.Templates.Add(new CreateTemplateRefRequest { Id = TemplateRefId1.ToString(), Name = "Test Template", ChainedJobs = new List<CreateChainedJobTemplateRequest> { new CreateChainedJobTemplateRequest { TemplateId = TemplateRefId2.ToString(), Trigger = "Setup Build" } } });
-			StreamConfig.Templates.Add(new CreateTemplateRefRequest { Id = TemplateRefId2.ToString(), Name = "Test Template" });
-			StreamConfig.Tabs.Add(new CreateJobsTabRequest { Title = "foo", Templates = new List<string> { TemplateRefId1.ToString(), TemplateRefId2.ToString() } });
+			StreamConfig streamConfig = new StreamConfig();
+			streamConfig.Templates.Add(new CreateTemplateRefRequest { Id = templateRefId1.ToString(), Name = "Test Template", ChainedJobs = new List<CreateChainedJobTemplateRequest> { new CreateChainedJobTemplateRequest { TemplateId = templateRefId2.ToString(), Trigger = "Setup Build" } } });
+			streamConfig.Templates.Add(new CreateTemplateRefRequest { Id = templateRefId2.ToString(), Name = "Test Template" });
+			streamConfig.Tabs.Add(new CreateJobsTabRequest { Title = "foo", Templates = new List<string> { templateRefId1.ToString(), templateRefId2.ToString() } });
 
-			StreamId StreamId = new StreamId("ue5-main");
-			IStream? Stream = await StreamService.GetStreamAsync(StreamId);
-			Stream = await StreamService.StreamCollection.TryCreateOrReplaceAsync(new StreamId("ue5-main"), Stream, String.Empty, String.Empty, ProjectId, StreamConfig);
+			StreamId streamId = new StreamId("ue5-main");
+			IStream? stream = await StreamService.GetStreamAsync(streamId);
+			stream = await StreamService.StreamCollection.TryCreateOrReplaceAsync(new StreamId("ue5-main"), stream, String.Empty, String.Empty, projectId, streamConfig);
 
-			IJob Job = await JobService.CreateJobAsync(null, Stream!, TemplateRefId1, Template.Id, Graph, "Hello", 1234, 1233, 999, null, null, null, null, null, null, Stream!.Templates[TemplateRefId1].ChainedJobs, true, true, null, null, new List<string>());
-			Assert.AreEqual(1, Job.ChainedJobs.Count);
+			IJob job = await JobService.CreateJobAsync(null, stream!, templateRefId1, template.Id, graph, "Hello", 1234, 1233, 999, null, null, null, null, null, null, stream!.Templates[templateRefId1].ChainedJobs, true, true, null, null, new List<string>());
+			Assert.AreEqual(1, job.ChainedJobs.Count);
 
-			Job = Deref(await JobService.UpdateBatchAsync(Job, Job.Batches[0].Id, LogId.GenerateNewId(), JobStepBatchState.Running));
-			Job = Deref(await JobService.UpdateStepAsync(Job, Job.Batches[0].Id, Job.Batches[0].Steps[0].Id, JobStepState.Running));
-			Job = Deref(await JobService.UpdateStepAsync(Job, Job.Batches[0].Id, Job.Batches[0].Steps[0].Id, JobStepState.Completed, JobStepOutcome.Success));
+			job = Deref(await JobService.UpdateBatchAsync(job, job.Batches[0].Id, LogId.GenerateNewId(), JobStepBatchState.Running));
+			job = Deref(await JobService.UpdateStepAsync(job, job.Batches[0].Id, job.Batches[0].Steps[0].Id, JobStepState.Running));
+			job = Deref(await JobService.UpdateStepAsync(job, job.Batches[0].Id, job.Batches[0].Steps[0].Id, JobStepState.Completed, JobStepOutcome.Success));
 
-			Assert.IsNotNull(Job.ChainedJobs[0].JobId);
+			Assert.IsNotNull(job.ChainedJobs[0].JobId);
 
-			IJob? ChainedJob = await JobCollection.GetAsync(Job.ChainedJobs[0].JobId!.Value);
-			Assert.IsNotNull(ChainedJob);
-			Assert.AreEqual(ChainedJob!.Id, Job!.ChainedJobs[0].JobId);
+			IJob? chainedJob = await JobCollection.GetAsync(job.ChainedJobs[0].JobId!.Value);
+			Assert.IsNotNull(chainedJob);
+			Assert.AreEqual(chainedJob!.Id, job!.ChainedJobs[0].JobId);
 
-			Assert.AreEqual(ChainedJob!.Change, Job!.Change);
-			Assert.AreEqual(ChainedJob!.CodeChange, Job!.CodeChange);
-			Assert.AreEqual(ChainedJob!.PreflightChange, Job!.PreflightChange);
-			Assert.AreEqual(ChainedJob!.StartedByUserId, Job!.StartedByUserId);
+			Assert.AreEqual(chainedJob!.Change, job!.Change);
+			Assert.AreEqual(chainedJob!.CodeChange, job!.CodeChange);
+			Assert.AreEqual(chainedJob!.PreflightChange, job!.PreflightChange);
+			Assert.AreEqual(chainedJob!.StartedByUserId, job!.StartedByUserId);
 		}
 
 		[TestMethod]
 		public async Task StopAnyDuplicateJobsByPreflight()
 		{
-			Fixture Fixture = await CreateFixtureAsync();
-
-			string[] Args = {"-Target=bogus"};
-			IJob OrgJob = await CreatePreflightJob(Fixture, "tpl-ref-1", "tpl-hash-1", "elvis", 1000, Args);
-			IJob NewJob = await CreatePreflightJob(Fixture, "tpl-ref-1", "tpl-hash-1", "elvis", 1000, Args);
-			IJob DifferentTplRef = await CreatePreflightJob(Fixture, "tpl-ref-other", "tpl-hash-1", "elvis", 1000, Args);
-			IJob DifferentTplHash = await CreatePreflightJob(Fixture, "tpl-ref-1", "tpl-hash-other", "elvis", 1000, Args);
-			IJob DifferentUserName = await CreatePreflightJob(Fixture, "tpl-ref-1", "tpl-hash-1", "julia", 1000, Args);
-			IJob DifferentArgs = await CreatePreflightJob(Fixture, "tpl-ref-1", "tpl-hash-1", "elvis", 1000, new []{"-Target=other"});
+			Fixture fixture = await CreateFixtureAsync();
+
+			string[] args = {"-Target=bogus"};
+			IJob orgJob = await CreatePreflightJob(fixture, "tpl-ref-1", "tpl-hash-1", "elvis", 1000, args);
+			IJob newJob = await CreatePreflightJob(fixture, "tpl-ref-1", "tpl-hash-1", "elvis", 1000, args);
+			IJob differentTplRef = await CreatePreflightJob(fixture, "tpl-ref-other", "tpl-hash-1", "elvis", 1000, args);
+			IJob differentTplHash = await CreatePreflightJob(fixture, "tpl-ref-1", "tpl-hash-other", "elvis", 1000, args);
+			IJob differentUserName = await CreatePreflightJob(fixture, "tpl-ref-1", "tpl-hash-1", "julia", 1000, args);
+			IJob differentArgs = await CreatePreflightJob(fixture, "tpl-ref-1", "tpl-hash-1", "elvis", 1000, new []{"-Target=other"});
 			
-			OrgJob = (await JobService.GetJobAsync(OrgJob.Id))!;
-			NewJob = (await JobService.GetJobAsync(NewJob.Id))!;
-			DifferentTplRef = (await JobService.GetJobAsync(DifferentTplRef.Id))!;
-			DifferentTplHash = (await JobService.GetJobAsync(DifferentTplHash.Id))!;
-			DifferentUserName = (await JobService.GetJobAsync(DifferentUserName.Id))!;
-			DifferentArgs = (await JobService.GetJobAsync(DifferentArgs.Id))!;
+			orgJob = (await JobService.GetJobAsync(orgJob.Id))!;
+			newJob = (await JobService.GetJobAsync(newJob.Id))!;
+			differentTplRef = (await JobService.GetJobAsync(differentTplRef.Id))!;
+			differentTplHash = (await JobService.GetJobAsync(differentTplHash.Id))!;
+			differentUserName = (await JobService.GetJobAsync(differentUserName.Id))!;
+			differentArgs = (await JobService.GetJobAsync(differentArgs.Id))!;
 			
-			Assert.AreEqual(KnownUsers.System, OrgJob.AbortedByUserId);
-			Assert.IsNull(NewJob.AbortedByUserId);
-			Assert.IsNull(DifferentTplRef.AbortedByUserId);
-			Assert.IsNull(DifferentTplHash.AbortedByUserId);
-			Assert.IsNull(DifferentUserName.AbortedByUserId);
-			Assert.IsNull(DifferentArgs.AbortedByUserId);
+			Assert.AreEqual(KnownUsers.System, orgJob.AbortedByUserId);
+			Assert.IsNull(newJob.AbortedByUserId);
+			Assert.IsNull(differentTplRef.AbortedByUserId);
+			Assert.IsNull(differentTplHash.AbortedByUserId);
+			Assert.IsNull(differentUserName.AbortedByUserId);
+			Assert.IsNull(differentArgs.AbortedByUserId);
 		}
 
-		private async Task<IJob> CreatePreflightJob(Fixture Fixture, string TemplateRefId, string TemplateHash, string StartedByUserName, int PreflightChange, string[] Arguments)
+		private async Task<IJob> CreatePreflightJob(Fixture fixture, string templateRefId, string templateHash, string startedByUserName, int preflightChange, string[] arguments)
 		{
-			IUser User = await UserCollection.FindOrAddUserByLoginAsync(StartedByUserName);
+			IUser user = await UserCollection.FindOrAddUserByLoginAsync(startedByUserName);
 			return await JobService.CreateJobAsync(
 				JobId: JobId.GenerateNewId(),
-				Stream: Fixture!.Stream!,
-				TemplateRefId: new TemplateRefId(TemplateRefId),
-				TemplateHash: new ContentHash(Encoding.ASCII.GetBytes(TemplateHash)),
-				Graph: Fixture!.Graph,
+				Stream: fixture!.Stream!,
+				TemplateRefId: new TemplateRefId(templateRefId),
+				TemplateHash: new ContentHash(Encoding.ASCII.GetBytes(templateHash)),
+				Graph: fixture!.Graph,
 				Name: "hello1",
 				Change: 1000001,
 				CodeChange: 1000002,
-				PreflightChange: PreflightChange,
+				PreflightChange: preflightChange,
 				ClonedPreflightChange: null,
-				StartedByUserId: User.Id,
+				StartedByUserId: user.Id,
 				Priority: Priority.Normal,
 				null,
 				null,
@@ -116,7 +115,7 @@ namespace Horde.Build.Tests
 				false,
 				null,
 				null,
-				Arguments: new List<string>(Arguments)
+				Arguments: new List<string>(arguments)
 			);
 		}
 		
@@ -147,43 +146,43 @@ namespace Horde.Build.Tests
 		[TestMethod]
 		public async Task TestRunEarly()
 		{
-			IAgent? Agent = await AgentService.CreateAgentAsync("TestAgent", true, null, new List<StringId<IPool>> { new StringId<IPool>("win") });
-			await AgentService.CreateSessionAsync(Agent, AgentStatus.Ok, new List<string>(), new Dictionary<string, int>(), null);
+			IAgent? agent = await AgentService.CreateAgentAsync("TestAgent", true, null, new List<StringId<IPool>> { new StringId<IPool>("win") });
+			await AgentService.CreateSessionAsync(agent, AgentStatus.Ok, new List<string>(), new Dictionary<string, int>(), null);
 
-			IProject? Project = await ProjectService.Collection.AddOrUpdateAsync(new ProjectId("ue5"), "", "", 0, new ProjectConfig { Name = "UE5" });
-			Assert.IsNotNull(Project);
+			IProject? project = await ProjectService.Collection.AddOrUpdateAsync(new ProjectId("ue5"), "", "", 0, new ProjectConfig { Name = "UE5" });
+			Assert.IsNotNull(project);
 
-			StreamId StreamId = new StreamId("ue5-main");
-			IStream? Stream = await StreamCollection.GetAsync(StreamId);
-			Stream = await StreamCollection.TryCreateOrReplaceAsync(StreamId, Stream, "", "", Project!.Id, new StreamConfig { Name = "//UE5/Main" });
+			StreamId streamId = new StreamId("ue5-main");
+			IStream? stream = await StreamCollection.GetAsync(streamId);
+			stream = await StreamCollection.TryCreateOrReplaceAsync(streamId, stream, "", "", project!.Id, new StreamConfig { Name = "//UE5/Main" });
 
-			ITemplate Template = await TemplateCollection.AddAsync("Test template");
-			IGraph Graph = await GraphCollection.AddAsync(Template);
+			ITemplate template = await TemplateCollection.AddAsync("Test template");
+			IGraph graph = await GraphCollection.AddAsync(template);
 
-			NewGroup GroupA = new NewGroup("win", new List<NewNode>());
-			GroupA.Nodes.Add(new NewNode("Compile"));
+			NewGroup groupA = new NewGroup("win", new List<NewNode>());
+			groupA.Nodes.Add(new NewNode("Compile"));
 
-			NewGroup GroupB = new NewGroup("win", new List<NewNode>());
-			GroupB.Nodes.Add(new NewNode("Cook", RunEarly: true));
-			GroupB.Nodes.Add(new NewNode("Middle"));
-			GroupB.Nodes.Add(new NewNode("Pak", InputDependencies: new List<string> { "Compile", "Cook" }));
+			NewGroup groupB = new NewGroup("win", new List<NewNode>());
+			groupB.Nodes.Add(new NewNode("Cook", RunEarly: true));
+			groupB.Nodes.Add(new NewNode("Middle"));
+			groupB.Nodes.Add(new NewNode("Pak", InputDependencies: new List<string> { "Compile", "Cook" }));
 
-			Graph = await GraphCollection.AppendAsync(Graph, new List<NewGroup> { GroupA, GroupB });
+			graph = await GraphCollection.AppendAsync(graph, new List<NewGroup> { groupA, groupB });
 
-			IJob Job = await JobService.CreateJobAsync(null, Stream!, new TemplateRefId("temp"), Template.Id, Graph, "Hello", 1234, 1233, 999, null, null, null, null, null, null, null, true, true, null, null, new List<string> { "-Target=Pak" });
+			IJob job = await JobService.CreateJobAsync(null, stream!, new TemplateRefId("temp"), template.Id, graph, "Hello", 1234, 1233, 999, null, null, null, null, null, null, null, true, true, null, null, new List<string> { "-Target=Pak" });
 
-			Job = Deref(await JobService.UpdateBatchAsync(Job, Job.Batches[0].Id, LogId.GenerateNewId(), JobStepBatchState.Running));
-			Job = Deref(await JobService.UpdateStepAsync(Job, Job.Batches[0].Id, Job.Batches[0].Steps[0].Id, JobStepState.Running));
-			Job = Deref(await JobService.UpdateStepAsync(Job, Job.Batches[0].Id, Job.Batches[0].Steps[0].Id, JobStepState.Completed, JobStepOutcome.Success));
+			job = Deref(await JobService.UpdateBatchAsync(job, job.Batches[0].Id, LogId.GenerateNewId(), JobStepBatchState.Running));
+			job = Deref(await JobService.UpdateStepAsync(job, job.Batches[0].Id, job.Batches[0].Steps[0].Id, JobStepState.Running));
+			job = Deref(await JobService.UpdateStepAsync(job, job.Batches[0].Id, job.Batches[0].Steps[0].Id, JobStepState.Completed, JobStepOutcome.Success));
 
-			Job = Deref(await JobService.UpdateBatchAsync(Job, Job.Batches[1].Id, LogId.GenerateNewId(), JobStepBatchState.Running));
-			Job = Deref(await JobService.UpdateStepAsync(Job, Job.Batches[1].Id, Job.Batches[1].Steps[0].Id, JobStepState.Running));
+			job = Deref(await JobService.UpdateBatchAsync(job, job.Batches[1].Id, LogId.GenerateNewId(), JobStepBatchState.Running));
+			job = Deref(await JobService.UpdateStepAsync(job, job.Batches[1].Id, job.Batches[1].Steps[0].Id, JobStepState.Running));
 
-			Job = Deref(await JobService.UpdateBatchAsync(Job, Job.Batches[2].Id, LogId.GenerateNewId(), JobStepBatchState.Running));
-			Job = Deref(await JobService.UpdateStepAsync(Job, Job.Batches[2].Id, Job.Batches[2].Steps[0].Id, JobStepState.Running));
-			Job = Deref(await JobService.UpdateStepAsync(Job, Job.Batches[2].Id, Job.Batches[2].Steps[0].Id, JobStepState.Completed, JobStepOutcome.Success));
+			job = Deref(await JobService.UpdateBatchAsync(job, job.Batches[2].Id, LogId.GenerateNewId(), JobStepBatchState.Running));
+			job = Deref(await JobService.UpdateStepAsync(job, job.Batches[2].Id, job.Batches[2].Steps[0].Id, JobStepState.Running));
+			job = Deref(await JobService.UpdateStepAsync(job, job.Batches[2].Id, job.Batches[2].Steps[0].Id, JobStepState.Completed, JobStepOutcome.Success));
 
-			Assert.AreEqual(JobStepState.Waiting, Job.Batches[2].Steps[1].State);
+			Assert.AreEqual(JobStepState.Waiting, job.Batches[2].Steps[1].State);
 		}
 	}
 }
diff --git a/Engine/Source/Programs/Horde/Horde.Build.Tests/JobTaskSourceTests.cs b/Engine/Source/Programs/Horde/Horde.Build.Tests/JobTaskSourceTests.cs
--- a/Engine/Source/Programs/Horde/Horde.Build.Tests/JobTaskSourceTests.cs
+++ b/Engine/Source/Programs/Horde/Horde.Build.Tests/JobTaskSourceTests.cs
@@ -11,96 +11,91 @@ using Horde.Build.Utilities;
 
 namespace Horde.Build.Tests
 {
-	using JobId = ObjectId<IJob>;
-	using LogId = ObjectId<ILogFile>;
-
 	[TestClass]
 	public class JobTaskSourceTests : TestSetup
 	{
-		private bool EventReceived;
-		private IPool? EventPool;
-		private bool? EventPoolHasAgentsOnline;
+		private bool _eventReceived;
+		private bool? _eventPoolHasAgentsOnline;
 		
 		[TestMethod]
 		public async Task UpdateJobQueueNormal()
 		{
-			Fixture Fixture = await SetupPoolWithAgentAsync(IsPoolAutoScaled: true, ShouldCreateAgent: true, IsAgentEnabled: true);
+			Fixture fixture = await SetupPoolWithAgentAsync(isPoolAutoScaled: true, shouldCreateAgent: true, isAgentEnabled: true);
 
 			Assert.AreEqual(0, JobTaskSource.GetQueueForTesting().Count);
 			await JobTaskSource.TickAsync(CancellationToken.None);
 			Assert.AreEqual(1, JobTaskSource.GetQueueForTesting().Count);
-			Assert.AreEqual(Fixture.Job1.Id, JobTaskSource.GetQueueForTesting().Min!.Id.Item1);
+			Assert.AreEqual(fixture.Job1.Id, JobTaskSource.GetQueueForTesting().Min!.Id.Item1);
 			Assert.AreEqual(JobStepBatchState.Ready, JobTaskSource.GetQueueForTesting().Min!.Batch.State);
 			
-			Assert.IsTrue(EventReceived);
-			Assert.IsTrue(EventPoolHasAgentsOnline!.Value);
+			Assert.IsTrue(_eventReceived);
+			Assert.IsTrue(_eventPoolHasAgentsOnline!.Value);
 		}
 		
 		[TestMethod]
 		public async Task UpdateJobQueueWithNoAgentsInPool()
 		{
-			Fixture Fixture = await SetupPoolWithAgentAsync(IsPoolAutoScaled: true, ShouldCreateAgent: false, IsAgentEnabled: false);
+			Fixture fixture = await SetupPoolWithAgentAsync(isPoolAutoScaled: true, shouldCreateAgent: false, isAgentEnabled: false);
 			
 			Assert.AreEqual(0, JobTaskSource.GetQueueForTesting().Count);
 			await JobTaskSource.TickAsync(CancellationToken.None);
 			Assert.AreEqual(0, JobTaskSource.GetQueueForTesting().Count);
 
-			IJob Job = (await JobService.GetJobAsync(Fixture.Job1.Id))!;
-			Assert.AreEqual(JobStepBatchError.NoAgentsInPool, Job.Batches[0].Error);
+			IJob job = (await JobService.GetJobAsync(fixture.Job1.Id))!;
+			Assert.AreEqual(JobStepBatchError.NoAgentsInPool, job.Batches[0].Error);
 			
-			Assert.IsFalse(EventReceived);
+			Assert.IsFalse(_eventReceived);
 		}
 		
 		[TestMethod]
 		public async Task UpdateJobQueueWithNoAgentsOnlineInPool()
 		{
-			Fixture Fixture = await SetupPoolWithAgentAsync(IsPoolAutoScaled: false, ShouldCreateAgent: true, IsAgentEnabled: false);
+			Fixture fixture = await SetupPoolWithAgentAsync(isPoolAutoScaled: false, shouldCreateAgent: true, isAgentEnabled: false);
 			
 			Assert.AreEqual(0, JobTaskSource.GetQueueForTesting().Count);
 			await JobTaskSource.TickAsync(CancellationToken.None);
 			Assert.AreEqual(0, JobTaskSource.GetQueueForTesting().Count);
 
-			IJob Job = (await JobService.GetJobAsync(Fixture.Job1.Id))!;
-			Assert.AreEqual(JobStepBatchError.NoAgentsOnline, Job.Batches[0].Error);
+			IJob job = (await JobService.GetJobAsync(fixture.Job1.Id))!;
+			Assert.AreEqual(JobStepBatchError.NoAgentsOnline, job.Batches[0].Error);
 			
-			Assert.IsFalse(EventReceived);
+			Assert.IsFalse(_eventReceived);
 		}
 		
 		[TestMethod]
 		public async Task UpdateJobQueueWithNoAgentsOnlineInAutoScaledPool()
 		{
-			Fixture Fixture = await SetupPoolWithAgentAsync(IsPoolAutoScaled: true, ShouldCreateAgent: true, IsAgentEnabled: false);
+			Fixture fixture = await SetupPoolWithAgentAsync(isPoolAutoScaled: true, shouldCreateAgent: true, isAgentEnabled: false);
 			
 			Assert.AreEqual(0, JobTaskSource.GetQueueForTesting().Count);
 			await JobTaskSource.TickAsync(CancellationToken.None);
 			Assert.AreEqual(1, JobTaskSource.GetQueueForTesting().Count);
 
-			Assert.AreEqual(Fixture.Job1.Id, JobTaskSource.GetQueueForTesting().Min!.Id.Item1);
+			Assert.AreEqual(fixture.Job1.Id, JobTaskSource.GetQueueForTesting().Min!.Id.Item1);
 			Assert.AreEqual(JobStepBatchState.Ready, JobTaskSource.GetQueueForTesting().Min!.Batch.State);
 
-			Assert.IsTrue(EventReceived);
-			Assert.IsFalse(EventPoolHasAgentsOnline!.Value);
+			Assert.IsTrue(_eventReceived);
+			Assert.IsFalse(_eventPoolHasAgentsOnline!.Value);
 		}
 
-		private async Task<Fixture> SetupPoolWithAgentAsync(bool IsPoolAutoScaled, bool ShouldCreateAgent, bool IsAgentEnabled)
+		private async Task<Fixture> SetupPoolWithAgentAsync(bool isPoolAutoScaled, bool shouldCreateAgent, bool isAgentEnabled)
 		{
-			Fixture Fixture = await CreateFixtureAsync();
-			IPool Pool = await PoolService.CreatePoolAsync(Fixture.PoolName, null, IsPoolAutoScaled, 0, 0);
+			Fixture fixture = await CreateFixtureAsync();
+			IPool pool = await PoolService.CreatePoolAsync(Fixture.PoolName, null, isPoolAutoScaled, 0, 0);
 
-			if (ShouldCreateAgent)
+			if (shouldCreateAgent)
 			{
-				IAgent? Agent = await AgentService.CreateAgentAsync("TestAgent", IsAgentEnabled, null, new List<StringId<IPool>> { Pool.Id });
-				await AgentService.CreateSessionAsync(Agent, AgentStatus.Ok, new List<string>(), new Dictionary<string, int>(), null);
+				IAgent? agent = await AgentService.CreateAgentAsync("TestAgent", isAgentEnabled, null, new List<StringId<IPool>> { pool.Id });
+				await AgentService.CreateSessionAsync(agent, AgentStatus.Ok, new List<string>(), new Dictionary<string, int>(), null);
 			}
 			
-			JobTaskSource.OnJobScheduled += (Pool, PoolHasAgentsOnline, Job, Graph, BatchId) =>
+			JobTaskSource.OnJobScheduled += (pool, poolHasAgentsOnline, job, graph, batchId) =>
 			{
-				EventReceived = true;
-				EventPool = Pool;
-				EventPoolHasAgentsOnline = PoolHasAgentsOnline;
+				_eventReceived = true;
+				_eventPoolHasAgentsOnline = poolHasAgentsOnline;
 			};
 
-			return Fixture;
+			return fixture;
 		}
 	}
 }
diff --git a/Engine/Source/Programs/Horde/Horde.Build.Tests/JobsControllerDbTest.cs b/Engine/Source/Programs/Horde/Horde.Build.Tests/JobsControllerDbTest.cs
--- a/Engine/Source/Programs/Horde/Horde.Build.Tests/JobsControllerDbTest.cs
+++ b/Engine/Source/Programs/Horde/Horde.Build.Tests/JobsControllerDbTest.cs
@@ -1,22 +1,19 @@
 // Copyright Epic Games, Inc. All Rights Reserved.
 
-using System;
 using System.Collections.Generic;
-using System.Diagnostics;
 using System.Threading.Tasks;
 using Horde.Build.Api;
 using Microsoft.VisualStudio.TestTools.UnitTesting;
 using Horde.Build.Controllers;
 using Horde.Build.Models;
-using Horde.Build.Utilities;
 using Microsoft.AspNetCore.Mvc;
 
 namespace Horde.Build.Tests
 {
-    /// <summary>
-    /// Database-only test for testing the Job controller. Different from the JobsController test that set up
-    /// the entire ASP.NET chain.
-    /// </summary>
+	/// <summary>
+	/// Database-only test for testing the Job controller. Different from the JobsController test that set up
+	/// the entire ASP.NET chain.
+	/// </summary>
 	[TestClass]
     public class JobsControllerDbTest : TestSetup
     {
@@ -24,54 +21,54 @@ namespace Horde.Build.Tests
         [TestMethod]
         public async Task GetJobs()
         {
-			Fixture Fixture = await CreateFixtureAsync();
+			Fixture fixture = await CreateFixtureAsync();
 
-			ActionResult<List<object>> Res = await JobsController.FindJobsAsync();
-	        Assert.AreEqual(2, Res.Value!.Count);
-	        Assert.AreEqual("hello2", (Res.Value[0] as GetJobResponse)!.Name);
-	        Assert.AreEqual("hello1", (Res.Value[1] as GetJobResponse)!.Name);
+			ActionResult<List<object>> res = await JobsController.FindJobsAsync();
+	        Assert.AreEqual(2, res.Value!.Count);
+	        Assert.AreEqual("hello2", (res.Value[0] as GetJobResponse)!.Name);
+	        Assert.AreEqual("hello1", (res.Value[1] as GetJobResponse)!.Name);
 	        
-	        Res = await JobsController.FindJobsAsync(IncludePreflight: false);
-	        Assert.AreEqual(1, Res.Value!.Count);
-	        Assert.AreEqual("hello2", (Res.Value[0] as GetJobResponse)!.Name);
+	        res = await JobsController.FindJobsAsync(IncludePreflight: false);
+	        Assert.AreEqual(1, res.Value!.Count);
+	        Assert.AreEqual("hello2", (res.Value[0] as GetJobResponse)!.Name);
         }
         
         [TestMethod]
         public async Task AbortStepTest()
         {
-			Fixture Fixture = await CreateFixtureAsync();
+			Fixture fixture = await CreateFixtureAsync();
 
-	        IJob Job = Fixture.Job1;
-	        SubResourceId BatchId = Job.Batches[0].Id;
-	        SubResourceId StepId = Job.Batches[0].Steps[0].Id;
+	        IJob job = fixture.Job1;
+	        SubResourceId batchId = job.Batches[0].Id;
+	        SubResourceId stepId = job.Batches[0].Steps[0].Id;
 
-	        object Obj = (await JobsController.GetStepAsync(Job.Id, BatchId, StepId)).Value!;
-	        GetStepResponse StepRes = (Obj as GetStepResponse)!;
-	        Assert.IsFalse(StepRes.AbortRequested);
+	        object obj = (await JobsController.GetStepAsync(job.Id, batchId, stepId)).Value!;
+	        GetStepResponse stepRes = (obj as GetStepResponse)!;
+	        Assert.IsFalse(stepRes.AbortRequested);
 	        
-	        UpdateStepRequest UpdateReq = new UpdateStepRequest();
-	        UpdateReq.AbortRequested = true;
-	        Obj = (await JobsController.UpdateStepAsync(Job.Id, BatchId, StepId, UpdateReq)).Value!;
-	        UpdateStepResponse UpdateRes = (Obj as UpdateStepResponse)!;
+	        UpdateStepRequest updateReq = new UpdateStepRequest();
+	        updateReq.AbortRequested = true;
+	        obj = (await JobsController.UpdateStepAsync(job.Id, batchId, stepId, updateReq)).Value!;
+	        UpdateStepResponse updateRes = (obj as UpdateStepResponse)!;
 	        
-	        Obj = (await JobsController.GetStepAsync(Job.Id, BatchId, StepId)).Value!;
-	        StepRes = (Obj as GetStepResponse)!;
-	        Assert.IsTrue(StepRes.AbortRequested);
+	        obj = (await JobsController.GetStepAsync(job.Id, batchId, stepId)).Value!;
+	        stepRes = (obj as GetStepResponse)!;
+	        Assert.IsTrue(stepRes.AbortRequested);
 //	        Assert.AreEqual("Anonymous", StepRes.AbortByUser);
         }
         
         [TestMethod]
         public async Task FindJobTimingsTest()
         {
-	        Fixture Fixture = await CreateFixtureAsync();
-	        IJob Job = Fixture.Job1;
-	        string[] Templates = { Job.TemplateId.ToString() };
-	        object Obj = (await JobsController.FindJobTimingsAsync(Fixture.Stream!.Id.ToString(), Templates)).Value!;
-	        FindJobTimingsResponse Res = (Obj as FindJobTimingsResponse)!;
-	        Assert.AreEqual(1, Res.Timings.Count);
-	        GetJobTimingResponse TimingResponse = Res.Timings[Job.Id.ToString()];
-	        Assert.AreEqual(0, TimingResponse.JobResponse!.Labels!.Count);
-	        Assert.AreEqual(Job.Name, TimingResponse.Job!.Name);
+	        Fixture fixture = await CreateFixtureAsync();
+	        IJob job = fixture.Job1;
+	        string[] templates = { job.TemplateId.ToString() };
+	        object obj = (await JobsController.FindJobTimingsAsync(fixture.Stream!.Id.ToString(), templates)).Value!;
+	        FindJobTimingsResponse res = (obj as FindJobTimingsResponse)!;
+	        Assert.AreEqual(1, res.Timings.Count);
+	        GetJobTimingResponse timingResponse = res.Timings[job.Id.ToString()];
+	        Assert.AreEqual(0, timingResponse.JobResponse!.Labels!.Count);
+	        Assert.AreEqual(job.Name, timingResponse.Job!.Name);
         }
     }
 }
\ No newline at end of file
diff --git a/Engine/Source/Programs/Horde/Horde.Build.Tests/JobsControllerTest.cs b/Engine/Source/Programs/Horde/Horde.Build.Tests/JobsControllerTest.cs
--- a/Engine/Source/Programs/Horde/Horde.Build.Tests/JobsControllerTest.cs
+++ b/Engine/Source/Programs/Horde/Horde.Build.Tests/JobsControllerTest.cs
@@ -1,5 +1,6 @@
 ﻿// Copyright Epic Games, Inc. All Rights Reserved.
 
+using System;
 using System.Net;
 using System.Threading.Tasks;
 using Microsoft.VisualStudio.TestTools.UnitTesting;
@@ -9,25 +10,25 @@ namespace Horde.Build.Tests
     [TestClass]
     public class JobsControllerTest : ControllerIntegrationTest
     {
-        private static string GetUri(string JobId, string StepId, string FileName)
+        private static Uri GetUri(string jobId, string stepId, string fileName)
         {
-            return $"/api/v1/jobs/{JobId}/steps/{StepId}/artifacts/{FileName}/data";
+            return new Uri($"http://localhost/api/v1/jobs/{jobId}/steps/{stepId}/artifacts/{fileName}/data");
         }
 
         [TestMethod]
         public async Task GetArtifactDataByFilenameTest()
         {
-            var Fixture = await GetFixture();
-            var Art = Fixture.Job1Artifact;
+			Fixture? fixture = await GetFixture();
+			Models.IArtifact? art = fixture.Job1Artifact;
 
-            // Test existing filename
-            var Res = await client.GetAsync(GetUri(Art.JobId.ToString(), Art.StepId.ToString()!, Art.Name));
-            Res.EnsureSuccessStatusCode();
-            Assert.AreEqual(Fixture.Job1ArtifactData, await Res.Content.ReadAsStringAsync());
+			// Test existing filename
+			System.Net.Http.HttpResponseMessage? res = await Client.GetAsync(GetUri(art.JobId.ToString(), art.StepId.ToString()!, art.Name));
+            res.EnsureSuccessStatusCode();
+            Assert.AreEqual(fixture.Job1ArtifactData, await res.Content.ReadAsStringAsync());
 
             // Test non-existing filename
-            Res = await client.GetAsync(GetUri(Art.JobId.ToString(), Art.StepId.ToString()!, "bogus.txt"));
-            Assert.AreEqual(Res.StatusCode, HttpStatusCode.NotFound);
+            res = await Client.GetAsync(GetUri(art.JobId.ToString(), art.StepId.ToString()!, "bogus.txt"));
+            Assert.AreEqual(res.StatusCode, HttpStatusCode.NotFound);
         }
     }
 }
\ No newline at end of file
diff --git a/Engine/Source/Programs/Horde/Horde.Build.Tests/LeasesControllerDbTest.cs b/Engine/Source/Programs/Horde/Horde.Build.Tests/LeasesControllerDbTest.cs
--- a/Engine/Source/Programs/Horde/Horde.Build.Tests/LeasesControllerDbTest.cs
+++ b/Engine/Source/Programs/Horde/Horde.Build.Tests/LeasesControllerDbTest.cs
@@ -7,7 +7,6 @@ using HordeCommon;
 using Horde.Build.Api;
 using Horde.Build.Models;
 using Horde.Build.Utilities;
-using Horde.Build.Tests;
 using Microsoft.AspNetCore.Mvc;
 using Microsoft.VisualStudio.TestTools.UnitTesting;
 
@@ -22,26 +21,26 @@ namespace Horde.Build.Tests
 		[TestMethod]
 		public async Task FindLeases()
 		{
-			DateTimeOffset MinTime = Clock.UtcNow - TimeSpan.FromMinutes(5);
-			DateTimeOffset MaxTime = Clock.UtcNow;
+			DateTimeOffset minTime = Clock.UtcNow - TimeSpan.FromMinutes(5);
+			DateTimeOffset maxTime = Clock.UtcNow;
 
-			ILease Lease1 = await CreateLease(Clock.UtcNow - TimeSpan.FromMinutes(10), TimeSpan.FromMinutes(6));
-			ILease Lease2 = await CreateLease(Clock.UtcNow - TimeSpan.FromMinutes(7), TimeSpan.FromMinutes(3.1));
-			ILease OutOfTimeWindow = await CreateLease(Clock.UtcNow - TimeSpan.FromMinutes(7), TimeSpan.FromMinutes(25));
+			ILease lease1 = await CreateLease(Clock.UtcNow - TimeSpan.FromMinutes(10), TimeSpan.FromMinutes(6));
+			ILease lease2 = await CreateLease(Clock.UtcNow - TimeSpan.FromMinutes(7), TimeSpan.FromMinutes(3.1));
+			ILease outOfTimeWindow = await CreateLease(Clock.UtcNow - TimeSpan.FromMinutes(7), TimeSpan.FromMinutes(25));
 			
-			ActionResult<List<object>> Res = await LeasesController.FindLeasesAsync(null, null, null, null, MinTime, MaxTime);
-			Assert.AreEqual(2, Res.Value!.Count);
-			Assert.AreEqual(Lease2.Id.ToString(), (Res.Value[0] as GetAgentLeaseResponse)!.Id);
-			Assert.AreEqual(Lease1.Id.ToString(), (Res.Value[1] as GetAgentLeaseResponse)!.Id);
+			ActionResult<List<object>> res = await LeasesController.FindLeasesAsync(null, null, null, null, minTime, maxTime);
+			Assert.AreEqual(2, res.Value!.Count);
+			Assert.AreEqual(lease2.Id.ToString(), (res.Value[0] as GetAgentLeaseResponse)!.Id);
+			Assert.AreEqual(lease1.Id.ToString(), (res.Value[1] as GetAgentLeaseResponse)!.Id);
 		}
 
-		private async Task<ILease> CreateLease(DateTime StartTime, TimeSpan Duration)
+		private async Task<ILease> CreateLease(DateTime startTime, TimeSpan duration)
 		{
-			ObjectId<ILease> Id = ObjectId<ILease>.GenerateNewId();
-			ObjectId<ISession> SessionId = ObjectId<ISession>.GenerateNewId();
-			ILease Lease = await LeaseCollection.AddAsync(Id, "myLease", new AgentId("agent-1"), SessionId, null, null, null, StartTime, new byte[] { });
-			await LeaseCollection.TrySetOutcomeAsync(Id, StartTime + Duration, LeaseOutcome.Success, null);
-			return Lease;
+			ObjectId<ILease> id = ObjectId<ILease>.GenerateNewId();
+			ObjectId<ISession> sessionId = ObjectId<ISession>.GenerateNewId();
+			ILease lease = await LeaseCollection.AddAsync(id, "myLease", new AgentId("agent-1"), sessionId, null, null, null, startTime, Array.Empty<byte>());
+			await LeaseCollection.TrySetOutcomeAsync(id, startTime + duration, LeaseOutcome.Success, null);
+			return lease;
 		}
 	}
 }
\ No newline at end of file
diff --git a/Engine/Source/Programs/Horde/Horde.Build.Tests/LogBuilderTests.cs b/Engine/Source/Programs/Horde/Horde.Build.Tests/LogBuilderTests.cs
--- a/Engine/Source/Programs/Horde/Horde.Build.Tests/LogBuilderTests.cs
+++ b/Engine/Source/Programs/Horde/Horde.Build.Tests/LogBuilderTests.cs
@@ -7,7 +7,6 @@ using Horde.Build.Models;
 using Horde.Build.Utilities;
 using Microsoft.Extensions.Logging.Abstractions;
 using Microsoft.VisualStudio.TestTools.UnitTesting;
-using MongoDB.Bson;
 using System;
 using System.Collections.Generic;
 using System.Text;
@@ -23,79 +22,79 @@ namespace Horde.Build.Tests
 		[TestMethod]
 		public async Task TestLocalLogBuilder()
 		{
-			ILogBuilder Builder = new LocalLogBuilder();
-			await TestBuilder(Builder);
+			ILogBuilder builder = new LocalLogBuilder();
+			await TestBuilder(builder);
 		}
 
 		[TestMethod]
 		public async Task TestRedisLogBuilder()
 		{
-			ILogBuilder Builder = new RedisLogBuilder(GetRedisConnectionPool(), NullLogger.Instance);
-			await TestBuilder(Builder);
+			ILogBuilder builder = new RedisLogBuilder(GetRedisConnectionPool(), NullLogger.Instance);
+			await TestBuilder(builder);
 		}
 
-		public async Task TestBuilder(ILogBuilder Builder)
+		public static async Task TestBuilder(ILogBuilder builder)
 		{
-			LogId LogId = LogId.GenerateNewId();
+			LogId logId = LogId.GenerateNewId();
 
 			const long Offset = 100;
-			Assert.IsTrue(await Builder.AppendAsync(LogId, Offset, Offset, 0, 1, Encoding.UTF8.GetBytes("hello\n"), LogType.Text));
-
-			LogChunkData? Chunk1 = await Builder.GetChunkAsync(LogId, Offset, 0);
-			Assert.IsNotNull(Chunk1);
-			Assert.AreEqual(1, Chunk1!.LineCount);
-			Assert.AreEqual(6, Chunk1!.Length);
-			Assert.AreEqual(1, Chunk1!.SubChunks.Count);
-			Assert.AreEqual(1, Chunk1!.SubChunks[0].LineCount);
-			Assert.AreEqual(6, Chunk1!.SubChunks[0].Length);
-
-			Assert.IsTrue(await Builder.AppendAsync(LogId, Offset, Offset + 6, 1, 1, Encoding.UTF8.GetBytes("world\n"), LogType.Text));
-
-			LogChunkData? Chunk2 = await Builder.GetChunkAsync(LogId, Offset, 1);
-			Assert.IsNotNull(Chunk2);
-			Assert.AreEqual(2, Chunk2!.LineCount);
-			Assert.AreEqual(12, Chunk2!.Length);
-			Assert.AreEqual(1, Chunk2!.SubChunks.Count);
-			Assert.AreEqual(2, Chunk2!.SubChunks[0].LineCount);
-			Assert.AreEqual(12, Chunk2!.SubChunks[0].Length);
-
-			await Builder.CompleteSubChunkAsync(LogId, Offset);
-			Assert.IsTrue(await Builder.AppendAsync(LogId, Offset, Offset + 12, 2, 1, Encoding.UTF8.GetBytes("foo\n"), LogType.Text));
-			Assert.IsTrue(await Builder.AppendAsync(LogId, Offset, Offset + 16, 3, 2, Encoding.UTF8.GetBytes("bar\nbaz\n"), LogType.Text));
-			await Builder.CompleteSubChunkAsync(LogId, Offset);
-
-			LogChunkData? Chunk3 = await Builder.GetChunkAsync(LogId, Offset, 4);
-			Assert.IsNotNull(Chunk3);
-			Assert.AreEqual(5, Chunk3!.LineCount);
-			Assert.AreEqual(24, Chunk3!.Length);
-			Assert.AreEqual(2, Chunk3!.SubChunks.Count);
-			Assert.AreEqual(2, Chunk3!.SubChunks[0].LineCount);
-			Assert.AreEqual(12, Chunk3!.SubChunks[0].Length);
-			Assert.AreEqual(3, Chunk3!.SubChunks[1].LineCount);
-			Assert.AreEqual(12, Chunk3!.SubChunks[1].Length);
-
-			Assert.AreEqual(0, Chunk3!.GetLineOffsetWithinChunk(0));
-			Assert.AreEqual(6, Chunk3!.GetLineOffsetWithinChunk(1));
-			Assert.AreEqual(12, Chunk3!.GetLineOffsetWithinChunk(2));
-			Assert.AreEqual(16, Chunk3!.GetLineOffsetWithinChunk(3));
-			Assert.AreEqual(20, Chunk3!.GetLineOffsetWithinChunk(4));
-
-			List<(LogId, long)> Chunks1 = await Builder.TouchChunksAsync(TimeSpan.Zero);
-			Assert.AreEqual(1, Chunks1.Count);
-			Assert.AreEqual((LogId, Offset), Chunks1[0]);
-
-			await Builder.CompleteChunkAsync(LogId, Offset);
-
-			LogChunkData? Chunk4 = await Builder.GetChunkAsync(LogId, Offset, 5);
-			Assert.AreEqual(24, Chunk4!.Length);
-
-			List<(LogId, long)> Chunks2 = await Builder.TouchChunksAsync(TimeSpan.Zero);
-			Assert.AreEqual(1, Chunks2.Count);
-
-			await Builder.RemoveChunkAsync(LogId, Offset);
-
-			List<(LogId, long)> Chunks3 = await Builder.TouchChunksAsync(TimeSpan.Zero);
-			Assert.AreEqual(0, Chunks3.Count);
+			Assert.IsTrue(await builder.AppendAsync(logId, Offset, Offset, 0, 1, Encoding.UTF8.GetBytes("hello\n"), LogType.Text));
+
+			LogChunkData? chunk1 = await builder.GetChunkAsync(logId, Offset, 0);
+			Assert.IsNotNull(chunk1);
+			Assert.AreEqual(1, chunk1!.LineCount);
+			Assert.AreEqual(6, chunk1!.Length);
+			Assert.AreEqual(1, chunk1!.SubChunks.Count);
+			Assert.AreEqual(1, chunk1!.SubChunks[0].LineCount);
+			Assert.AreEqual(6, chunk1!.SubChunks[0].Length);
+
+			Assert.IsTrue(await builder.AppendAsync(logId, Offset, Offset + 6, 1, 1, Encoding.UTF8.GetBytes("world\n"), LogType.Text));
+
+			LogChunkData? chunk2 = await builder.GetChunkAsync(logId, Offset, 1);
+			Assert.IsNotNull(chunk2);
+			Assert.AreEqual(2, chunk2!.LineCount);
+			Assert.AreEqual(12, chunk2!.Length);
+			Assert.AreEqual(1, chunk2!.SubChunks.Count);
+			Assert.AreEqual(2, chunk2!.SubChunks[0].LineCount);
+			Assert.AreEqual(12, chunk2!.SubChunks[0].Length);
+
+			await builder.CompleteSubChunkAsync(logId, Offset);
+			Assert.IsTrue(await builder.AppendAsync(logId, Offset, Offset + 12, 2, 1, Encoding.UTF8.GetBytes("foo\n"), LogType.Text));
+			Assert.IsTrue(await builder.AppendAsync(logId, Offset, Offset + 16, 3, 2, Encoding.UTF8.GetBytes("bar\nbaz\n"), LogType.Text));
+			await builder.CompleteSubChunkAsync(logId, Offset);
+
+			LogChunkData? chunk3 = await builder.GetChunkAsync(logId, Offset, 4);
+			Assert.IsNotNull(chunk3);
+			Assert.AreEqual(5, chunk3!.LineCount);
+			Assert.AreEqual(24, chunk3!.Length);
+			Assert.AreEqual(2, chunk3!.SubChunks.Count);
+			Assert.AreEqual(2, chunk3!.SubChunks[0].LineCount);
+			Assert.AreEqual(12, chunk3!.SubChunks[0].Length);
+			Assert.AreEqual(3, chunk3!.SubChunks[1].LineCount);
+			Assert.AreEqual(12, chunk3!.SubChunks[1].Length);
+
+			Assert.AreEqual(0, chunk3!.GetLineOffsetWithinChunk(0));
+			Assert.AreEqual(6, chunk3!.GetLineOffsetWithinChunk(1));
+			Assert.AreEqual(12, chunk3!.GetLineOffsetWithinChunk(2));
+			Assert.AreEqual(16, chunk3!.GetLineOffsetWithinChunk(3));
+			Assert.AreEqual(20, chunk3!.GetLineOffsetWithinChunk(4));
+
+			List<(LogId, long)> chunks1 = await builder.TouchChunksAsync(TimeSpan.Zero);
+			Assert.AreEqual(1, chunks1.Count);
+			Assert.AreEqual((logId, Offset), chunks1[0]);
+
+			await builder.CompleteChunkAsync(logId, Offset);
+
+			LogChunkData? chunk4 = await builder.GetChunkAsync(logId, Offset, 5);
+			Assert.AreEqual(24, chunk4!.Length);
+
+			List<(LogId, long)> chunks2 = await builder.TouchChunksAsync(TimeSpan.Zero);
+			Assert.AreEqual(1, chunks2.Count);
+
+			await builder.RemoveChunkAsync(logId, Offset);
+
+			List<(LogId, long)> chunks3 = await builder.TouchChunksAsync(TimeSpan.Zero);
+			Assert.AreEqual(0, chunks3.Count);
 		}
 	}
 }
diff --git a/Engine/Source/Programs/Horde/Horde.Build.Tests/LogChunkTests.cs b/Engine/Source/Programs/Horde/Horde.Build.Tests/LogChunkTests.cs
--- a/Engine/Source/Programs/Horde/Horde.Build.Tests/LogChunkTests.cs
+++ b/Engine/Source/Programs/Horde/Horde.Build.Tests/LogChunkTests.cs
@@ -12,7 +12,7 @@ namespace Horde.Build.Tests
 	[TestClass]
 	public class LogChunkTests
 	{
-		string[] Text =
+		string[] _text =
 		{
 			@"{ ""id"": 1, ""message"": ""foo"" }",
 			@"{ ""id"": 2, ""message"": ""bar"" }",
@@ -22,25 +22,25 @@ namespace Horde.Build.Tests
 		[TestMethod]
 		public void TestSerialization()
 		{
-			byte[] TextData = Encoding.UTF8.GetBytes(String.Join("\n", Text) + "\n");
+			byte[] textData = Encoding.UTF8.GetBytes(String.Join("\n", _text) + "\n");
 
-			List<LogSubChunkData> SubChunks = new List<LogSubChunkData>();
-			SubChunks.Add(new LogSubChunkData(LogType.Text, 0, 0, new LogText(TextData, TextData.Length)));
-			SubChunks.Add(new LogSubChunkData(LogType.Json, TextData.Length, 3, new LogText(TextData, TextData.Length)));
+			List<LogSubChunkData> subChunks = new List<LogSubChunkData>();
+			subChunks.Add(new LogSubChunkData(LogType.Text, 0, 0, new LogText(textData, textData.Length)));
+			subChunks.Add(new LogSubChunkData(LogType.Json, textData.Length, 3, new LogText(textData, textData.Length)));
 
-			LogChunkData OldChunkData = new LogChunkData(0, 0, SubChunks);
-			byte[] Data = OldChunkData.ToByteArray();
-			LogChunkData NewChunkData = LogChunkData.FromMemory(Data, 0, 0);
+			LogChunkData oldChunkData = new LogChunkData(0, 0, subChunks);
+			byte[] data = oldChunkData.ToByteArray();
+			LogChunkData newChunkData = LogChunkData.FromMemory(data, 0, 0);
 
-			Assert.AreEqual(OldChunkData.Length, NewChunkData.Length);
-			Assert.AreEqual(OldChunkData.SubChunks.Count, NewChunkData.SubChunks.Count);
-			Assert.IsTrue(OldChunkData.SubChunkOffset.AsSpan(0, OldChunkData.SubChunkOffset.Length).SequenceEqual(NewChunkData.SubChunkOffset.AsSpan(0, NewChunkData.SubChunkOffset.Length)));
-			Assert.IsTrue(OldChunkData.SubChunkLineIndex.AsSpan(0, OldChunkData.SubChunkLineIndex.Length).SequenceEqual(NewChunkData.SubChunkLineIndex.AsSpan(0, NewChunkData.SubChunkLineIndex.Length)));
-			for (int Idx = 0; Idx < OldChunkData.SubChunks.Count; Idx++)
+			Assert.AreEqual(oldChunkData.Length, newChunkData.Length);
+			Assert.AreEqual(oldChunkData.SubChunks.Count, newChunkData.SubChunks.Count);
+			Assert.IsTrue(oldChunkData.SubChunkOffset.AsSpan(0, oldChunkData.SubChunkOffset.Length).SequenceEqual(newChunkData.SubChunkOffset.AsSpan(0, newChunkData.SubChunkOffset.Length)));
+			Assert.IsTrue(oldChunkData.SubChunkLineIndex.AsSpan(0, oldChunkData.SubChunkLineIndex.Length).SequenceEqual(newChunkData.SubChunkLineIndex.AsSpan(0, newChunkData.SubChunkLineIndex.Length)));
+			for (int idx = 0; idx < oldChunkData.SubChunks.Count; idx++)
 			{
-				LogSubChunkData OldSubChunkData = OldChunkData.SubChunks[Idx];
-				LogSubChunkData NewSubChunkData = NewChunkData.SubChunks[Idx];
-				Assert.IsTrue(OldSubChunkData.InflateText().Data.Span.SequenceEqual(NewSubChunkData.InflateText().Data.Span));
+				LogSubChunkData oldSubChunkData = oldChunkData.SubChunks[idx];
+				LogSubChunkData newSubChunkData = newChunkData.SubChunks[idx];
+				Assert.IsTrue(oldSubChunkData.InflateText().Data.Span.SequenceEqual(newSubChunkData.InflateText().Data.Span));
 			}
 		}
 	}
diff --git a/Engine/Source/Programs/Horde/Horde.Build.Tests/LogFileServiceTests.cs b/Engine/Source/Programs/Horde/Horde.Build.Tests/LogFileServiceTests.cs
--- a/Engine/Source/Programs/Horde/Horde.Build.Tests/LogFileServiceTests.cs
+++ b/Engine/Source/Programs/Horde/Horde.Build.Tests/LogFileServiceTests.cs
@@ -2,35 +2,26 @@
 
 using System;
 using System.Collections.Generic;
-using System.Diagnostics;
 using System.IO;
 using System.Linq;
-using System.Net;
-using System.Runtime.InteropServices.ComTypes;
 using System.Security.Claims;
 using System.Text;
 using System.Threading.Tasks;
 using HordeCommon;
-using Horde.Build;
 using Horde.Build.Api;
 using Horde.Build.Collections.Impl;
 using Horde.Build.Logs;
 using Horde.Build.Logs.Builder;
 using Horde.Build.Logs.Readers;
-using Horde.Build.Logs.Storage;
-using Horde.Build.Logs.Storage.Impl;
 using Horde.Build.Models;
 using Horde.Build.Services;
 using Horde.Build.Storage;
 using Horde.Build.Storage.Backends;
 using Horde.Build.Utilities;
-using Microsoft.Extensions.Configuration;
 using Microsoft.Extensions.Logging;
 using Microsoft.Extensions.Logging.Abstractions;
 using Microsoft.VisualStudio.TestTools.UnitTesting;
 using MongoDB.Bson;
-using StackExchange.Redis;
-using PoolId = Horde.Build.Utilities.StringId<Horde.Build.Models.IPool>;
 
 namespace Horde.Build.Tests
 {
@@ -40,298 +31,307 @@ namespace Horde.Build.Tests
 	[TestClass]
     public class LogFileServiceTest : DatabaseIntegrationTest
     {
-		private readonly FakeClock Clock;
-        private readonly LogFileService LogFileService;
+		private readonly FakeClock _clock;
+        private readonly LogFileService _logFileService;
+		private readonly ILoggerFactory _loggerFactory;
+		private readonly ILogStorage _logStorage;
 
-        public LogFileServiceTest()
+		public LogFileServiceTest()
         {
-            var LogFileCollection = new LogFileCollection(GetDatabaseService());
+            LogFileCollection logFileCollection = new LogFileCollection(GetDatabaseServiceSingleton());
 
-            var LoggerFactory = (ILoggerFactory) new LoggerFactory();
-            var Logger = LoggerFactory.CreateLogger<LogFileService>();
+			_loggerFactory = new LoggerFactory();
+            ILogger<LogFileService> logger = _loggerFactory.CreateLogger<LogFileService>();
 
-            // Just to satisfy the parameter, need to be fixed
-            IConfiguration Config = new ConfigurationBuilder().Build();
-
-			ILogBuilder LogBuilder = new RedisLogBuilder(GetRedisConnectionPool(), NullLogger.Instance);
-			ILogStorage LogStorage = new PersistentLogStorage(new TransientStorageBackend().ForType<PersistentLogStorage>(), NullLogger<PersistentLogStorage>.Instance);
-			Clock = new FakeClock();
-			LogFileService = new LogFileService(LogFileCollection, null!, LogBuilder, LogStorage, Clock, Logger);
+			ILogBuilder logBuilder = new RedisLogBuilder(GetRedisConnectionPool(), NullLogger.Instance);
+			_logStorage = new PersistentLogStorage(new TransientStorageBackend().ForType<PersistentLogStorage>(), NullLogger<PersistentLogStorage>.Instance);
+			_clock = new FakeClock();
+			_logFileService = new LogFileService(logFileCollection, null!, logBuilder, _logStorage, _clock, logger);
         }
 
+		protected override void Dispose(bool disposing)
+		{
+			base.Dispose(disposing);
+
+			_logFileService.Dispose();
+			_logStorage.Dispose();
+			_loggerFactory.Dispose();
+		}
+
 		[TestMethod]
 		public void PlainTextDecoder()
 		{
-			string JsonText = "{\"time\":\"2022-01-26T14:18:29\",\"level\":\"Error\",\"message\":\"       \\tdepends on: /mnt/horde/U5\\u002BR5.0\\u002BInc\\u002BMin/Sync/Engine/Binaries/Linux/libUnrealEditor-Core.so\",\"id\":1,\"line\":5,\"lineCount\":1028}";
-			byte[] JsonData = Encoding.UTF8.GetBytes(JsonText);
+			string jsonText = "{\"time\":\"2022-01-26T14:18:29\",\"level\":\"Error\",\"message\":\"       \\tdepends on: /mnt/horde/U5\\u002BR5.0\\u002BInc\\u002BMin/Sync/Engine/Binaries/Linux/libUnrealEditor-Core.so\",\"id\":1,\"line\":5,\"lineCount\":1028}";
+			byte[] jsonData = Encoding.UTF8.GetBytes(jsonText);
 
-			byte[] OutputData = new byte[JsonData.Length];
-			int OutputLength = LogText.ConvertToPlainText(JsonData, OutputData, 0);
-			string OutputText = Encoding.UTF8.GetString(OutputData, 0, OutputLength);
+			byte[] outputData = new byte[jsonData.Length];
+			int outputLength = LogText.ConvertToPlainText(jsonData, outputData, 0);
+			string outputText = Encoding.UTF8.GetString(outputData, 0, outputLength);
 
-			string ExpectText = "       \tdepends on: /mnt/horde/U5+R5.0+Inc+Min/Sync/Engine/Binaries/Linux/libUnrealEditor-Core.so\n";
-			Assert.AreEqual(OutputText, ExpectText);
+			string expectText = "       \tdepends on: /mnt/horde/U5+R5.0+Inc+Min/Sync/Engine/Binaries/Linux/libUnrealEditor-Core.so\n";
+			Assert.AreEqual(outputText, expectText);
 			
 		}
 
 		[TestMethod]
         public async Task WriteLogLifecycleOldTest()
         {
-			JobId JobId = JobId.GenerateNewId();
-            ILogFile LogFile = await LogFileService.CreateLogFileAsync(JobId, null, LogType.Text);
+			JobId jobId = JobId.GenerateNewId();
+            ILogFile logFile = await _logFileService.CreateLogFileAsync(jobId, null, LogType.Text);
 
-            LogFile = (await ((ILogFileService)LogFileService).WriteLogDataAsync(LogFile, 0, 0, Encoding.ASCII.GetBytes("hello\n"), true))!;
-			LogFile = (await ((ILogFileService)LogFileService).WriteLogDataAsync(LogFile, 6, 1, Encoding.ASCII.GetBytes("foo\nbar\n"), true))!;
-			LogFile = (await ((ILogFileService)LogFileService).WriteLogDataAsync(LogFile, 6 + 8, 3, Encoding.ASCII.GetBytes("baz\n"), false))!;
+            logFile = (await ((ILogFileService)_logFileService).WriteLogDataAsync(logFile, 0, 0, Encoding.ASCII.GetBytes("hello\n"), true))!;
+			logFile = (await ((ILogFileService)_logFileService).WriteLogDataAsync(logFile, 6, 1, Encoding.ASCII.GetBytes("foo\nbar\n"), true))!;
+			logFile = (await ((ILogFileService)_logFileService).WriteLogDataAsync(logFile, 6 + 8, 3, Encoding.ASCII.GetBytes("baz\n"), false))!;
 
-            Assert.AreEqual("hello", await ReadLogFile(LogFileService, LogFile, 0, 5));
-            Assert.AreEqual("foo\nbar\nbaz\n", await ReadLogFile(LogFileService, LogFile, 6, 12));
+            Assert.AreEqual("hello", await ReadLogFile(_logFileService, logFile, 0, 5));
+            Assert.AreEqual("foo\nbar\nbaz\n", await ReadLogFile(_logFileService, logFile, 6, 12));
 
-            var Metadata = await LogFileService.GetMetadataAsync(LogFile);
-            Assert.AreEqual(6 + 8 + 4, Metadata.Length);
-            Assert.AreEqual(4, Metadata.MaxLineIndex);
+            LogMetadata metadata = await _logFileService.GetMetadataAsync(logFile);
+            Assert.AreEqual(6 + 8 + 4, metadata.Length);
+            Assert.AreEqual(4, metadata.MaxLineIndex);
 
-            Assert.AreEqual((0, 0), await LogFileService.GetLineOffsetAsync(LogFile, 0));
-            Assert.AreEqual((1, 6), await LogFileService.GetLineOffsetAsync(LogFile, 1));
-            Assert.AreEqual((2, 10), await LogFileService.GetLineOffsetAsync(LogFile, 2));
-            Assert.AreEqual((3, 14), await LogFileService.GetLineOffsetAsync(LogFile, 3));
+            Assert.AreEqual((0, 0), await _logFileService.GetLineOffsetAsync(logFile, 0));
+            Assert.AreEqual((1, 6), await _logFileService.GetLineOffsetAsync(logFile, 1));
+            Assert.AreEqual((2, 10), await _logFileService.GetLineOffsetAsync(logFile, 2));
+            Assert.AreEqual((3, 14), await _logFileService.GetLineOffsetAsync(logFile, 3));
         }
         
         [TestMethod]
         public async Task WriteLogLifecycleTest()
         {
-	        await WriteLogLifecycle(LogFileService, 30);
+	        await WriteLogLifecycle(_logFileService, 30);
         }
 
         
-        private static async Task AssertMetadata(ILogFileService LogFileService, ILogFile LogFile, long ExpectedLength, long ExpectedMaxLineIndex)
+        private static async Task AssertMetadata(ILogFileService logFileService, ILogFile logFile, long expectedLength, long expectedMaxLineIndex)
         {
-	        LogMetadata Metadata = await LogFileService.GetMetadataAsync(LogFile);
-	        Assert.AreEqual(ExpectedLength, Metadata.Length);
-	        Assert.AreEqual(ExpectedMaxLineIndex, Metadata.MaxLineIndex);
+	        LogMetadata metadata = await logFileService.GetMetadataAsync(logFile);
+	        Assert.AreEqual(expectedLength, metadata.Length);
+	        Assert.AreEqual(expectedMaxLineIndex, metadata.MaxLineIndex);
         }
 
-        private static async Task AssertChunk(ILogFileService LogFileService, LogId LogFileId, long NumChunks, int ChunkId, long Offset, long Length,
-	        long LineIndex)
+        private static async Task AssertChunk(ILogFileService logFileService, LogId logFileId, long numChunks, int chunkId, long offset, long length,
+	        long lineIndex)
         {
-	        ILogFile? LogFile = await LogFileService.GetLogFileAsync(LogFileId);
-	        Assert.AreEqual(NumChunks, LogFile!.Chunks.Count);
+	        ILogFile? logFile = await logFileService.GetLogFileAsync(logFileId);
+	        Assert.AreEqual(numChunks, logFile!.Chunks.Count);
 
-	        ILogChunk Chunk = LogFile.Chunks[ChunkId];
-	        Assert.AreEqual(Offset, Chunk.Offset);
-	        Assert.AreEqual(Length, Chunk.Length);
-	        Assert.AreEqual(LineIndex, Chunk.LineIndex);
+	        ILogChunk chunk = logFile.Chunks[chunkId];
+	        Assert.AreEqual(offset, chunk.Offset);
+	        Assert.AreEqual(length, chunk.Length);
+	        Assert.AreEqual(lineIndex, chunk.LineIndex);
         }
 
-        private static async Task AssertLineOffset(ILogFileService LogFileService, LogId LogFileId, int LineIndex, int ClampedLineIndex, long Offset)
+        private static async Task AssertLineOffset(ILogFileService logFileService, LogId logFileId, int lineIndex, int clampedLineIndex, long offset)
         {
-	        ILogFile? LogFile = await LogFileService.GetLogFileAsync(LogFileId);
-	        Assert.AreEqual((ClampedLineIndex, Offset), await LogFileService.GetLineOffsetAsync(LogFile!, LineIndex));
+	        ILogFile? logFile = await logFileService.GetLogFileAsync(logFileId);
+	        Assert.AreEqual((clampedLineIndex, offset), await logFileService.GetLineOffsetAsync(logFile!, lineIndex));
         }
         
-        protected static async Task<string> ReadLogFile(ILogFileService LogFileService, ILogFile LogFile, long Offset, long Length)
+        protected static async Task<string> ReadLogFile(ILogFileService logFileService, ILogFile logFile, long offset, long length)
         {
-	        Stream Stream = await LogFileService.OpenRawStreamAsync(LogFile, Offset, Length);
-	        return new StreamReader(Stream).ReadToEnd();
+	        using Stream stream = await logFileService.OpenRawStreamAsync(logFile, offset, length);
+			using StreamReader streamReader = new StreamReader(stream);
+	        return await streamReader.ReadToEndAsync();
         }
         
-        public async Task WriteLogLifecycle(ILogFileService Lfs, int MaxChunkLength)
+        public async Task WriteLogLifecycle(ILogFileService lfs, int maxChunkLength)
         {
-			JobId JobId = JobId.GenerateNewId();
-            ILogFile LogFile = await Lfs.CreateLogFileAsync(JobId, null, LogType.Text);
+			JobId jobId = JobId.GenerateNewId();
+            ILogFile logFile = await lfs.CreateLogFileAsync(jobId, null, LogType.Text);
 
-            string Str1 = "hello\n";
-            string Str2 = "foo\nbar\n";
-            string Str3 = "baz\nqux\nquux\n";
-            string Str4 = "quuz\n";
+            string str1 = "hello\n";
+            string str2 = "foo\nbar\n";
+            string str3 = "baz\nqux\nquux\n";
+            string str4 = "quuz\n";
 
-            int LineIndex = 0;
-            int Offset = 0;
+            int lineIndex = 0;
+            int offset = 0;
 
             // First write with flush. Will become chunk #1
-            LogFile = (await Lfs.WriteLogDataAsync(LogFile, Offset, LineIndex, Encoding.ASCII.GetBytes(Str1), true, MaxChunkLength))!;
-			await Clock.AdvanceAsync(TimeSpan.FromSeconds(2.0));
-            Assert.AreEqual(Str1, await ReadLogFile(Lfs, LogFile, 0, Str1.Length));
-            Assert.AreEqual(Str1, await ReadLogFile(Lfs, LogFile, 0, Str1.Length + 100)); // Reading too far is valid?
-            await AssertMetadata(Lfs, LogFile, Str1.Length, 1);
-            await AssertChunk(Lfs, LogFile.Id, 1, 0, 0, Str1.Length, LineIndex);
-            await AssertLineOffset(Lfs, LogFile.Id, 0, 0, 0);
-            await AssertLineOffset(Lfs, LogFile.Id, 1, 1, 6);
-            await AssertLineOffset(Lfs, LogFile.Id, 1, 1, 6);
-            await AssertLineOffset(Lfs, LogFile.Id, 1, 1, 6);
+            logFile = (await lfs.WriteLogDataAsync(logFile, offset, lineIndex, Encoding.ASCII.GetBytes(str1), true, maxChunkLength))!;
+			await _clock.AdvanceAsync(TimeSpan.FromSeconds(2.0));
+            Assert.AreEqual(str1, await ReadLogFile(lfs, logFile, 0, str1.Length));
+            Assert.AreEqual(str1, await ReadLogFile(lfs, logFile, 0, str1.Length + 100)); // Reading too far is valid?
+            await AssertMetadata(lfs, logFile, str1.Length, 1);
+            await AssertChunk(lfs, logFile.Id, 1, 0, 0, str1.Length, lineIndex);
+            await AssertLineOffset(lfs, logFile.Id, 0, 0, 0);
+            await AssertLineOffset(lfs, logFile.Id, 1, 1, 6);
+            await AssertLineOffset(lfs, logFile.Id, 1, 1, 6);
+            await AssertLineOffset(lfs, logFile.Id, 1, 1, 6);
 
             // Second write without flushing. Will become chunk #2
-            Offset += Str1.Length;
-            LineIndex += Str1.Count(f => f == '\n');
-            LogFile = (await Lfs.WriteLogDataAsync(LogFile, Offset, LineIndex, Encoding.ASCII.GetBytes(Str2), false, MaxChunkLength))!;
-			await Clock.AdvanceAsync(TimeSpan.FromSeconds(2.0));
-			Assert.AreEqual(Str1 + Str2, await ReadLogFile(Lfs, LogFile, 0, Str1.Length + Str2.Length));
-            await AssertMetadata(Lfs, LogFile, Str1.Length + Str2.Length, 3); // FIXME: what are max line index?
-            await AssertChunk(Lfs, LogFile.Id, 2, 0, 0, Str1.Length, 0);
-            await AssertChunk(Lfs, LogFile.Id, 2, 1, Str1.Length, 0, 1); // Last chunk have length zero as it's being written
-            await AssertLineOffset(Lfs, LogFile.Id, 0, 0, 0);
-            await AssertLineOffset(Lfs, LogFile.Id, 1, 1, 6);
-            await AssertLineOffset(Lfs, LogFile.Id, 2, 2, 10);
-            await AssertLineOffset(Lfs, LogFile.Id, 3, 3, 14);
+            offset += str1.Length;
+            lineIndex += str1.Count(f => f == '\n');
+            logFile = (await lfs.WriteLogDataAsync(logFile, offset, lineIndex, Encoding.ASCII.GetBytes(str2), false, maxChunkLength))!;
+			await _clock.AdvanceAsync(TimeSpan.FromSeconds(2.0));
+			Assert.AreEqual(str1 + str2, await ReadLogFile(lfs, logFile, 0, str1.Length + str2.Length));
+            await AssertMetadata(lfs, logFile, str1.Length + str2.Length, 3); // FIXME: what are max line index?
+            await AssertChunk(lfs, logFile.Id, 2, 0, 0, str1.Length, 0);
+            await AssertChunk(lfs, logFile.Id, 2, 1, str1.Length, 0, 1); // Last chunk have length zero as it's being written
+            await AssertLineOffset(lfs, logFile.Id, 0, 0, 0);
+            await AssertLineOffset(lfs, logFile.Id, 1, 1, 6);
+            await AssertLineOffset(lfs, logFile.Id, 2, 2, 10);
+            await AssertLineOffset(lfs, logFile.Id, 3, 3, 14);
 
             // Third write without flushing. Will become chunk #2
-            Offset += Str2.Length;
-            LineIndex += Str3.Count(f => f == '\n');
-            LogFile = (await Lfs.WriteLogDataAsync(LogFile, Offset, LineIndex, Encoding.ASCII.GetBytes(Str3), false, MaxChunkLength))!;
-			await Clock.AdvanceAsync(TimeSpan.FromSeconds(2.0));
-			Assert.AreEqual(Str1 + Str2 + Str3, await ReadLogFile(Lfs, LogFile, 0, Str1.Length + Str2.Length + Str3.Length));
+            offset += str2.Length;
+            lineIndex += str3.Count(f => f == '\n');
+            logFile = (await lfs.WriteLogDataAsync(logFile, offset, lineIndex, Encoding.ASCII.GetBytes(str3), false, maxChunkLength))!;
+			await _clock.AdvanceAsync(TimeSpan.FromSeconds(2.0));
+			Assert.AreEqual(str1 + str2 + str3, await ReadLogFile(lfs, logFile, 0, str1.Length + str2.Length + str3.Length));
             //await AssertMetadata(Lfs, LogFile, Str1.Length + Str2.Length + Str3.Length, 8);
             // Since no flush has happened, chunks should be identical to last write
-            await AssertChunk(Lfs, LogFile.Id, 2, 0, 0, Str1.Length, 0);
-            await AssertChunk(Lfs, LogFile.Id, 2, 1, Str1.Length, 0, 1);
-            await AssertLineOffset(Lfs, LogFile.Id, 0, 0, 0);
-            await AssertLineOffset(Lfs, LogFile.Id, 1, 1, 6);
-            await AssertLineOffset(Lfs, LogFile.Id, 2, 2, 10);
-            await AssertLineOffset(Lfs, LogFile.Id, 3, 3, 14);
-            await AssertLineOffset(Lfs, LogFile.Id, 4, 4, 18);
+            await AssertChunk(lfs, logFile.Id, 2, 0, 0, str1.Length, 0);
+            await AssertChunk(lfs, logFile.Id, 2, 1, str1.Length, 0, 1);
+            await AssertLineOffset(lfs, logFile.Id, 0, 0, 0);
+            await AssertLineOffset(lfs, logFile.Id, 1, 1, 6);
+            await AssertLineOffset(lfs, logFile.Id, 2, 2, 10);
+            await AssertLineOffset(lfs, logFile.Id, 3, 3, 14);
+            await AssertLineOffset(lfs, logFile.Id, 4, 4, 18);
 
             // Fourth write with flush. Will become chunk #2
-            Offset += Str3.Length;
-            LineIndex += Str4.Count(f => f == '\n');
-            LogFile = (await Lfs.WriteLogDataAsync(LogFile, Offset, LineIndex, Encoding.ASCII.GetBytes(Str4), true, MaxChunkLength))!;
-			await Clock.AdvanceAsync(TimeSpan.FromSeconds(2.0));
-			Assert.AreEqual(Str1 + Str2 + Str3 + Str4,
-                await ReadLogFile(Lfs, LogFile, 0, Str1.Length + Str2.Length + Str3.Length + Str4.Length));
-            Assert.AreEqual(Str3 + Str4,
-                await ReadLogFile(Lfs, LogFile, Str1.Length + Str2.Length, Str3.Length + Str4.Length));
-            await AssertMetadata(Lfs, LogFile, Str1.Length + Str2.Length + Str3.Length + Str4.Length, 7);
-            await AssertChunk(Lfs, LogFile.Id, 2, 0, 0, Str1.Length, 0);
-            await AssertChunk(Lfs, LogFile.Id, 2, 1, Str1.Length, Str2.Length + Str3.Length + Str4.Length, 1);
-            await AssertLineOffset(Lfs, LogFile.Id, 0, 0, 0);
-            await AssertLineOffset(Lfs, LogFile.Id, 1, 1, 6);
-            await AssertLineOffset(Lfs, LogFile.Id, 2, 2, 10);
-            await AssertLineOffset(Lfs, LogFile.Id, 3, 3, 14);
-            await AssertLineOffset(Lfs, LogFile.Id, 4, 4, 18);
-            await AssertLineOffset(Lfs, LogFile.Id, 5, 5, 22);
-            await AssertLineOffset(Lfs, LogFile.Id, 6, 6, 27);
-            await AssertLineOffset(Lfs, LogFile.Id, 7, 7, 32);
+            offset += str3.Length;
+            lineIndex += str4.Count(f => f == '\n');
+            logFile = (await lfs.WriteLogDataAsync(logFile, offset, lineIndex, Encoding.ASCII.GetBytes(str4), true, maxChunkLength))!;
+			await _clock.AdvanceAsync(TimeSpan.FromSeconds(2.0));
+			Assert.AreEqual(str1 + str2 + str3 + str4,
+                await ReadLogFile(lfs, logFile, 0, str1.Length + str2.Length + str3.Length + str4.Length));
+            Assert.AreEqual(str3 + str4,
+                await ReadLogFile(lfs, logFile, str1.Length + str2.Length, str3.Length + str4.Length));
+            await AssertMetadata(lfs, logFile, str1.Length + str2.Length + str3.Length + str4.Length, 7);
+            await AssertChunk(lfs, logFile.Id, 2, 0, 0, str1.Length, 0);
+            await AssertChunk(lfs, logFile.Id, 2, 1, str1.Length, str2.Length + str3.Length + str4.Length, 1);
+            await AssertLineOffset(lfs, logFile.Id, 0, 0, 0);
+            await AssertLineOffset(lfs, logFile.Id, 1, 1, 6);
+            await AssertLineOffset(lfs, logFile.Id, 2, 2, 10);
+            await AssertLineOffset(lfs, logFile.Id, 3, 3, 14);
+            await AssertLineOffset(lfs, logFile.Id, 4, 4, 18);
+            await AssertLineOffset(lfs, logFile.Id, 5, 5, 22);
+            await AssertLineOffset(lfs, logFile.Id, 6, 6, 27);
+            await AssertLineOffset(lfs, logFile.Id, 7, 7, 32);
             
             // Fifth write with flush and data that will span more than chunk. Will become chunk #3
-            string A = "Lorem ipsum dolor sit amet\n";
-            string B = "consectetur adipiscing\n";
-            string Str5 = A + B;
+            string a = "Lorem ipsum dolor sit amet\n";
+            string b = "consectetur adipiscing\n";
+            string str5 = a + b;
             
-            Offset += Str4.Length;
-            LineIndex += Str5.Count(f => f == '\n');
+            offset += str4.Length;
+            lineIndex += str5.Count(f => f == '\n');
             
             // Using this single write below will fail the ReadLogFile assert below. A bug?
             //await LogFileService.WriteLogDataAsync(LogFile, Offset, LineIndex, Encoding.ASCII.GetBytes(Str5), true);
             
             // Dividing it in two like this will work however
-            LogFile = (await Lfs.WriteLogDataAsync(LogFile, Offset, LineIndex, Encoding.ASCII.GetBytes(A), false, MaxChunkLength))!;
-			LogFile = (await Lfs.WriteLogDataAsync(LogFile, Offset + A.Length, LineIndex + 1, Encoding.ASCII.GetBytes(B), true, MaxChunkLength))!;
+            logFile = (await lfs.WriteLogDataAsync(logFile, offset, lineIndex, Encoding.ASCII.GetBytes(a), false, maxChunkLength))!;
+			logFile = (await lfs.WriteLogDataAsync(logFile, offset + a.Length, lineIndex + 1, Encoding.ASCII.GetBytes(b), true, maxChunkLength))!;
 
-			await Clock.AdvanceAsync(TimeSpan.FromSeconds(2.0));
-			await AssertMetadata(Lfs, LogFile, Str1.Length + Str2.Length + Str3.Length + Str4.Length + Str5.Length, 9);
-            await AssertChunk(Lfs, LogFile.Id, 4, 0, 0, Str1.Length, 0);
-            await AssertChunk(Lfs, LogFile.Id, 4, 1, Str1.Length, Str2.Length + Str3.Length + Str4.Length, 1);
-            await AssertChunk(Lfs, LogFile.Id, 4, 2, Offset, A.Length, 7);
-            await AssertChunk(Lfs, LogFile.Id, 4, 3, Offset + A.Length, B.Length, 8);
+			await _clock.AdvanceAsync(TimeSpan.FromSeconds(2.0));
+			await AssertMetadata(lfs, logFile, str1.Length + str2.Length + str3.Length + str4.Length + str5.Length, 9);
+            await AssertChunk(lfs, logFile.Id, 4, 0, 0, str1.Length, 0);
+            await AssertChunk(lfs, logFile.Id, 4, 1, str1.Length, str2.Length + str3.Length + str4.Length, 1);
+            await AssertChunk(lfs, logFile.Id, 4, 2, offset, a.Length, 7);
+            await AssertChunk(lfs, logFile.Id, 4, 3, offset + a.Length, b.Length, 8);
 
-            Assert.AreEqual(Str5, await ReadLogFile(Lfs, LogFile, Offset, Str5.Length));
+            Assert.AreEqual(str5, await ReadLogFile(lfs, logFile, offset, str5.Length));
         }
 
         [TestMethod]
         public async Task GetLogFileTest()
         {
-            await GetDatabaseService().Database.DropCollectionAsync("LogFiles");
-            Assert.AreEqual(0, (await LogFileService.GetLogFilesAsync()).Count);
+            await GetDatabaseServiceSingleton().Database.DropCollectionAsync("LogFiles");
+            Assert.AreEqual(0, (await _logFileService.GetLogFilesAsync()).Count);
 
 			// Will implicitly test GetLogFileAsync(), AddCachedLogFile()
-			JobId JobId = JobId.GenerateNewId();
-            ObjectId SessionId = ObjectId.GenerateNewId();
-            ILogFile A = await LogFileService.CreateLogFileAsync(JobId, new ObjectId<ISession>(SessionId), LogType.Text);
-            ILogFile B = (await LogFileService.GetCachedLogFileAsync(A.Id))!;
-            Assert.AreEqual(A.JobId, B.JobId);
-            Assert.AreEqual(A.SessionId, B.SessionId);
-            Assert.AreEqual(A.Type, B.Type);
-
-            ILogFile? NotFound = await LogFileService.GetCachedLogFileAsync(LogId.GenerateNewId());
-            Assert.IsNull(NotFound);
-
-            await LogFileService.CreateLogFileAsync(JobId.GenerateNewId(), new ObjectId<ISession>(ObjectId.GenerateNewId()), LogType.Text);
-            Assert.AreEqual(2, (await LogFileService.GetLogFilesAsync()).Count);
+			JobId jobId = JobId.GenerateNewId();
+            ObjectId sessionId = ObjectId.GenerateNewId();
+            ILogFile a = await _logFileService.CreateLogFileAsync(jobId, new ObjectId<ISession>(sessionId), LogType.Text);
+            ILogFile b = (await _logFileService.GetCachedLogFileAsync(a.Id))!;
+            Assert.AreEqual(a.JobId, b.JobId);
+            Assert.AreEqual(a.SessionId, b.SessionId);
+            Assert.AreEqual(a.Type, b.Type);
+
+            ILogFile? notFound = await _logFileService.GetCachedLogFileAsync(LogId.GenerateNewId());
+            Assert.IsNull(notFound);
+
+            await _logFileService.CreateLogFileAsync(JobId.GenerateNewId(), new ObjectId<ISession>(ObjectId.GenerateNewId()), LogType.Text);
+            Assert.AreEqual(2, (await _logFileService.GetLogFilesAsync()).Count);
         }
 
         [TestMethod]
         public async Task AuthorizeForSession()
         {
-			JobId JobId = JobId.GenerateNewId();
-            ObjectId SessionId = ObjectId.GenerateNewId();
-            ILogFile LogFile = await LogFileService.CreateLogFileAsync(JobId, new ObjectId<ISession>(SessionId), LogType.Text);
-            ILogFile LogFileNoSession = await LogFileService.CreateLogFileAsync(JobId, null, LogType.Text);
+			JobId jobId = JobId.GenerateNewId();
+            ObjectId sessionId = ObjectId.GenerateNewId();
+            ILogFile logFile = await _logFileService.CreateLogFileAsync(jobId, new ObjectId<ISession>(sessionId), LogType.Text);
+            ILogFile logFileNoSession = await _logFileService.CreateLogFileAsync(jobId, null, LogType.Text);
 
-            var HasClaim = new ClaimsPrincipal(new ClaimsIdentity(new List<Claim>
+			ClaimsPrincipal hasClaim = new ClaimsPrincipal(new ClaimsIdentity(new List<Claim>
             {
-                new Claim(HordeClaimTypes.AgentSessionId, SessionId.ToString()),
+                new Claim(HordeClaimTypes.AgentSessionId, sessionId.ToString()),
             }, "TestAuthType"));
-            var HasNoClaim = new ClaimsPrincipal(new ClaimsIdentity(new List<Claim>
+            ClaimsPrincipal hasNoClaim = new ClaimsPrincipal(new ClaimsIdentity(new List<Claim>
             {
                 new Claim(HordeClaimTypes.AgentSessionId, "invalid-session-id"),
             }, "TestAuthType"));
 
-            Assert.IsTrue(LogFileService.AuthorizeForSession(LogFile, HasClaim));
-            Assert.IsFalse(LogFileService.AuthorizeForSession(LogFile, HasNoClaim));
-            Assert.IsFalse(LogFileService.AuthorizeForSession(LogFileNoSession, HasClaim));
+            Assert.IsTrue(LogFileService.AuthorizeForSession(logFile, hasClaim));
+            Assert.IsFalse(LogFileService.AuthorizeForSession(logFile, hasNoClaim));
+            Assert.IsFalse(LogFileService.AuthorizeForSession(logFileNoSession, hasClaim));
         }
 
 		[TestMethod]
 		public async Task ChunkSplitting()
 		{
-			JobId JobId = JobId.GenerateNewId();
-			ILogFile LogFile = await LogFileService.CreateLogFileAsync(JobId, null, LogType.Text);
+			JobId jobId = JobId.GenerateNewId();
+			ILogFile logFile = await _logFileService.CreateLogFileAsync(jobId, null, LogType.Text);
 
-			long Offset = 0;
+			long offset = 0;
 
 			const int MaxChunkSize = 4;
 			const int MaxSubChunkLineCount = 128;
 
-			byte[] Line1 = Encoding.UTF8.GetBytes("hello world\n");
-			await LogFileService.WriteLogDataAsync(LogFile, Offset, 0, Line1, false, MaxChunkSize, MaxSubChunkLineCount);
-			Offset += Line1.Length;
-
-			byte[] Line2 = Encoding.UTF8.GetBytes("ab\n");
-			await LogFileService.WriteLogDataAsync(LogFile, Offset, 1, Line2, false, MaxChunkSize, MaxSubChunkLineCount);
-			Offset += Line2.Length;
-
-			byte[] Line3 = Encoding.UTF8.GetBytes("a\n");
-			await LogFileService.WriteLogDataAsync(LogFile, Offset, 2, Line3, false, MaxChunkSize, MaxSubChunkLineCount);
-			Offset += Line3.Length;
-
-			byte[] Line4 = Encoding.UTF8.GetBytes("b\n");
-			await LogFileService.WriteLogDataAsync(LogFile, Offset, 3, Line4, false, MaxChunkSize, MaxSubChunkLineCount);
-			Offset += Line4.Length;
-
-			byte[] Line5 = Encoding.UTF8.GetBytes("a\nb\n");
-			await LogFileService.WriteLogDataAsync(LogFile, Offset, 4, Line5, false, MaxChunkSize, MaxSubChunkLineCount);
-			Offset += Line5.Length;
-
-			await LogFileService.FlushAsync();
-			LogFile = (await LogFileService.GetLogFileAsync(LogFile.Id))!;
-			Assert.AreEqual(4, LogFile.Chunks.Count);
-			Assert.AreEqual(6, LogFile.MaxLineIndex);
-
-			Assert.AreEqual(0, LogFile.Chunks[0].LineIndex);
-			Assert.AreEqual(1, LogFile.Chunks[1].LineIndex);
-			Assert.AreEqual(2, LogFile.Chunks[2].LineIndex);
-			Assert.AreEqual(4, LogFile.Chunks[3].LineIndex);
-
-			Assert.AreEqual(12, LogFile.Chunks[0].Length);
-			Assert.AreEqual(3, LogFile.Chunks[1].Length);
-			Assert.AreEqual(4, LogFile.Chunks[2].Length);
-			Assert.AreEqual(4, LogFile.Chunks[3].Length);
-
-			Assert.AreEqual(0, LogFile.Chunks.GetChunkForLine(0));
-			Assert.AreEqual(1, LogFile.Chunks.GetChunkForLine(1));
-			Assert.AreEqual(2, LogFile.Chunks.GetChunkForLine(2));
-			Assert.AreEqual(2, LogFile.Chunks.GetChunkForLine(3));
-			Assert.AreEqual(3, LogFile.Chunks.GetChunkForLine(4));
-			Assert.AreEqual(3, LogFile.Chunks.GetChunkForLine(5));
+			byte[] line1 = Encoding.UTF8.GetBytes("hello world\n");
+			await _logFileService.WriteLogDataAsync(logFile, offset, 0, line1, false, MaxChunkSize, MaxSubChunkLineCount);
+			offset += line1.Length;
+
+			byte[] line2 = Encoding.UTF8.GetBytes("ab\n");
+			await _logFileService.WriteLogDataAsync(logFile, offset, 1, line2, false, MaxChunkSize, MaxSubChunkLineCount);
+			offset += line2.Length;
+
+			byte[] line3 = Encoding.UTF8.GetBytes("a\n");
+			await _logFileService.WriteLogDataAsync(logFile, offset, 2, line3, false, MaxChunkSize, MaxSubChunkLineCount);
+			offset += line3.Length;
+
+			byte[] line4 = Encoding.UTF8.GetBytes("b\n");
+			await _logFileService.WriteLogDataAsync(logFile, offset, 3, line4, false, MaxChunkSize, MaxSubChunkLineCount);
+			offset += line4.Length;
+
+			byte[] line5 = Encoding.UTF8.GetBytes("a\nb\n");
+			await _logFileService.WriteLogDataAsync(logFile, offset, 4, line5, false, MaxChunkSize, MaxSubChunkLineCount);
+			offset += line5.Length;
+
+			await _logFileService.FlushAsync();
+			logFile = (await _logFileService.GetLogFileAsync(logFile.Id))!;
+			Assert.AreEqual(4, logFile.Chunks.Count);
+			Assert.AreEqual(6, logFile.MaxLineIndex);
+
+			Assert.AreEqual(0, logFile.Chunks[0].LineIndex);
+			Assert.AreEqual(1, logFile.Chunks[1].LineIndex);
+			Assert.AreEqual(2, logFile.Chunks[2].LineIndex);
+			Assert.AreEqual(4, logFile.Chunks[3].LineIndex);
+
+			Assert.AreEqual(12, logFile.Chunks[0].Length);
+			Assert.AreEqual(3, logFile.Chunks[1].Length);
+			Assert.AreEqual(4, logFile.Chunks[2].Length);
+			Assert.AreEqual(4, logFile.Chunks[3].Length);
+
+			Assert.AreEqual(0, logFile.Chunks.GetChunkForLine(0));
+			Assert.AreEqual(1, logFile.Chunks.GetChunkForLine(1));
+			Assert.AreEqual(2, logFile.Chunks.GetChunkForLine(2));
+			Assert.AreEqual(2, logFile.Chunks.GetChunkForLine(3));
+			Assert.AreEqual(3, logFile.Chunks.GetChunkForLine(4));
+			Assert.AreEqual(3, logFile.Chunks.GetChunkForLine(5));
 		}
 	}
 }
\ No newline at end of file
