Horde: Update Horde.Build.Tests to use standard coding conventions, and fix a number of issues caught by static analysis.

diff --git a/Engine/Source/Programs/Horde/Horde.Build.Tests/LogIndexingTests.cs b/Engine/Source/Programs/Horde/Horde.Build.Tests/LogIndexingTests.cs
--- a/Engine/Source/Programs/Horde/Horde.Build.Tests/LogIndexingTests.cs
+++ b/Engine/Source/Programs/Horde/Horde.Build.Tests/LogIndexingTests.cs
@@ -1,7 +1,6 @@
 // Copyright Epic Games, Inc. All Rights Reserved.
 
 using Horde.Build.Api;
-using Horde.Build.Collections;
 using Horde.Build.Logs;
 using Horde.Build.Logs.Builder;
 using Horde.Build.Logs.Storage;
@@ -14,7 +13,6 @@ using Microsoft.Extensions.Configuration;
 using Microsoft.Extensions.DependencyInjection;
 using Microsoft.Extensions.Logging;
 using Microsoft.VisualStudio.TestTools.UnitTesting;
-using MongoDB.Bson;
 using System;
 using System.Collections.Generic;
 using System.IO;
@@ -23,7 +21,6 @@ using System.Text;
 using System.Threading.Tasks;
 using Horde.Build.Collections.Impl;
 using HordeCommon;
-using Microsoft.Extensions.Logging.Abstractions;
 
 namespace Horde.Build.Tests
 {
@@ -32,106 +29,117 @@ namespace Horde.Build.Tests
 	[TestClass]
 	public class LogIndexingTests : DatabaseIntegrationTest
 	{
-		private readonly ILogFileService LogFileService;
-		private byte[] Data = Resources.TextFile;
+		private readonly ILogFileService _logFileService;
+		private readonly NullLogStorage _nullLogStorage;
+		private readonly LocalLogStorage _logStorage;
+		private byte[] _data = Resources.TextFile;
 
 		public LogIndexingTests()
 		{
-			LogFileCollection LogFileCollection = new LogFileCollection(GetDatabaseService());
+			LogFileCollection logFileCollection = new LogFileCollection(GetDatabaseServiceSingleton());
 
-			ServiceProvider ServiceProvider = new ServiceCollection()
-				.AddLogging(Builder => Builder.AddConsole().SetMinimumLevel(LogLevel.Debug))
+			ServiceProvider serviceProvider = new ServiceCollection()
+				.AddLogging(builder => builder.AddConsole().SetMinimumLevel(LogLevel.Debug))
 				.BuildServiceProvider();
 
-			ILoggerFactory LoggerFactory = ServiceProvider.GetRequiredService<ILoggerFactory>();
-			ILogger<LogFileService> Logger = LoggerFactory.CreateLogger<LogFileService>();
+			ILoggerFactory loggerFactory = serviceProvider.GetRequiredService<ILoggerFactory>();
+			ILogger<LogFileService> logger = loggerFactory.CreateLogger<LogFileService>();
 
 			// Just to satisfy the parameter, need to be fixed
-			IConfiguration Config = new ConfigurationBuilder().Build();
+			IConfiguration config = new ConfigurationBuilder().Build();
 
-			ILogBuilder LogBuilder = new LocalLogBuilder();
-			ILogStorage LogStorage = new LocalLogStorage(20, new NullLogStorage());
-			LogFileService = new LogFileService(LogFileCollection, null!, LogBuilder, LogStorage, new FakeClock(), Logger);
+			ILogBuilder logBuilder = new LocalLogBuilder();
+			_nullLogStorage = new NullLogStorage();
+			_logStorage = new LocalLogStorage(20, _nullLogStorage);
+			_logFileService = new LogFileService(logFileCollection, null!, logBuilder, _logStorage, new FakeClock(), logger);
+		}
+
+		protected override void Dispose(bool disposing)
+		{
+			base.Dispose(disposing);
+
+			_logStorage.Dispose();
+			_nullLogStorage.Dispose();
 		}
 
 		[TestMethod]
 		public async Task IndexTests()
 		{
-			JobId JobId = JobId.GenerateNewId();
-			ILogFile LogFile = await LogFileService.CreateLogFileAsync(JobId, null, LogType.Text);
+			JobId jobId = JobId.GenerateNewId();
+			ILogFile logFile = await _logFileService.CreateLogFileAsync(jobId, null, LogType.Text);
 
 			// Write the test data to the log file in blocks
-			int Offset = 0;
-			int LineIndex = 0;
-			while (Offset < Data.Length)
+			int offset = 0;
+			int lineIndex = 0;
+			while (offset < _data.Length)
 			{
-				int Length = 0;
-				int LineCount = 0;
+				int length = 0;
+				int lineCount = 0;
 
-				for(int Idx = Offset; Length == 0 || Idx < Math.Min(Data.Length, Offset + 1883); Idx++)
+				for(int idx = offset; idx < Math.Min(_data.Length, offset + 1883); idx++)
 				{
-					if(Data[Idx] == '\n')
+					if(_data[idx] == '\n')
 					{
-						Length = (Idx + 1) - Offset;
-						LineCount++;
+						length = (idx + 1) - offset;
+						lineCount++;
 						break;
 					}
 				}
 
-				LogFile = await WriteLogDataAsync(LogFile, Offset, LineIndex, Data.AsMemory(Offset, Length), false);
+				logFile = await WriteLogDataAsync(logFile, offset, lineIndex, _data.AsMemory(offset, length), false);
 
-				Offset += Length;
-				LineIndex += LineCount;
+				offset += length;
+				lineIndex += lineCount;
 			}
 
 			// Read the data back out and check it's the same
-			byte[] ReadData = new byte[Data.Length];
-			using (Stream Stream = await LogFileService.OpenRawStreamAsync(LogFile, 0, Data.Length))
+			byte[] readData = new byte[_data.Length];
+			using (Stream stream = await _logFileService.OpenRawStreamAsync(logFile, 0, _data.Length))
 			{
-				int ReadSize = await Stream.ReadAsync(ReadData, 0, ReadData.Length);
-				Assert.AreEqual(ReadData.Length, ReadSize);
+				int readSize = await stream.ReadAsync(readData, 0, readData.Length);
+				Assert.AreEqual(readData.Length, readSize);
 
-				int EqualSize = 0;
-				while(EqualSize < Data.Length && Data[EqualSize] == ReadData[EqualSize])
+				int equalSize = 0;
+				while(equalSize < _data.Length && _data[equalSize] == readData[equalSize])
 				{
-					EqualSize++;
+					equalSize++;
 				}
 
-				Assert.AreEqual(EqualSize, ReadSize);
+				Assert.AreEqual(equalSize, readSize);
 			}
 
 			// Test some searches
-			await SearchLogDataTestAsync(LogFile);
+			await SearchLogDataTestAsync(logFile);
 
 			// Generate an index and test again
-			LogFile = await WriteLogDataAsync(LogFile, Offset, LineIndex, Array.Empty<byte>(), true);
-			await SearchLogDataTestAsync(LogFile);
+			logFile = await WriteLogDataAsync(logFile, offset, lineIndex, Array.Empty<byte>(), true);
+			await SearchLogDataTestAsync(logFile);
 		}
 
 		[TestMethod]
 		public void TrieTests()
 		{
-			ReadOnlyTrieBuilder Builder = new ReadOnlyTrieBuilder();
+			ReadOnlyTrieBuilder builder = new ReadOnlyTrieBuilder();
 
-			ulong[] Values = { 1, 2, 3, 5, 8, 13, 21, 34, 55, 89, 144, 233, 377, 610, 987, 1597, 2584, 4181, 6765, 10946, 0xfedcba9876543210UL };
-			foreach (ulong Value in Values)
+			ulong[] values = { 1, 2, 3, 5, 8, 13, 21, 34, 55, 89, 144, 233, 377, 610, 987, 1597, 2584, 4181, 6765, 10946, 0xfedcba9876543210UL };
+			foreach (ulong value in values)
 			{
-				Builder.Add(Value);
+				builder.Add(value);
 			}
 
-			ReadOnlyTrie Trie = Builder.Build();
-			Assert.IsTrue(Enumerable.SequenceEqual(Trie.EnumerateRange(0, ulong.MaxValue), Values));
-			Assert.IsTrue(Enumerable.SequenceEqual(Trie.EnumerateRange(0, 90), Values.Where(x => x <= 90)));
-			Assert.IsTrue(Enumerable.SequenceEqual(Trie.EnumerateRange(2, 89), Values.Where(x => x >= 2 && x <= 89)));
+			ReadOnlyTrie trie = builder.Build();
+			Assert.IsTrue(Enumerable.SequenceEqual(trie.EnumerateRange(0, UInt64.MaxValue), values));
+			Assert.IsTrue(Enumerable.SequenceEqual(trie.EnumerateRange(0, 90), values.Where(x => x <= 90)));
+			Assert.IsTrue(Enumerable.SequenceEqual(trie.EnumerateRange(2, 89), values.Where(x => x >= 2 && x <= 89)));
 		}
 
 		[TestMethod]
 		public async Task PartialTokenTests()
 		{
-			JobId JobId = JobId.GenerateNewId();
-			ILogFile LogFile = await LogFileService.CreateLogFileAsync(JobId, null, LogType.Text);
+			JobId jobId = JobId.GenerateNewId();
+			ILogFile logFile = await _logFileService.CreateLogFileAsync(jobId, null, LogType.Text);
 
-			string[] Lines =
+			string[] lines =
 			{
 				"abcdefghi\n",
 				"jklmno123\n",
@@ -139,29 +147,29 @@ namespace Horde.Build.Tests
 				"wx\n"
 			};
 
-			int Length = 0;
-			for (int LineIdx = 0; LineIdx < Lines.Length; LineIdx++)
+			int length = 0;
+			for (int lineIdx = 0; lineIdx < lines.Length; lineIdx++)
 			{
-				LogFile = await WriteLogDataAsync(LogFile, Length, LineIdx, Encoding.UTF8.GetBytes(Lines[LineIdx]), true);
-				Length += Lines[LineIdx].Length;
+				logFile = await WriteLogDataAsync(logFile, length, lineIdx, Encoding.UTF8.GetBytes(lines[lineIdx]), true);
+				length += lines[lineIdx].Length;
 			}
 
-			for (int LineIdx = 0; LineIdx < Lines.Length; LineIdx++)
+			for (int lineIdx = 0; lineIdx < lines.Length; lineIdx++)
 			{
-				for (int StrLen = 1; StrLen < 7; StrLen++)
+				for (int strLen = 1; strLen < 7; strLen++)
 				{
-					for (int StrOfs = 0; StrOfs + StrLen < Lines[LineIdx].Length - 1; StrOfs++)
+					for (int strOfs = 0; strOfs + strLen < lines[lineIdx].Length - 1; strOfs++)
 					{
-						string Str = Lines[LineIdx].Substring(StrOfs, StrLen);
+						string str = lines[lineIdx].Substring(strOfs, strLen);
 
-						LogSearchStats Stats = new LogSearchStats();
-						List<int> Results = await LogFileService.SearchLogDataAsync(LogFile, Str, 0, 5, Stats);
-						Assert.AreEqual(1, Results.Count);
-						Assert.AreEqual(LineIdx, Results[0]);
+						LogSearchStats stats = new LogSearchStats();
+						List<int> results = await _logFileService.SearchLogDataAsync(logFile, str, 0, 5, stats);
+						Assert.AreEqual(1, results.Count);
+						Assert.AreEqual(lineIdx, results[0]);
 
-						Assert.AreEqual(1, Stats.NumScannedBlocks);
-						Assert.AreEqual(3, Stats.NumSkippedBlocks);
-						Assert.AreEqual(0, Stats.NumFalsePositiveBlocks);
+						Assert.AreEqual(1, stats.NumScannedBlocks);
+						Assert.AreEqual(3, stats.NumSkippedBlocks);
+						Assert.AreEqual(0, stats.NumFalsePositiveBlocks);
 					}
 				}
 			}
@@ -170,71 +178,71 @@ namespace Horde.Build.Tests
 		[TestMethod]
 		public async Task AppendIndexTests()
 		{
-			JobId JobId = JobId.GenerateNewId();
-			ILogFile LogFile = await LogFileService.CreateLogFileAsync(JobId, null, LogType.Text);
+			JobId jobId = JobId.GenerateNewId();
+			ILogFile logFile = await _logFileService.CreateLogFileAsync(jobId, null, LogType.Text);
 
-			LogFile = await WriteLogDataAsync(LogFile, 0, 0, Encoding.UTF8.GetBytes("abc\n"), true);
-			LogFile = await WriteLogDataAsync(LogFile, 4, 1, Encoding.UTF8.GetBytes("def\n"), true);
-			LogFile = await WriteLogDataAsync(LogFile, 8, 2, Encoding.UTF8.GetBytes("ghi\n"), false);
+			logFile = await WriteLogDataAsync(logFile, 0, 0, Encoding.UTF8.GetBytes("abc\n"), true);
+			logFile = await WriteLogDataAsync(logFile, 4, 1, Encoding.UTF8.GetBytes("def\n"), true);
+			logFile = await WriteLogDataAsync(logFile, 8, 2, Encoding.UTF8.GetBytes("ghi\n"), false);
 
-			await ((LogFileService)LogFileService).FlushPendingWritesAsync();
+			await ((LogFileService)_logFileService).FlushPendingWritesAsync();
 
 			{
-				LogSearchStats Stats = new LogSearchStats();
-				List<int> Results = await LogFileService.SearchLogDataAsync(LogFile, "abc", 0, 5, Stats);
-				Assert.AreEqual(1, Results.Count);
-				Assert.AreEqual(0, Results[0]);
-
-				Assert.AreEqual(2, Stats.NumScannedBlocks); // abc + ghi (no index yet because it hasn't been flushed)
-				Assert.AreEqual(1, Stats.NumSkippedBlocks); // def
-				Assert.AreEqual(0, Stats.NumFalsePositiveBlocks);
+				LogSearchStats stats = new LogSearchStats();
+				List<int> results = await _logFileService.SearchLogDataAsync(logFile, "abc", 0, 5, stats);
+				Assert.AreEqual(1, results.Count);
+				Assert.AreEqual(0, results[0]);
+
+				Assert.AreEqual(2, stats.NumScannedBlocks); // abc + ghi (no index yet because it hasn't been flushed)
+				Assert.AreEqual(1, stats.NumSkippedBlocks); // def
+				Assert.AreEqual(0, stats.NumFalsePositiveBlocks);
 			}
 			{
-				LogSearchStats Stats = new LogSearchStats();
-				List<int> Results = await LogFileService.SearchLogDataAsync(LogFile, "def", 0, 5, Stats);
-				Assert.AreEqual(1, Results.Count);
-				Assert.AreEqual(1, Results[0]);
-
-				Assert.AreEqual(2, Stats.NumScannedBlocks); // def + ghi (no index yet because it hasn't been flushed)
-				Assert.AreEqual(1, Stats.NumSkippedBlocks); // abc
-				Assert.AreEqual(0, Stats.NumFalsePositiveBlocks);
+				LogSearchStats stats = new LogSearchStats();
+				List<int> results = await _logFileService.SearchLogDataAsync(logFile, "def", 0, 5, stats);
+				Assert.AreEqual(1, results.Count);
+				Assert.AreEqual(1, results[0]);
+
+				Assert.AreEqual(2, stats.NumScannedBlocks); // def + ghi (no index yet because it hasn't been flushed)
+				Assert.AreEqual(1, stats.NumSkippedBlocks); // abc
+				Assert.AreEqual(0, stats.NumFalsePositiveBlocks);
 			}
 			{
-				LogSearchStats Stats = new LogSearchStats();
-				List<int> Results = await LogFileService.SearchLogDataAsync(LogFile, "ghi", 0, 5, Stats);
-				Assert.AreEqual(1, Results.Count);
-				Assert.AreEqual(2, Results[0]);
-
-				Assert.AreEqual(1, Stats.NumScannedBlocks); // ghi
-				Assert.AreEqual(2, Stats.NumSkippedBlocks); // abc + def
-				Assert.AreEqual(0, Stats.NumFalsePositiveBlocks);
+				LogSearchStats stats = new LogSearchStats();
+				List<int> results = await _logFileService.SearchLogDataAsync(logFile, "ghi", 0, 5, stats);
+				Assert.AreEqual(1, results.Count);
+				Assert.AreEqual(2, results[0]);
+
+				Assert.AreEqual(1, stats.NumScannedBlocks); // ghi
+				Assert.AreEqual(2, stats.NumSkippedBlocks); // abc + def
+				Assert.AreEqual(0, stats.NumFalsePositiveBlocks);
 			}
 		}
 
-		async Task<ILogFile> WriteLogDataAsync(ILogFile LogFile, long Offset, int LineIndex, ReadOnlyMemory<byte> Data, bool Flush)
+		async Task<ILogFile> WriteLogDataAsync(ILogFile logFile, long offset, int lineIndex, ReadOnlyMemory<byte> data, bool flush)
 		{
 			const int MaxChunkLength = 32 * 1024;
 			const int MaxSubChunkLineCount = 128;
 
-			ILogFile? NewLogFile = await LogFileService.WriteLogDataAsync(LogFile, Offset, LineIndex, Data, Flush, MaxChunkLength, MaxSubChunkLineCount);
-			Assert.IsNotNull(NewLogFile);
-			return NewLogFile!;
+			ILogFile? newLogFile = await _logFileService.WriteLogDataAsync(logFile, offset, lineIndex, data, flush, MaxChunkLength, MaxSubChunkLineCount);
+			Assert.IsNotNull(newLogFile);
+			return newLogFile!;
 		}
 
-		async Task SearchLogDataTestAsync(ILogFile LogFile)
+		async Task SearchLogDataTestAsync(ILogFile logFile)
 		{
-			await SearchLogDataTestAsync(LogFile, "HISPANIOLA", 0, 4, new[] { 1503, 1520, 1525, 1595 });
-			await SearchLogDataTestAsync(LogFile, "Hispaniola", 0, 4, new[] { 1503, 1520, 1525, 1595 });
-			await SearchLogDataTestAsync(LogFile, "HizpaniolZ", 0, 4, Array.Empty<int>());
-			await SearchLogDataTestAsync(LogFile, "Pieces of eight!", 0, 100, new[] { 2227, 2228, 5840, 5841, 7520 });
-			await SearchLogDataTestAsync(LogFile, "NEWSLETTER", 0, 100, new[] { 7886 });
+			await SearchLogDataTestAsync(logFile, "HISPANIOLA", 0, 4, new[] { 1503, 1520, 1525, 1595 });
+			await SearchLogDataTestAsync(logFile, "Hispaniola", 0, 4, new[] { 1503, 1520, 1525, 1595 });
+			await SearchLogDataTestAsync(logFile, "HizpaniolZ", 0, 4, Array.Empty<int>());
+			await SearchLogDataTestAsync(logFile, "Pieces of eight!", 0, 100, new[] { 2227, 2228, 5840, 5841, 7520 });
+			await SearchLogDataTestAsync(logFile, "NEWSLETTER", 0, 100, new[] { 7886 });
 		}
 
-		async Task SearchLogDataTestAsync(ILogFile LogFile, string Text, int FirstLine, int Count, int[] ExpectedLines)
+		async Task SearchLogDataTestAsync(ILogFile logFile, string text, int firstLine, int count, int[] expectedLines)
 		{
-			LogSearchStats Stats = new LogSearchStats();
-			List<int> Lines = await LogFileService.SearchLogDataAsync(LogFile, Text, FirstLine, Count, Stats);
-			Assert.IsTrue(Lines.SequenceEqual(ExpectedLines));
+			LogSearchStats stats = new LogSearchStats();
+			List<int> lines = await _logFileService.SearchLogDataAsync(logFile, text, firstLine, count, stats);
+			Assert.IsTrue(lines.SequenceEqual(expectedLines));
 		}
 	}
 }
diff --git a/Engine/Source/Programs/Horde/Horde.Build.Tests/Models/Job/GetLabelStatesTests.cs b/Engine/Source/Programs/Horde/Horde.Build.Tests/Models/Job/GetLabelStatesTests.cs
--- a/Engine/Source/Programs/Horde/Horde.Build.Tests/Models/Job/GetLabelStatesTests.cs
+++ b/Engine/Source/Programs/Horde/Horde.Build.Tests/Models/Job/GetLabelStatesTests.cs
@@ -5,7 +5,6 @@ using Horde.Build.Api;
 using HordeCommon;
 using Horde.Build.Models;
 using Microsoft.VisualStudio.TestTools.UnitTesting;
-using MongoDB.Bson;
 using Moq;
 using Horde.Build.Utilities;
 
@@ -16,218 +15,218 @@ namespace Horde.Build.Tests
 	[TestClass]
 	public class GetLabelStatesTests
 	{
-		public static INode MockNode(string Name)
+		public static INode MockNode(string name)
 		{
-			Mock<INode> Node = new Mock<INode>();
-			Node.SetupGet(x => x.Name).Returns(Name);
-			return Node.Object;
+			Mock<INode> node = new Mock<INode>();
+			node.SetupGet(x => x.Name).Returns(name);
+			return node.Object;
 		}
 
-		public static ILabel MockLabel(string LabelName, List<NodeRef> Nodes)
+		public static ILabel MockLabel(string labelName, List<NodeRef> nodes)
 		{
-			Mock<ILabel> Label = new Mock<ILabel>();
-			Label.SetupGet(x => x.DashboardName).Returns(LabelName);
-			Label.SetupGet(x => x.IncludedNodes).Returns(Nodes);
-			Label.SetupGet(x => x.RequiredNodes).Returns(Nodes);
-			return Label.Object;
+			Mock<ILabel> label = new Mock<ILabel>();
+			label.SetupGet(x => x.DashboardName).Returns(labelName);
+			label.SetupGet(x => x.IncludedNodes).Returns(nodes);
+			label.SetupGet(x => x.RequiredNodes).Returns(nodes);
+			return label.Object;
 		}
 
 		public static IGraph CreateGraph()
 		{
-			List<INode> Nodes1 = new List<INode>();
-			Nodes1.Add(MockNode("Update Version Files"));
-			Nodes1.Add(MockNode("Compile UnrealHeaderTool Win64"));
-			Nodes1.Add(MockNode("Compile ShooterGameEditor Win64"));
-			Nodes1.Add(MockNode("Cook ShooterGame Win64"));
-
-			List<INode> Nodes2 = new List<INode>();
-			Nodes2.Add(MockNode("Compile UnrealHeaderTool Mac"));
-			Nodes2.Add(MockNode("Compile ShooterGameEditor Mac"));
-			Nodes2.Add(MockNode("Cook ShooterGame Mac"));
-
-			List<ILabel> Labels = new List<ILabel>();
-			Labels.Add(MockLabel("Label_AllGroup1Nodes", new List<NodeRef>() { new NodeRef(0, 0), new NodeRef(0, 1), new NodeRef(0, 2), new NodeRef(0, 3) }));
-			Labels.Add(MockLabel("Label_Win64UHTOnly", new List<NodeRef>() { new NodeRef(0, 1) }));
-			Labels.Add(MockLabel("Label_Group1TwoNodes", new List<NodeRef>() { new NodeRef(0, 0), new NodeRef(0, 3) }));
-			Labels.Add(MockLabel("Label_BothGroups", new List<NodeRef>() { new NodeRef(0, 0), new NodeRef(1, 0) }));
-
-			Mock<INodeGroup> Group1 = new Mock<INodeGroup>(MockBehavior.Strict);
-			Group1.SetupGet(x => x.Nodes).Returns(Nodes1);
-
-			Mock<INodeGroup> Group2 = new Mock<INodeGroup>(MockBehavior.Strict);
-			Group2.SetupGet(x => x.Nodes).Returns(Nodes2);
-
-			Mock<IGraph> GraphMock = new Mock<IGraph>(MockBehavior.Strict);
-			GraphMock.SetupGet(x => x.Groups).Returns(new List<INodeGroup> { Group1.Object, Group2.Object });
-			GraphMock.SetupGet(x => x.Labels).Returns(Labels);
-			return GraphMock.Object;
+			List<INode> nodes1 = new List<INode>();
+			nodes1.Add(MockNode("Update Version Files"));
+			nodes1.Add(MockNode("Compile UnrealHeaderTool Win64"));
+			nodes1.Add(MockNode("Compile ShooterGameEditor Win64"));
+			nodes1.Add(MockNode("Cook ShooterGame Win64"));
+
+			List<INode> nodes2 = new List<INode>();
+			nodes2.Add(MockNode("Compile UnrealHeaderTool Mac"));
+			nodes2.Add(MockNode("Compile ShooterGameEditor Mac"));
+			nodes2.Add(MockNode("Cook ShooterGame Mac"));
+
+			List<ILabel> labels = new List<ILabel>();
+			labels.Add(MockLabel("Label_AllGroup1Nodes", new List<NodeRef>() { new NodeRef(0, 0), new NodeRef(0, 1), new NodeRef(0, 2), new NodeRef(0, 3) }));
+			labels.Add(MockLabel("Label_Win64UHTOnly", new List<NodeRef>() { new NodeRef(0, 1) }));
+			labels.Add(MockLabel("Label_Group1TwoNodes", new List<NodeRef>() { new NodeRef(0, 0), new NodeRef(0, 3) }));
+			labels.Add(MockLabel("Label_BothGroups", new List<NodeRef>() { new NodeRef(0, 0), new NodeRef(1, 0) }));
+
+			Mock<INodeGroup> group1 = new Mock<INodeGroup>(MockBehavior.Strict);
+			group1.SetupGet(x => x.Nodes).Returns(nodes1);
+
+			Mock<INodeGroup> group2 = new Mock<INodeGroup>(MockBehavior.Strict);
+			group2.SetupGet(x => x.Nodes).Returns(nodes2);
+
+			Mock<IGraph> graphMock = new Mock<IGraph>(MockBehavior.Strict);
+			graphMock.SetupGet(x => x.Groups).Returns(new List<INodeGroup> { group1.Object, group2.Object });
+			graphMock.SetupGet(x => x.Labels).Returns(labels);
+			return graphMock.Object;
 		}
 
-		public static IJob CreateGetLabelStatesJob(string Name, IGraph Graph, List<(JobStepState, JobStepOutcome)> JobStepData)
+		public static IJob CreateGetLabelStatesJob(string name, IGraph graph, List<(JobStepState, JobStepOutcome)> jobStepData)
 		{
-			JobId JobId = JobId.Parse("5ec16da1774cb4000107c2c1");
-			List<IJobStepBatch> Batches = new List<IJobStepBatch>();
-			int DataIdx = 0;
-			for (int GroupIdx = 0; GroupIdx < Graph.Groups.Count; GroupIdx++)
+			JobId jobId = JobId.Parse("5ec16da1774cb4000107c2c1");
+			List<IJobStepBatch> batches = new List<IJobStepBatch>();
+			int dataIdx = 0;
+			for (int groupIdx = 0; groupIdx < graph.Groups.Count; groupIdx++)
 			{
-				INodeGroup Group = Graph.Groups[GroupIdx];
+				INodeGroup @group = graph.Groups[groupIdx];
 
-				List<IJobStep> Steps = new List<IJobStep>();
-				for (int NodeIdx = 0; NodeIdx < Group.Nodes.Count; NodeIdx++, DataIdx++)
+				List<IJobStep> steps = new List<IJobStep>();
+				for (int nodeIdx = 0; nodeIdx < @group.Nodes.Count; nodeIdx++, dataIdx++)
 				{
-					SubResourceId StepId = new SubResourceId((ushort)((GroupIdx * 100) + NodeIdx));
-
-					Mock<IJobStep> Step = new Mock<IJobStep>(MockBehavior.Strict);
-					Step.SetupGet(x => x.Id).Returns(StepId);
-					Step.SetupGet(x => x.NodeIdx).Returns(NodeIdx);
-					Step.SetupGet(x => x.State).Returns(JobStepData[DataIdx].Item1);
-					Step.SetupGet(x => x.Outcome).Returns(JobStepData[DataIdx].Item2);
-					Steps.Add(Step.Object);
+					SubResourceId stepId = new SubResourceId((ushort)((groupIdx * 100) + nodeIdx));
+
+					Mock<IJobStep> step = new Mock<IJobStep>(MockBehavior.Strict);
+					step.SetupGet(x => x.Id).Returns(stepId);
+					step.SetupGet(x => x.NodeIdx).Returns(nodeIdx);
+					step.SetupGet(x => x.State).Returns(jobStepData[dataIdx].Item1);
+					step.SetupGet(x => x.Outcome).Returns(jobStepData[dataIdx].Item2);
+					steps.Add(step.Object);
 				}
 
-				SubResourceId BatchId = new SubResourceId((ushort)(GroupIdx * 100));
+				SubResourceId batchId = new SubResourceId((ushort)(groupIdx * 100));
 
-				Mock<IJobStepBatch> Batch = new Mock<IJobStepBatch>(MockBehavior.Strict);
-				Batch.SetupGet(x => x.Id).Returns(BatchId);
-				Batch.SetupGet(x => x.GroupIdx).Returns(GroupIdx);
-				Batch.SetupGet(x => x.Steps).Returns(Steps);
-				Batches.Add(Batch.Object);
+				Mock<IJobStepBatch> batch = new Mock<IJobStepBatch>(MockBehavior.Strict);
+				batch.SetupGet(x => x.Id).Returns(batchId);
+				batch.SetupGet(x => x.GroupIdx).Returns(groupIdx);
+				batch.SetupGet(x => x.Steps).Returns(steps);
+				batches.Add(batch.Object);
 			}
 
-			Mock<IJob> Job = new Mock<IJob>(MockBehavior.Strict);
-			Job.SetupGet(x => x.Id).Returns(JobId);
-			Job.SetupGet(x => x.Name).Returns(Name);
-			Job.SetupGet(x => x.Batches).Returns(Batches);
-			return Job.Object;
+			Mock<IJob> job = new Mock<IJob>(MockBehavior.Strict);
+			job.SetupGet(x => x.Id).Returns(jobId);
+			job.SetupGet(x => x.Name).Returns(name);
+			job.SetupGet(x => x.Batches).Returns(batches);
+			return job.Object;
 		}
 
-		public static List<(JobStepState, JobStepOutcome)> GetJobStepDataList(JobStepState State, JobStepOutcome Outcome)
+		public static List<(JobStepState, JobStepOutcome)> GetJobStepDataList(JobStepState state, JobStepOutcome outcome)
 		{
-			List<(JobStepState, JobStepOutcome)> Data = new List<(JobStepState, JobStepOutcome)>();
+			List<(JobStepState, JobStepOutcome)> data = new List<(JobStepState, JobStepOutcome)>();
 			for (int i = 0; i < 7; ++i)
 			{
-				Data.Add((State, Outcome));
+				data.Add((state, outcome));
 			}
-			return Data;
+			return data;
 		}
 
 		[TestMethod]
-		public void GetLabelStates_InitialStateTest()
+		public void GetLabelStatesInitialStateTest()
 		{
 			// Verify initial state. All job steps are unspecified, so all labels unspecified.
-			IGraph Graph = CreateGraph();
-			IJob Job = CreateGetLabelStatesJob("Job", Graph, GetJobStepDataList(JobStepState.Unspecified, JobStepOutcome.Unspecified));
-			IReadOnlyList<(LabelState, LabelOutcome)> LabelStates = Job.GetLabelStates(Graph);
-			Assert.AreEqual(Graph.Labels.Count, LabelStates.Count);
-			foreach ((LabelState State, LabelOutcome Outcome) Result in LabelStates)
+			IGraph graph = CreateGraph();
+			IJob job = CreateGetLabelStatesJob("Job", graph, GetJobStepDataList(JobStepState.Unspecified, JobStepOutcome.Unspecified));
+			IReadOnlyList<(LabelState, LabelOutcome)> labelStates = job.GetLabelStates(graph);
+			Assert.AreEqual(graph.Labels.Count, labelStates.Count);
+			foreach ((LabelState State, LabelOutcome Outcome) result in labelStates)
 			{
-				Assert.IsTrue(Result.State == LabelState.Running);
-				Assert.IsTrue(Result.Outcome == LabelOutcome.Success);
+				Assert.IsTrue(result.State == LabelState.Running);
+				Assert.IsTrue(result.Outcome == LabelOutcome.Success);
 			}
 		}
 
 		[TestMethod]
-		public void GetLabelStates_RunningStateTest()
+		public void GetLabelStatesRunningStateTest()
 		{
 			// Verify that if any nodes are in the ready or running state that the approrpirate labels switch to running.
-			List<(JobStepState, JobStepOutcome)> JobStepData = GetJobStepDataList(JobStepState.Unspecified, JobStepOutcome.Unspecified);
-			JobStepData[3] = (JobStepState.Ready, JobStepOutcome.Unspecified);
-			JobStepData[4] = (JobStepState.Running, JobStepOutcome.Unspecified);
-			IGraph Graph = CreateGraph();
-			IJob Job = CreateGetLabelStatesJob("Job", Graph, JobStepData);
-			IReadOnlyList<(LabelState State, LabelOutcome Outcome)> LabelStates = Job.GetLabelStates(Graph);
-			Assert.AreEqual(Graph.Labels.Count, LabelStates.Count);
-			Assert.IsTrue(LabelStates[0].State == LabelState.Running);
-			Assert.IsTrue(LabelStates[0].Outcome == LabelOutcome.Success);
-			Assert.IsTrue(LabelStates[1].State == LabelState.Running);
-			Assert.IsTrue(LabelStates[1].Outcome == LabelOutcome.Success);
-			Assert.IsTrue(LabelStates[2].State == LabelState.Running);
-			Assert.IsTrue(LabelStates[2].Outcome == LabelOutcome.Success);
-			Assert.IsTrue(LabelStates[3].State == LabelState.Running);
-			Assert.IsTrue(LabelStates[3].Outcome == LabelOutcome.Success);
+			List<(JobStepState, JobStepOutcome)> jobStepData = GetJobStepDataList(JobStepState.Unspecified, JobStepOutcome.Unspecified);
+			jobStepData[3] = (JobStepState.Ready, JobStepOutcome.Unspecified);
+			jobStepData[4] = (JobStepState.Running, JobStepOutcome.Unspecified);
+			IGraph graph = CreateGraph();
+			IJob job = CreateGetLabelStatesJob("Job", graph, jobStepData);
+			IReadOnlyList<(LabelState State, LabelOutcome Outcome)> labelStates = job.GetLabelStates(graph);
+			Assert.AreEqual(graph.Labels.Count, labelStates.Count);
+			Assert.IsTrue(labelStates[0].State == LabelState.Running);
+			Assert.IsTrue(labelStates[0].Outcome == LabelOutcome.Success);
+			Assert.IsTrue(labelStates[1].State == LabelState.Running);
+			Assert.IsTrue(labelStates[1].Outcome == LabelOutcome.Success);
+			Assert.IsTrue(labelStates[2].State == LabelState.Running);
+			Assert.IsTrue(labelStates[2].Outcome == LabelOutcome.Success);
+			Assert.IsTrue(labelStates[3].State == LabelState.Running);
+			Assert.IsTrue(labelStates[3].Outcome == LabelOutcome.Success);
 		}
 
 		[TestMethod]
-		public void GetLabelStates_WarningOutcomeTest()
+		public void GetLabelStatesWarningOutcomeTest()
 		{
 			// Verify that if a step has a warning, any applicable labels are set to warning.
-			List<(JobStepState, JobStepOutcome)> JobStepData = GetJobStepDataList(JobStepState.Ready, JobStepOutcome.Unspecified);
-			JobStepData[3] = (JobStepState.Running, JobStepOutcome.Warnings);
-			IGraph Graph = CreateGraph();
-			IJob Job = CreateGetLabelStatesJob("Job", Graph, JobStepData);
-			IReadOnlyList<(LabelState State, LabelOutcome Outcome)> LabelStates = Job.GetLabelStates(Graph);
-			Assert.AreEqual(Graph.Labels.Count, LabelStates.Count);
-			Assert.IsTrue(LabelStates[0].State == LabelState.Running);
-			Assert.IsTrue(LabelStates[0].Outcome == LabelOutcome.Warnings);
-			Assert.IsTrue(LabelStates[1].State == LabelState.Running);
-			Assert.IsTrue(LabelStates[1].Outcome == LabelOutcome.Success);
-			Assert.IsTrue(LabelStates[2].State == LabelState.Running);
-			Assert.IsTrue(LabelStates[2].Outcome == LabelOutcome.Warnings);
-			Assert.IsTrue(LabelStates[3].State == LabelState.Running);
-			Assert.IsTrue(LabelStates[3].Outcome == LabelOutcome.Success);
+			List<(JobStepState, JobStepOutcome)> jobStepData = GetJobStepDataList(JobStepState.Ready, JobStepOutcome.Unspecified);
+			jobStepData[3] = (JobStepState.Running, JobStepOutcome.Warnings);
+			IGraph graph = CreateGraph();
+			IJob job = CreateGetLabelStatesJob("Job", graph, jobStepData);
+			IReadOnlyList<(LabelState State, LabelOutcome Outcome)> labelStates = job.GetLabelStates(graph);
+			Assert.AreEqual(graph.Labels.Count, labelStates.Count);
+			Assert.IsTrue(labelStates[0].State == LabelState.Running);
+			Assert.IsTrue(labelStates[0].Outcome == LabelOutcome.Warnings);
+			Assert.IsTrue(labelStates[1].State == LabelState.Running);
+			Assert.IsTrue(labelStates[1].Outcome == LabelOutcome.Success);
+			Assert.IsTrue(labelStates[2].State == LabelState.Running);
+			Assert.IsTrue(labelStates[2].Outcome == LabelOutcome.Warnings);
+			Assert.IsTrue(labelStates[3].State == LabelState.Running);
+			Assert.IsTrue(labelStates[3].Outcome == LabelOutcome.Success);
 		}
 
 		[TestMethod]
-		public void GetLabelStates_ErrorOutcomeTest()
+		public void GetLabelStatesErrorOutcomeTest()
 		{
 			// Verify that if a step has an error, any applicable labels are set to failure.
-			List<(JobStepState, JobStepOutcome)> JobStepData = GetJobStepDataList(JobStepState.Ready, JobStepOutcome.Unspecified);
-			JobStepData[2] = (JobStepState.Running, JobStepOutcome.Warnings);
-			JobStepData[4] = (JobStepState.Running, JobStepOutcome.Failure);
-			IGraph Graph = CreateGraph();
-			IJob Job = CreateGetLabelStatesJob("Job", Graph, JobStepData);
-			IReadOnlyList<(LabelState State, LabelOutcome Outcome)> LabelStates = Job.GetLabelStates(Graph);
-			Assert.AreEqual(Graph.Labels.Count, LabelStates.Count);
-			Assert.IsTrue(LabelStates[0].State == LabelState.Running);
-			Assert.IsTrue(LabelStates[0].Outcome == LabelOutcome.Warnings);
-			Assert.IsTrue(LabelStates[1].State == LabelState.Running);
-			Assert.IsTrue(LabelStates[1].Outcome == LabelOutcome.Success);
-			Assert.IsTrue(LabelStates[2].State == LabelState.Running);
-			Assert.IsTrue(LabelStates[2].Outcome == LabelOutcome.Success);
-			Assert.IsTrue(LabelStates[3].State == LabelState.Running);
-			Assert.IsTrue(LabelStates[3].Outcome == LabelOutcome.Failure);
+			List<(JobStepState, JobStepOutcome)> jobStepData = GetJobStepDataList(JobStepState.Ready, JobStepOutcome.Unspecified);
+			jobStepData[2] = (JobStepState.Running, JobStepOutcome.Warnings);
+			jobStepData[4] = (JobStepState.Running, JobStepOutcome.Failure);
+			IGraph graph = CreateGraph();
+			IJob job = CreateGetLabelStatesJob("Job", graph, jobStepData);
+			IReadOnlyList<(LabelState State, LabelOutcome Outcome)> labelStates = job.GetLabelStates(graph);
+			Assert.AreEqual(graph.Labels.Count, labelStates.Count);
+			Assert.IsTrue(labelStates[0].State == LabelState.Running);
+			Assert.IsTrue(labelStates[0].Outcome == LabelOutcome.Warnings);
+			Assert.IsTrue(labelStates[1].State == LabelState.Running);
+			Assert.IsTrue(labelStates[1].Outcome == LabelOutcome.Success);
+			Assert.IsTrue(labelStates[2].State == LabelState.Running);
+			Assert.IsTrue(labelStates[2].Outcome == LabelOutcome.Success);
+			Assert.IsTrue(labelStates[3].State == LabelState.Running);
+			Assert.IsTrue(labelStates[3].Outcome == LabelOutcome.Failure);
 		}
 
 		[TestMethod]
-		public void GetLabelStates_CompleteStateTest()
+		public void GetLabelStatesCompleteStateTest()
 		{
 			// Verify that if a label's steps are all complete the state is set to complete.
-			List<(JobStepState, JobStepOutcome)> JobStepData = GetJobStepDataList(JobStepState.Completed, JobStepOutcome.Success);
-			JobStepData[2] = (JobStepState.Running, JobStepOutcome.Failure);
-			JobStepData[4] = (JobStepState.Ready, JobStepOutcome.Unspecified);
-			IGraph Graph = CreateGraph();
-			IJob Job = CreateGetLabelStatesJob("Job", Graph, JobStepData);
-			IReadOnlyList<(LabelState State, LabelOutcome Outcome)> LabelStates = Job.GetLabelStates(Graph);
-			Assert.AreEqual(Graph.Labels.Count, LabelStates.Count);
-			Assert.IsTrue(LabelStates[0].State == LabelState.Running);
-			Assert.IsTrue(LabelStates[0].Outcome == LabelOutcome.Failure);
-			Assert.IsTrue(LabelStates[1].State == LabelState.Complete);
-			Assert.IsTrue(LabelStates[1].Outcome == LabelOutcome.Success);
-			Assert.IsTrue(LabelStates[2].State == LabelState.Complete);
-			Assert.IsTrue(LabelStates[2].Outcome == LabelOutcome.Success);
-			Assert.IsTrue(LabelStates[3].State == LabelState.Running);
-			Assert.IsTrue(LabelStates[3].Outcome == LabelOutcome.Success);
+			List<(JobStepState, JobStepOutcome)> jobStepData = GetJobStepDataList(JobStepState.Completed, JobStepOutcome.Success);
+			jobStepData[2] = (JobStepState.Running, JobStepOutcome.Failure);
+			jobStepData[4] = (JobStepState.Ready, JobStepOutcome.Unspecified);
+			IGraph graph = CreateGraph();
+			IJob job = CreateGetLabelStatesJob("Job", graph, jobStepData);
+			IReadOnlyList<(LabelState State, LabelOutcome Outcome)> labelStates = job.GetLabelStates(graph);
+			Assert.AreEqual(graph.Labels.Count, labelStates.Count);
+			Assert.IsTrue(labelStates[0].State == LabelState.Running);
+			Assert.IsTrue(labelStates[0].Outcome == LabelOutcome.Failure);
+			Assert.IsTrue(labelStates[1].State == LabelState.Complete);
+			Assert.IsTrue(labelStates[1].Outcome == LabelOutcome.Success);
+			Assert.IsTrue(labelStates[2].State == LabelState.Complete);
+			Assert.IsTrue(labelStates[2].Outcome == LabelOutcome.Success);
+			Assert.IsTrue(labelStates[3].State == LabelState.Running);
+			Assert.IsTrue(labelStates[3].Outcome == LabelOutcome.Success);
 		}
 
 		[TestMethod]
-		public void GetLabelStates_SkippedAbortedStepTest()
+		public void GetLabelStatesSkippedAbortedStepTest()
 		{
 			// Verify that if a label's steps are complete and successful but any were skipped or aborted the result is unspecified.
-			List<(JobStepState, JobStepOutcome)> JobStepData = GetJobStepDataList(JobStepState.Completed, JobStepOutcome.Success);
-			JobStepData[2] = (JobStepState.Aborted, JobStepOutcome.Unspecified);
-			JobStepData[4] = (JobStepState.Skipped, JobStepOutcome.Unspecified);
-			IGraph Graph = CreateGraph();
-			IJob Job = CreateGetLabelStatesJob("Job", Graph, JobStepData);
-			IReadOnlyList<(LabelState State, LabelOutcome Outcome)> LabelStates = Job.GetLabelStates(Graph);
-			Assert.AreEqual(Graph.Labels.Count, LabelStates.Count);
-			Assert.IsTrue(LabelStates[0].State == LabelState.Complete);
-			Assert.IsTrue(LabelStates[0].Outcome == LabelOutcome.Unspecified);
-			Assert.IsTrue(LabelStates[1].State == LabelState.Complete);
-			Assert.IsTrue(LabelStates[1].Outcome == LabelOutcome.Success);
-			Assert.IsTrue(LabelStates[2].State == LabelState.Complete);
-			Assert.IsTrue(LabelStates[2].Outcome == LabelOutcome.Success);
-			Assert.IsTrue(LabelStates[3].State == LabelState.Complete);
-			Assert.IsTrue(LabelStates[3].Outcome == LabelOutcome.Unspecified);
+			List<(JobStepState, JobStepOutcome)> jobStepData = GetJobStepDataList(JobStepState.Completed, JobStepOutcome.Success);
+			jobStepData[2] = (JobStepState.Aborted, JobStepOutcome.Unspecified);
+			jobStepData[4] = (JobStepState.Skipped, JobStepOutcome.Unspecified);
+			IGraph graph = CreateGraph();
+			IJob job = CreateGetLabelStatesJob("Job", graph, jobStepData);
+			IReadOnlyList<(LabelState State, LabelOutcome Outcome)> labelStates = job.GetLabelStates(graph);
+			Assert.AreEqual(graph.Labels.Count, labelStates.Count);
+			Assert.IsTrue(labelStates[0].State == LabelState.Complete);
+			Assert.IsTrue(labelStates[0].Outcome == LabelOutcome.Unspecified);
+			Assert.IsTrue(labelStates[1].State == LabelState.Complete);
+			Assert.IsTrue(labelStates[1].Outcome == LabelOutcome.Success);
+			Assert.IsTrue(labelStates[2].State == LabelState.Complete);
+			Assert.IsTrue(labelStates[2].Outcome == LabelOutcome.Success);
+			Assert.IsTrue(labelStates[3].State == LabelState.Complete);
+			Assert.IsTrue(labelStates[3].Outcome == LabelOutcome.Unspecified);
 		}
 	}
 }
diff --git a/Engine/Source/Programs/Horde/Horde.Build.Tests/NotificationServiceTests.cs b/Engine/Source/Programs/Horde/Horde.Build.Tests/NotificationServiceTests.cs
--- a/Engine/Source/Programs/Horde/Horde.Build.Tests/NotificationServiceTests.cs
+++ b/Engine/Source/Programs/Horde/Horde.Build.Tests/NotificationServiceTests.cs
@@ -22,101 +22,95 @@ namespace Horde.Build.Tests
 
 	public class FakeNotificationSink : INotificationSink
 	{
-		public List<JobScheduledNotification> JobScheduledNotifications = new();
-		public int JobScheduledCallCount;
+		public List<JobScheduledNotification> JobScheduledNotifications { get; } = new();
+		public int JobScheduledCallCount { get; set; }
 		
-		public Task NotifyJobScheduledAsync(List<JobScheduledNotification> Notifications)
+		public Task NotifyJobScheduledAsync(List<JobScheduledNotification> notifications)
 		{
-			JobScheduledNotifications.AddRange(Notifications);
+			JobScheduledNotifications.AddRange(notifications);
 			JobScheduledCallCount++;
 			return Task.CompletedTask;
 		}
 
-		public Task NotifyJobCompleteAsync(IStream JobStream, IJob Job, IGraph Graph, LabelOutcome Outcome) { throw new NotImplementedException(); }
-		public Task NotifyJobCompleteAsync(IUser User, IStream JobStream, IJob Job, IGraph Graph, LabelOutcome Outcome) { throw new NotImplementedException(); }
-		public Task NotifyJobStepCompleteAsync(IUser User, IStream JobStream, IJob Job, IJobStepBatch Batch, IJobStep Step, INode Node, List<ILogEventData> JobStepEventData) { throw new NotImplementedException(); }
-		public Task NotifyLabelCompleteAsync(IUser User, IJob Job, IStream Stream, ILabel Label, int LabelIdx, LabelOutcome Outcome, List<(string, JobStepOutcome, Uri)> StepData) { throw new NotImplementedException(); }
-		public Task NotifyIssueUpdatedAsync(IIssue Issue) { throw new NotImplementedException(); }
-		public Task NotifyConfigUpdateFailureAsync(string ErrorMessage, string FileName, int? Change = null, IUser? Author = null, string? Description = null) { throw new NotImplementedException(); }
-		public Task NotifyDeviceServiceAsync(string Message, IDevice? Device = null, IDevicePool? Pool = null, IStream? Stream = null, IJob? Job = null, IJobStep? Step = null, INode? Node = null, IUser? User = null) { throw new NotImplementedException(); }
+		public Task NotifyJobCompleteAsync(IStream jobStream, IJob job, IGraph graph, LabelOutcome outcome) { throw new NotImplementedException(); }
+		public Task NotifyJobCompleteAsync(IUser user, IStream jobStream, IJob job, IGraph graph, LabelOutcome outcome) { throw new NotImplementedException(); }
+		public Task NotifyJobStepCompleteAsync(IUser user, IStream jobStream, IJob job, IJobStepBatch batch, IJobStep step, INode node, List<ILogEventData> jobStepEventData) { throw new NotImplementedException(); }
+		public Task NotifyLabelCompleteAsync(IUser user, IJob job, IStream stream, ILabel label, int labelIdx, LabelOutcome outcome, List<(string, JobStepOutcome, Uri)> stepData) { throw new NotImplementedException(); }
+		public Task NotifyIssueUpdatedAsync(IIssue issue) { throw new NotImplementedException(); }
+		public Task NotifyConfigUpdateFailureAsync(string errorMessage, string fileName, int? change = null, IUser? author = null, string? description = null) { throw new NotImplementedException(); }
+		public Task NotifyDeviceServiceAsync(string message, IDevice? device = null, IDevicePool? pool = null, IStream? stream = null, IJob? job = null, IJobStep? step = null, INode? node = null, IUser? user = null) { throw new NotImplementedException(); }
 	}
-	
 
 	[TestClass]
 	public class NotificationServiceTests : TestSetup
 	{
-		public IJob CreateJob(StreamId StreamId, int Change, string Name, IGraph Graph)
+		protected override void ConfigureServices(IServiceCollection services)
 		{
-			JobId JobId = new JobId("5ec16da1774cb4000107c2c1");
+			base.ConfigureServices(services);
 
-			List<IJobStepBatch> Batches = new List<IJobStepBatch>();
-			for (int GroupIdx = 0; GroupIdx < Graph.Groups.Count; GroupIdx++)
+			services.AddSingleton<FakeNotificationSink>();
+			services.AddSingleton<INotificationSink>(SP => SP.GetRequiredService<FakeNotificationSink>());
+		}
+
+		public static IJob CreateJob(StreamId streamId, int change, string name, IGraph graph)
+		{
+			JobId jobId = new JobId("5ec16da1774cb4000107c2c1");
+
+			List<IJobStepBatch> batches = new List<IJobStepBatch>();
+			for (int groupIdx = 0; groupIdx < graph.Groups.Count; groupIdx++)
 			{
-				INodeGroup Group = Graph.Groups[GroupIdx];
+				INodeGroup @group = graph.Groups[groupIdx];
 
-				List<IJobStep> Steps = new List<IJobStep>();
-				for (int NodeIdx = 0; NodeIdx < Group.Nodes.Count; NodeIdx++)
+				List<IJobStep> steps = new List<IJobStep>();
+				for (int nodeIdx = 0; nodeIdx < @group.Nodes.Count; nodeIdx++)
 				{
-					SubResourceId StepId = new SubResourceId((ushort)((GroupIdx * 100) + NodeIdx));
+					SubResourceId stepId = new SubResourceId((ushort)((groupIdx * 100) + nodeIdx));
 
-					Mock<IJobStep> Step = new Mock<IJobStep>(MockBehavior.Strict);
-					Step.SetupGet(x => x.Id).Returns(StepId);
-					Step.SetupGet(x => x.NodeIdx).Returns(NodeIdx);
+					Mock<IJobStep> step = new Mock<IJobStep>(MockBehavior.Strict);
+					step.SetupGet(x => x.Id).Returns(stepId);
+					step.SetupGet(x => x.NodeIdx).Returns(nodeIdx);
 
-					Steps.Add(Step.Object);
+					steps.Add(step.Object);
 				}
 
-				SubResourceId BatchId = new SubResourceId((ushort)(GroupIdx * 100));
+				SubResourceId batchId = new SubResourceId((ushort)(groupIdx * 100));
 
-				Mock<IJobStepBatch> Batch = new Mock<IJobStepBatch>(MockBehavior.Strict);
-				Batch.SetupGet(x => x.Id).Returns(BatchId);
-				Batch.SetupGet(x => x.GroupIdx).Returns(GroupIdx);
-				Batch.SetupGet(x => x.Steps).Returns(Steps);
-				Batches.Add(Batch.Object);
+				Mock<IJobStepBatch> batch = new Mock<IJobStepBatch>(MockBehavior.Strict);
+				batch.SetupGet(x => x.Id).Returns(batchId);
+				batch.SetupGet(x => x.GroupIdx).Returns(groupIdx);
+				batch.SetupGet(x => x.Steps).Returns(steps);
+				batches.Add(batch.Object);
 			}
 
-			Mock<IJob> Job = new Mock<IJob>(MockBehavior.Strict);
-			Job.SetupGet(x => x.Id).Returns(JobId);
-			Job.SetupGet(x => x.Name).Returns(Name);
-			Job.SetupGet(x => x.StreamId).Returns(StreamId);
-			Job.SetupGet(x => x.Change).Returns(Change);
-			Job.SetupGet(x => x.Batches).Returns(Batches);
-			return Job.Object;
+			Mock<IJob> job = new Mock<IJob>(MockBehavior.Strict);
+			job.SetupGet(x => x.Id).Returns(jobId);
+			job.SetupGet(x => x.Name).Returns(name);
+			job.SetupGet(x => x.StreamId).Returns(streamId);
+			job.SetupGet(x => x.Change).Returns(change);
+			job.SetupGet(x => x.Batches).Returns(batches);
+			return job.Object;
 		}
 		
 		[TestMethod]
 		public async Task NotifyJobScheduled()
 		{
-			FakeNotificationSink FakeSink = new();
-			NotificationService Service = GetNotificationService(FakeSink);
-			Fixture Fixture = await CreateFixtureAsync();
-			IPool Pool = await PoolService.CreatePoolAsync("BogusPool", Properties: new Dictionary<string, string>());
-
-			Assert.AreEqual(0, FakeSink.JobScheduledNotifications.Count);
-			Service.NotifyJobScheduled(Pool, false, Fixture.Job1, Fixture.Graph, SubResourceId.Random());
-			Service.NotifyJobScheduled(Pool, false, Fixture.Job2, Fixture.Graph, SubResourceId.Random());
+			FakeNotificationSink fakeSink = ServiceProvider.GetRequiredService<FakeNotificationSink>();
+
+			NotificationService service = (NotificationService)ServiceProvider.GetRequiredService<INotificationService>();
+			await service.Ticker.StartAsync();
+
+			Fixture fixture = await CreateFixtureAsync();
+			IPool pool = await PoolService.CreatePoolAsync("BogusPool", Properties: new Dictionary<string, string>());
+
+			Assert.AreEqual(0, fakeSink.JobScheduledNotifications.Count);
+			service.NotifyJobScheduled(pool, false, fixture.Job1, fixture.Graph, SubResourceId.Random());
+			service.NotifyJobScheduled(pool, false, fixture.Job2, fixture.Graph, SubResourceId.Random());
 			
 			// Currently no good way to wait for NotifyJobScheduled() to complete as the execution is completely async in background task (see ExecuteAsync)
 			await Task.Delay(1000);
-			await Clock.AdvanceAsync(Service.NotificationBatchInterval + TimeSpan.FromMinutes(5));
-			Assert.AreEqual(2, FakeSink.JobScheduledNotifications.Count);
-			Assert.AreEqual(1, FakeSink.JobScheduledCallCount);
-		}
-
-		private NotificationService GetNotificationService(INotificationSink? Sink)
-		{
-			ILogger<NotificationService> Logger = ServiceProvider.GetRequiredService<ILogger<NotificationService>>();
-			IDatabase Database = GetRedisConnectionPool().GetDatabase();
-			List<INotificationSink> Sinks = new ();
-			if (Sink != null) Sinks.Add(Sink);
-			
-			NotificationService Service = new NotificationService(
-				Sinks, ServerSettingsMon, Logger, GraphCollection, SubscriptionCollection,
-				NotificationTriggerCollection, UserCollection, JobService, StreamService, IssueService, LogFileService,
-				new NoOpDogStatsd(), Database, Clock);
-
-			Service.Ticker.StartAsync().Wait(5000);
-			return Service;
+			await Clock.AdvanceAsync(service.NotificationBatchInterval + TimeSpan.FromMinutes(5));
+			Assert.AreEqual(2, fakeSink.JobScheduledNotifications.Count);
+			Assert.AreEqual(1, fakeSink.JobScheduledCallCount);
 		}
 
 		//public void NotifyLabelUpdate(IJob Job, IReadOnlyList<(LabelState, LabelOutcome)> OldLabelStates, IReadOnlyList<(LabelState, LabelOutcome)> NewLabelStates)
diff --git a/Engine/Source/Programs/Horde/Horde.Build.Tests/PoolServiceTests.cs b/Engine/Source/Programs/Horde/Horde.Build.Tests/PoolServiceTests.cs
--- a/Engine/Source/Programs/Horde/Horde.Build.Tests/PoolServiceTests.cs
+++ b/Engine/Source/Programs/Horde/Horde.Build.Tests/PoolServiceTests.cs
@@ -3,8 +3,6 @@
 using System;
 using System.Collections.Generic;
 using System.Threading.Tasks;
-using HordeCommon;
-using Horde.Build.Collections.Impl;
 using Horde.Build.Models;
 using Horde.Build.Services;
 using Microsoft.VisualStudio.TestTools.UnitTesting;
@@ -12,10 +10,10 @@ using PoolId = Horde.Build.Utilities.StringId<Horde.Build.Models.IPool>;
 
 namespace Horde.Build.Tests
 {
-    [TestClass]
+	[TestClass]
     public class PoolServiceTests : TestSetup
     {
-        private readonly Dictionary<string, string> FixtureProps = new Dictionary<string, string>
+        private readonly Dictionary<string, string> _fixtureProps = new Dictionary<string, string>
         {
             {"foo", "bar"},
             {"lorem", "ipsum"}
@@ -25,9 +23,9 @@ namespace Horde.Build.Tests
         {
         }
 
-        private async Task<IPool> CreatePoolFixture(String Name)
+        private async Task<IPool> CreatePoolFixture(string name)
         {
-            return await PoolService.CreatePoolAsync(Name, Properties: FixtureProps);
+            return await PoolService.CreatePoolAsync(name, Properties: _fixtureProps);
         }
 
         [TestMethod]
@@ -35,66 +33,66 @@ namespace Horde.Build.Tests
         {
             Assert.IsNull(await PoolService.GetPoolAsync(new PoolId("this-does-not-exist")));
 
-            var NewPool = await CreatePoolFixture("create-pool");
-            var Pool = await PoolService.GetPoolAsync(NewPool.Id);
-			Assert.IsNotNull(Pool);
-            Assert.AreEqual("create-pool", Pool!.Id.ToString());
-            Assert.AreEqual("create-pool", Pool.Name);
-            Assert.AreEqual(FixtureProps.Count, Pool.Properties.Count);
-            Assert.AreEqual(FixtureProps["foo"], Pool.Properties["foo"]);
-            Assert.AreEqual(FixtureProps["lorem"], Pool.Properties["lorem"]);
+            IPool newPool = await CreatePoolFixture("create-pool");
+			IPool? pool = await PoolService.GetPoolAsync(newPool.Id);
+			Assert.IsNotNull(pool);
+            Assert.AreEqual("create-pool", pool!.Id.ToString());
+            Assert.AreEqual("create-pool", pool.Name);
+            Assert.AreEqual(_fixtureProps.Count, pool.Properties.Count);
+            Assert.AreEqual(_fixtureProps["foo"], pool.Properties["foo"]);
+            Assert.AreEqual(_fixtureProps["lorem"], pool.Properties["lorem"]);
         }
         
         [TestMethod]
         public async Task GetPoolsTest()
         {
-            await GetDatabaseService().Database.DropCollectionAsync("Pools");
+            await GetDatabaseServiceSingleton().Database.DropCollectionAsync("Pools");
             
-            var Pools = await PoolService.GetPoolsAsync();
-            Assert.AreEqual(Pools.Count, 0);
-            
-            var Pool0 = await CreatePoolFixture("multiple-pools-0");
-            var Pool1 = await CreatePoolFixture("multiple-pools-1");
-            Pools = await PoolService.GetPoolsAsync();
-            Assert.AreEqual(Pools.Count, 2);
-            Assert.AreEqual(Pools[0].Name, Pool0.Name);
-            Assert.AreEqual(Pools[1].Name, Pool1.Name);
+            List<IPool> pools = await PoolService.GetPoolsAsync();
+            Assert.AreEqual(pools.Count, 0);
+
+			IPool pool0 = await CreatePoolFixture("multiple-pools-0");
+			IPool pool1 = await CreatePoolFixture("multiple-pools-1");
+            pools = await PoolService.GetPoolsAsync();
+            Assert.AreEqual(pools.Count, 2);
+            Assert.AreEqual(pools[0].Name, pool0.Name);
+            Assert.AreEqual(pools[1].Name, pool1.Name);
         }
         
         [TestMethod]
         public async Task DeletePoolTest()
         {
             Assert.IsFalse(await PoolService.DeletePoolAsync(new PoolId("this-does-not-exist")));
-            var Pool = await CreatePoolFixture("pool-to-be-deleted");
-            Assert.IsTrue(await PoolService.DeletePoolAsync(Pool.Id));
+            IPool pool = await CreatePoolFixture("pool-to-be-deleted");
+            Assert.IsTrue(await PoolService.DeletePoolAsync(pool.Id));
         }
         
         [TestMethod]
         public async Task UpdatePoolTest()
         {
-			string UniqueSuffix = Guid.NewGuid().ToString("N");
-            var Pool = await CreatePoolFixture($"update-pool-{UniqueSuffix}");
-            var UpdatedProps = new Dictionary<string, string?>
+			string uniqueSuffix = Guid.NewGuid().ToString("N");
+            IPool pool = await CreatePoolFixture($"update-pool-{uniqueSuffix}");
+            Dictionary<string, string?> updatedProps = new Dictionary<string, string?>
             {
                 {"foo", "bar"},
                 {"lorem", null}, // This entry will get removed
                 {"cookies", "yumyum"},
             };
 
-            var UpdatedPool = await PoolService.UpdatePoolAsync(Pool, $"update-pool-new-name-{UniqueSuffix}", NewProperties: UpdatedProps);
-			Assert.IsNotNull(UpdatedPool);
-            Assert.AreEqual(Pool.Id, UpdatedPool!.Id);
-            Assert.AreEqual($"update-pool-new-name-{UniqueSuffix}", UpdatedPool.Name);
-            Assert.AreEqual(UpdatedPool.Properties.Count, 2);
-            Assert.AreEqual(UpdatedProps["foo"], UpdatedPool.Properties["foo"]);
-            Assert.AreEqual(UpdatedProps["cookies"], UpdatedPool.Properties["cookies"]);
+            IPool? updatedPool = await PoolService.UpdatePoolAsync(pool, $"update-pool-new-name-{uniqueSuffix}", NewProperties: updatedProps);
+			Assert.IsNotNull(updatedPool);
+            Assert.AreEqual(pool.Id, updatedPool!.Id);
+            Assert.AreEqual($"update-pool-new-name-{uniqueSuffix}", updatedPool.Name);
+            Assert.AreEqual(updatedPool.Properties.Count, 2);
+            Assert.AreEqual(updatedProps["foo"], updatedPool.Properties["foo"]);
+            Assert.AreEqual(updatedProps["cookies"], updatedPool.Properties["cookies"]);
         }
         
         [TestMethod]
         public async Task UpdatePoolCollectionTest()
         {
-	        IPool Pool = await CreatePoolFixture("update-pool-2");
-	        await PoolCollection.TryUpdateAsync(Pool, LastScaleUpTime: DateTime.UtcNow);
+	        IPool pool = await CreatePoolFixture("update-pool-2");
+	        await PoolCollection.TryUpdateAsync(pool, LastScaleUpTime: DateTime.UtcNow);
         }
     }
 }
\ No newline at end of file
diff --git a/Engine/Source/Programs/Horde/Horde.Build.Tests/PoolsControllerTest.cs b/Engine/Source/Programs/Horde/Horde.Build.Tests/PoolsControllerTest.cs
--- a/Engine/Source/Programs/Horde/Horde.Build.Tests/PoolsControllerTest.cs
+++ b/Engine/Source/Programs/Horde/Horde.Build.Tests/PoolsControllerTest.cs
@@ -6,31 +6,30 @@ using Horde.Build.Fleet.Autoscale;
 using Horde.Build.Api;
 using Horde.Build.Models;
 using Horde.Build.Utilities;
-using Horde.Build.Tests;
 using Microsoft.AspNetCore.Mvc;
 using Microsoft.VisualStudio.TestTools.UnitTesting;
 
 namespace Horde.Build.Tests
 {
-    [TestClass]
+	[TestClass]
     public class PoolsControllerTest : TestSetup
     {
         [TestMethod]
         public async Task GetPoolsTest()
         {
-	        IPool Pool1 = await PoolService.CreatePoolAsync("Pool1", Properties: new() { {"foo", "bar"}, {"lorem", "ipsum"} });
-	        ActionResult<List<object>> RawResult = await PoolsController.GetPoolsAsync();
-	        Assert.AreEqual(1, RawResult.Value!.Count);
-	        GetPoolResponse Response = (RawResult.Value![0] as GetPoolResponse)!;
-	        Assert.AreEqual(Pool1.Id.ToString(), Response.Id);
-	        Assert.AreEqual(Pool1.Name, Response.Name);
-	        Assert.AreEqual(Pool1.SizeStrategy, Response.SizeStrategy);
+	        IPool pool1 = await PoolService.CreatePoolAsync("Pool1", Properties: new() { {"foo", "bar"}, {"lorem", "ipsum"} });
+	        ActionResult<List<object>> rawResult = await PoolsController.GetPoolsAsync();
+	        Assert.AreEqual(1, rawResult.Value!.Count);
+	        GetPoolResponse response = (rawResult.Value![0] as GetPoolResponse)!;
+	        Assert.AreEqual(pool1.Id.ToString(), response.Id);
+	        Assert.AreEqual(pool1.Name, response.Name);
+	        Assert.AreEqual(pool1.SizeStrategy, response.SizeStrategy);
         }
         
         [TestMethod]
         public async Task CreatePoolsTest()
         {
-	        CreatePoolRequest Request = new CreatePoolRequest
+	        CreatePoolRequest request = new CreatePoolRequest
 	        {
 		        Name = "Pool1",
 		        ScaleOutCooldown = 111,
@@ -39,32 +38,32 @@ namespace Horde.Build.Tests
 		        JobQueueSettings = new Horde.Build.Api.JobQueueSettings(new Horde.Build.Fleet.Autoscale.JobQueueSettings(0.35, 0.85))
 	        };
 	        
-	        ActionResult<CreatePoolResponse> Result = await PoolsController.CreatePoolAsync(Request);
-	        CreatePoolResponse Response = Result.Value!;
+	        ActionResult<CreatePoolResponse> result = await PoolsController.CreatePoolAsync(request);
+	        CreatePoolResponse response = result.Value!;
 		        
-	        IPool Pool = (await PoolService.GetPoolAsync(new StringId<IPool>(Response.Id)))!;
-	        Assert.AreEqual(Request.Name, Pool.Name);
-	        Assert.AreEqual(Request.ScaleOutCooldown, (int)Pool.ScaleOutCooldown!.Value.TotalSeconds);
-	        Assert.AreEqual(Request.ScaleInCooldown, (int)Pool.ScaleInCooldown!.Value.TotalSeconds);
-	        Assert.AreEqual(Request.JobQueueSettings.ScaleOutFactor, Pool.JobQueueSettings!.ScaleOutFactor, 0.0001);
-	        Assert.AreEqual(Request.JobQueueSettings.ScaleInFactor, Pool.JobQueueSettings!.ScaleInFactor, 0.0001);
+	        IPool pool = (await PoolService.GetPoolAsync(new StringId<IPool>(response.Id)))!;
+	        Assert.AreEqual(request.Name, pool.Name);
+	        Assert.AreEqual(request.ScaleOutCooldown, (int)pool.ScaleOutCooldown!.Value.TotalSeconds);
+	        Assert.AreEqual(request.ScaleInCooldown, (int)pool.ScaleInCooldown!.Value.TotalSeconds);
+	        Assert.AreEqual(request.JobQueueSettings.ScaleOutFactor, pool.JobQueueSettings!.ScaleOutFactor, 0.0001);
+	        Assert.AreEqual(request.JobQueueSettings.ScaleInFactor, pool.JobQueueSettings!.ScaleInFactor, 0.0001);
         }
         
         [TestMethod]
         public async Task UpdatePoolTest()
         {
-	        IPool Pool1 = await PoolService.CreatePoolAsync("Pool1", Properties: new() { {"foo", "bar"}, {"lorem", "ipsum"} });
+	        IPool pool1 = await PoolService.CreatePoolAsync("Pool1", Properties: new() { {"foo", "bar"}, {"lorem", "ipsum"} });
 
-	        await PoolsController.UpdatePoolAsync(Pool1.Id.ToString(), new UpdatePoolRequest
+	        await PoolsController.UpdatePoolAsync(pool1.Id.ToString(), new UpdatePoolRequest
 	        {
 		        Name = "Pool1Modified",
 		        SizeStrategy = PoolSizeStrategy.JobQueue
 	        });
 
-	        ActionResult<object> GetResult = await PoolsController.GetPoolAsync(Pool1.Id.ToString());
-	        GetPoolResponse Response = (GetResult.Value! as GetPoolResponse)!;
-	        Assert.AreEqual("Pool1Modified", Response.Name);
-	        Assert.AreEqual(PoolSizeStrategy.JobQueue, Response.SizeStrategy);
+	        ActionResult<object> getResult = await PoolsController.GetPoolAsync(pool1.Id.ToString());
+	        GetPoolResponse response = (getResult.Value! as GetPoolResponse)!;
+	        Assert.AreEqual("Pool1Modified", response.Name);
+	        Assert.AreEqual(PoolSizeStrategy.JobQueue, response.SizeStrategy);
         }
     }
 }
\ No newline at end of file
diff --git a/Engine/Source/Programs/Horde/Horde.Build.Tests/ProgramTests.cs b/Engine/Source/Programs/Horde/Horde.Build.Tests/ProgramTests.cs
--- a/Engine/Source/Programs/Horde/Horde.Build.Tests/ProgramTests.cs
+++ b/Engine/Source/Programs/Horde/Horde.Build.Tests/ProgramTests.cs
@@ -6,7 +6,6 @@ using System.Runtime.InteropServices;
 using System.Security.Cryptography.X509Certificates;
 using EpicGames.Core;
 using Horde.Build.Commands;
-using Horde.Build;
 using Microsoft.VisualStudio.TestTools.UnitTesting;
 
 namespace Horde.Build.Tests
@@ -23,30 +22,30 @@ namespace Horde.Build.Tests
 				return;
 			}
 			
-			string FriendlyName = "A testing cert";
-			byte[] TempCertData = CertificateUtils.CreateSelfSignedCert("testing.epicgames.com", FriendlyName);
-			string TempCertPath = Path.GetTempFileName();
-			File.WriteAllBytes(TempCertPath, TempCertData);
+			string friendlyName = "A testing cert";
+			byte[] tempCertData = CertificateUtils.CreateSelfSignedCert("testing.epicgames.com", friendlyName);
+			string tempCertPath = Path.GetTempFileName();
+			File.WriteAllBytes(tempCertPath, tempCertData);
 			
 			// No cert given
 			Assert.IsNull(ServerCommand.ReadGrpcCertificate(new () { ServerPrivateCert = null }));
 
 			// Cert as file path
 			{
-				X509Certificate2? Cert = ServerCommand.ReadGrpcCertificate(new() { ServerPrivateCert = TempCertPath });
-				Assert.IsNotNull(Cert);
-				Assert.AreEqual(FriendlyName, Cert!.FriendlyName);
+				X509Certificate2? cert = ServerCommand.ReadGrpcCertificate(new() { ServerPrivateCert = tempCertPath });
+				Assert.IsNotNull(cert);
+				Assert.AreEqual(friendlyName, cert!.FriendlyName);
 			}
 
 			// Cert as base64 data
 			{
-				string TempCertBase64 = Convert.ToBase64String(TempCertData);
-				X509Certificate2? Cert = ServerCommand.ReadGrpcCertificate(new() { ServerPrivateCert = "base64:" + TempCertBase64 });
-				Assert.IsNotNull(Cert);
-				Assert.AreEqual(FriendlyName, Cert!.FriendlyName);	
+				string tempCertBase64 = Convert.ToBase64String(tempCertData);
+				X509Certificate2? cert = ServerCommand.ReadGrpcCertificate(new() { ServerPrivateCert = "base64:" + tempCertBase64 });
+				Assert.IsNotNull(cert);
+				Assert.AreEqual(friendlyName, cert!.FriendlyName);	
 			}
 			
-			File.Delete(TempCertPath);
+			File.Delete(tempCertPath);
 		}
 	}
 }
diff --git a/Engine/Source/Programs/Horde/Horde.Build.Tests/RpcServiceTest.cs b/Engine/Source/Programs/Horde/Horde.Build.Tests/RpcServiceTest.cs
--- a/Engine/Source/Programs/Horde/Horde.Build.Tests/RpcServiceTest.cs
+++ b/Engine/Source/Programs/Horde/Horde.Build.Tests/RpcServiceTest.cs
@@ -3,7 +3,6 @@
 using System;
 using System.Collections.Generic;
 using System.IO;
-using System.IO.Compression;
 using System.Reflection;
 using System.Security.Claims;
 using System.Threading;
@@ -11,18 +10,13 @@ using System.Threading.Channels;
 using System.Threading.Tasks;
 using Google.Protobuf;
 using Grpc.Core;
-using Horde.Build;
-using Horde.Build.Collections;
 using HordeCommon.Rpc;
 using Horde.Build.Models;
 using Horde.Build.Services;
 using Horde.Build.Utilities;
-using Horde.Build.Tests.Stubs.Collections;
 using Microsoft.AspNetCore.Http;
 using Microsoft.AspNetCore.Http.Features;
-using Microsoft.Extensions.Configuration;
 using Microsoft.Extensions.Hosting;
-using Microsoft.Extensions.Logging;
 using Microsoft.VisualStudio.TestTools.UnitTesting;
 using MongoDB.Bson;
 using AgentCapabilities = HordeCommon.Rpc.Messages.AgentCapabilities;
@@ -57,7 +51,7 @@ namespace Horde.Build.Tests
 	[TestClass]
 	public class RpcServiceTest : TestSetup
 	{
-		private readonly ServerCallContext AdminContext = new ServerCallContextStub("app-horde-admins");
+		private readonly ServerCallContext _adminContext = new ServerCallContextStub("app-horde-admins");
 
 		sealed class HttpContextStub : HttpContext
 		{
@@ -73,17 +67,17 @@ namespace Horde.Build.Tests
 			public override ClaimsPrincipal User { get; set; }
 			public override WebSocketManager WebSockets { get; } = null!;
 
-			public HttpContextStub(string RoleClaimType)
+			public HttpContextStub(string roleClaimType)
 			{
 				User = new ClaimsPrincipal(new ClaimsIdentity(new List<Claim>
 				{
-					new Claim(HordeClaimTypes.Role, RoleClaimType),
+					new Claim(HordeClaimTypes.Role, roleClaimType),
 				}, "TestAuthType"));
 			}
 			
-			public HttpContextStub(ClaimsPrincipal User)
+			public HttpContextStub(ClaimsPrincipal user)
 			{
-				this.User = User;
+				this.User = user;
 			}
 
 			public override void Abort()
@@ -92,17 +86,17 @@ namespace Horde.Build.Tests
 			}
 		}
 
-		public class ServerCallContextStub : ServerCallContext
+		class ServerCallContextStub : ServerCallContext
 		{
 			// Copied from ServerCallContextExtensions.cs in Grpc.Core
 			const string HttpContextKey = "__HttpContext";
 
-			public static ServerCallContext ForAdminWithAgentSessionId(string AgentSessionId)
+			public static ServerCallContext ForAdminWithAgentSessionId(string agentSessionId)
 			{
 				return new ServerCallContextStub(new ClaimsPrincipal(new ClaimsIdentity(new List<Claim>
 				{
 					new Claim(HordeClaimTypes.Role, "app-horde-admins"),
-					new Claim(HordeClaimTypes.AgentSessionId, AgentSessionId),
+					new Claim(HordeClaimTypes.AgentSessionId, agentSessionId),
 				}, "TestAuthType")));
 			}
 			
@@ -114,26 +108,26 @@ namespace Horde.Build.Tests
 				}, "TestAuthType")));
 			}
 
-			public ServerCallContextStub(string RoleClaimType)
+			public ServerCallContextStub(string roleClaimType)
 			{
 				// The GetHttpContext extension falls back to getting the HttpContext from UserState
 				// We can piggyback on that behavior during tests
-				UserState[HttpContextKey] = new HttpContextStub(RoleClaimType);
+				UserState[HttpContextKey] = new HttpContextStub(roleClaimType);
 			}
 			
-			public ServerCallContextStub(ClaimsPrincipal User)
+			public ServerCallContextStub(ClaimsPrincipal user)
 			{
 				// The GetHttpContext extension falls back to getting the HttpContext from UserState
 				// We can piggyback on that behavior during tests
-				UserState[HttpContextKey] = new HttpContextStub(User);
+				UserState[HttpContextKey] = new HttpContextStub(user);
 			}
 
-			protected override Task WriteResponseHeadersAsyncCore(Metadata ResponseHeaders)
+			protected override Task WriteResponseHeadersAsyncCore(Metadata responseHeaders)
 			{
 				throw new NotImplementedException();
 			}
 
-			protected override ContextPropagationToken CreatePropagationTokenCore(ContextPropagationOptions Options)
+			protected override ContextPropagationToken CreatePropagationTokenCore(ContextPropagationOptions options)
 			{
 				throw new NotImplementedException();
 			}
@@ -150,15 +144,15 @@ namespace Horde.Build.Tests
 			protected override AuthContext AuthContextCore { get; } = null!;
 		}
 		
-		public class RpcServiceInvoker : CallInvoker
+		class RpcServiceInvoker : CallInvoker
 		{
-			private readonly RpcService RpcService;
-			private readonly ServerCallContext ServerCallContext;
+			private readonly RpcService _rpcService;
+			private readonly ServerCallContext _serverCallContext;
 
-			public RpcServiceInvoker(RpcService RpcService, ServerCallContext ServerCallContext)
+			public RpcServiceInvoker(RpcService rpcService, ServerCallContext serverCallContext)
 			{
-				this.RpcService = RpcService;
-				this.ServerCallContext = ServerCallContext;
+				this._rpcService = rpcService;
+				this._serverCallContext = serverCallContext;
 			}
 
 			public override TResponse BlockingUnaryCall<TRequest, TResponse>(Method<TRequest, TResponse> method, string host, CallOptions options, TRequest request)
@@ -166,49 +160,49 @@ namespace Horde.Build.Tests
 				throw new NotImplementedException("Blocking calls are not supported! Method " + method.FullName);
 			}
 
-			public override AsyncUnaryCall<TResponse> AsyncUnaryCall<TRequest, TResponse>(Method<TRequest, TResponse> Method, string Host, CallOptions Options, TRequest Request)
+			public override AsyncUnaryCall<TResponse> AsyncUnaryCall<TRequest, TResponse>(Method<TRequest, TResponse> method, string host, CallOptions options, TRequest request)
 			{
-				MethodInfo MethodInfo = GetMethod(Method.Name);
-				Task<TResponse> Res = (MethodInfo.Invoke(RpcService, new object[] {Request, ServerCallContext}) as Task<TResponse>)!;
-				return new AsyncUnaryCall<TResponse>(Res, null!, null!, null!, null!, null!);
+				MethodInfo methodInfo = GetMethod(method.Name);
+				Task<TResponse> res = (methodInfo.Invoke(_rpcService, new object[] {request, _serverCallContext}) as Task<TResponse>)!;
+				return new AsyncUnaryCall<TResponse>(res, null!, null!, null!, null!, null!);
 			}
 
-			public override AsyncServerStreamingCall<TResponse> AsyncServerStreamingCall<TRequest, TResponse>(Method<TRequest, TResponse> Method, string Host, CallOptions Options,
-				TRequest Request)
+			public override AsyncServerStreamingCall<TResponse> AsyncServerStreamingCall<TRequest, TResponse>(Method<TRequest, TResponse> method, string host, CallOptions options,
+				TRequest request)
 			{
-				Console.WriteLine($"RpcServiceInvoker.AsyncServerStreamingCall(method={Method.FullName} request={Request})");
+				Console.WriteLine($"RpcServiceInvoker.AsyncServerStreamingCall(method={method.FullName} request={request})");
 				throw new NotImplementedException();
 			}
 
-			public override AsyncClientStreamingCall<TRequest, TResponse> AsyncClientStreamingCall<TRequest, TResponse>(Method<TRequest, TResponse> Method, string Host, CallOptions Options)
+			public override AsyncClientStreamingCall<TRequest, TResponse> AsyncClientStreamingCall<TRequest, TResponse>(Method<TRequest, TResponse> method, string host, CallOptions options)
 			{
-				Console.WriteLine($"RpcServiceInvoker.AsyncClientStreamingCall(method={Method.FullName})");
+				Console.WriteLine($"RpcServiceInvoker.AsyncClientStreamingCall(method={method.FullName})");
 				throw new NotImplementedException();
 			}
 
-			public override AsyncDuplexStreamingCall<TRequest, TResponse> AsyncDuplexStreamingCall<TRequest, TResponse>(Method<TRequest, TResponse> Method, string Host, CallOptions Options)
+			public override AsyncDuplexStreamingCall<TRequest, TResponse> AsyncDuplexStreamingCall<TRequest, TResponse>(Method<TRequest, TResponse> method, string host, CallOptions options)
 			{
-				Console.WriteLine($"RpcServiceInvoker.AsyncDuplexStreamingCall(method={Method.FullName})");
+				Console.WriteLine($"RpcServiceInvoker.AsyncDuplexStreamingCall(method={method.FullName})");
 
-				GrpcDuplexStreamHandler<TRequest> RequestStream = new GrpcDuplexStreamHandler<TRequest>(ServerCallContext);
-				GrpcDuplexStreamHandler<TResponse> ResponseStream = new GrpcDuplexStreamHandler<TResponse>(ServerCallContext);
+				GrpcDuplexStreamHandler<TRequest> requestStream = new GrpcDuplexStreamHandler<TRequest>(_serverCallContext);
+				GrpcDuplexStreamHandler<TResponse> responseStream = new GrpcDuplexStreamHandler<TResponse>(_serverCallContext);
 
-				MethodInfo MethodInfo = GetMethod(Method.Name);
-				Task MethodTask = (MethodInfo.Invoke(RpcService, new object[] {RequestStream, ResponseStream, ServerCallContext}) as Task)!;
-				MethodTask.ContinueWith(t => { Console.Error.WriteLine($"Uncaught exception in {Method.Name}: {t}"); }, TaskContinuationOptions.OnlyOnFaulted);
+				MethodInfo methodInfo = GetMethod(method.Name);
+				Task methodTask = (methodInfo.Invoke(_rpcService, new object[] {requestStream, responseStream, _serverCallContext}) as Task)!;
+				methodTask.ContinueWith(t => { Console.Error.WriteLine($"Uncaught exception in {method.Name}: {t}"); }, CancellationToken.None, TaskContinuationOptions.OnlyOnFaulted, TaskScheduler.Default);
 				
-				return new AsyncDuplexStreamingCall<TRequest, TResponse>(RequestStream, ResponseStream, null!, null!, null!, null!);
+				return new AsyncDuplexStreamingCall<TRequest, TResponse>(requestStream, responseStream, null!, null!, null!, null!);
 			}
 			
-			private MethodInfo GetMethod(string MethodName)
+			private MethodInfo GetMethod(string methodName)
 			{
-				MethodInfo? Method = RpcService.GetType().GetMethod(MethodName);
-				if (Method == null)
+				MethodInfo? method = _rpcService.GetType().GetMethod(methodName);
+				if (method == null)
 				{
-					throw new ArgumentException($"Method {MethodName} not found in RpcService");
+					throw new ArgumentException($"Method {methodName} not found in RpcService");
 				}
 
-				return Method;
+				return method;
 			}
 		}
 		
@@ -216,35 +210,35 @@ namespace Horde.Build.Tests
 		/// Combines and cross-writes streams for a duplex streaming call in gRPC
 		/// </summary>
 		/// <typeparam name="T"></typeparam>
-		public class GrpcDuplexStreamHandler<T> : IServerStreamWriter<T>, IClientStreamWriter<T>, IAsyncStreamReader<T> where T : class
+		class GrpcDuplexStreamHandler<T> : IServerStreamWriter<T>, IClientStreamWriter<T>, IAsyncStreamReader<T> where T : class
 		{
-			private readonly ServerCallContext ServerCallContext;
-			private readonly Channel<T> Channel;
+			private readonly ServerCallContext _serverCallContext;
+			private readonly Channel<T> _channel;
 
 			public WriteOptions? WriteOptions { get; set; }
 
-			public GrpcDuplexStreamHandler(ServerCallContext ServerCallContext)
+			public GrpcDuplexStreamHandler(ServerCallContext serverCallContext)
 			{
-				Channel = System.Threading.Channels.Channel.CreateUnbounded<T>();
+				_channel = System.Threading.Channels.Channel.CreateUnbounded<T>();
 
-				this.ServerCallContext = ServerCallContext;
+				_serverCallContext = serverCallContext;
 			}
 
 			public void Complete()
 			{
-				Channel.Writer.Complete();
+				_channel.Writer.Complete();
 			}
 
 			public IAsyncEnumerable<T> ReadAllAsync()
 			{
-				return Channel.Reader.ReadAllAsync();
+				return _channel.Reader.ReadAllAsync();
 			}
 
 			public async Task<T?> ReadNextAsync()
 			{
-				if (await Channel.Reader.WaitToReadAsync())
+				if (await _channel.Reader.WaitToReadAsync())
 				{
-					Channel.Reader.TryRead(out var message);
+					_channel.Reader.TryRead(out T? message);
 					return message;
 				}
 				else
@@ -253,14 +247,14 @@ namespace Horde.Build.Tests
 				}
 			}
 
-			public Task WriteAsync(T Message)
+			public Task WriteAsync(T message)
 			{
-				if (ServerCallContext.CancellationToken.IsCancellationRequested)
+				if (_serverCallContext.CancellationToken.IsCancellationRequested)
 				{
-					return Task.FromCanceled(ServerCallContext.CancellationToken);
+					return Task.FromCanceled(_serverCallContext.CancellationToken);
 				}
 
-				if (!Channel.Writer.TryWrite(Message))
+				if (!_channel.Writer.TryWrite(message))
 				{
 					throw new InvalidOperationException("Unable to write message.");
 				}
@@ -274,15 +268,15 @@ namespace Horde.Build.Tests
 				return Task.CompletedTask;
 			}
 
-			public async Task<bool> MoveNext(CancellationToken CancellationToken)
+			public async Task<bool> MoveNext(CancellationToken cancellationToken)
 			{
-				ServerCallContext.CancellationToken.ThrowIfCancellationRequested();
+				_serverCallContext.CancellationToken.ThrowIfCancellationRequested();
 				
-				if (await Channel.Reader.WaitToReadAsync(CancellationToken))
+				if (await _channel.Reader.WaitToReadAsync(cancellationToken))
 				{
-					if (Channel.Reader.TryRead(out var Message))
+					if (_channel.Reader.TryRead(out T? message))
 					{
-						Current = Message;
+						Current = message;
 						return true;
 					}
 				}
@@ -297,40 +291,40 @@ namespace Horde.Build.Tests
 		[TestMethod]
 		public async Task CreateSessionTest()
 		{
-			CreateSessionRequest Req = new CreateSessionRequest();
-			await Assert.ThrowsExceptionAsync<StructuredRpcException>(() => RpcService.CreateSession(Req, AdminContext));
+			CreateSessionRequest req = new CreateSessionRequest();
+			await Assert.ThrowsExceptionAsync<StructuredRpcException>(() => RpcService.CreateSession(req, _adminContext));
 
-			Req.Name = "MyName";
-			await Assert.ThrowsExceptionAsync<StructuredRpcException>(() => RpcService.CreateSession(Req, AdminContext));
+			req.Name = "MyName";
+			await Assert.ThrowsExceptionAsync<StructuredRpcException>(() => RpcService.CreateSession(req, _adminContext));
 
-			Req.Capabilities = new AgentCapabilities();
-			CreateSessionResponse Res = await RpcService.CreateSession(Req, AdminContext);
+			req.Capabilities = new AgentCapabilities();
+			CreateSessionResponse res = await RpcService.CreateSession(req, _adminContext);
 
-			Assert.AreEqual("MYNAME", Res.AgentId);
+			Assert.AreEqual("MYNAME", res.AgentId);
 			// TODO: Check Token, ExpiryTime, SessionId 
 		}
 
 		[TestMethod]
 		public async Task UpdateSessionTest()
 		{
-			CreateSessionRequest CreateReq = new CreateSessionRequest
+			CreateSessionRequest createReq = new CreateSessionRequest
 			{
 				Name = "UpdateSessionTest1", Capabilities = new AgentCapabilities()
 			};
-			CreateSessionResponse CreateRes = await RpcService.CreateSession(CreateReq, AdminContext);
-			string AgentId = CreateRes.AgentId;
-			string SessionId = CreateRes.SessionId;
-
-			TestAsyncStreamReader<UpdateSessionRequest> RequestStream =
-				new TestAsyncStreamReader<UpdateSessionRequest>(AdminContext);
-			TestServerStreamWriter<UpdateSessionResponse> ResponseStream =
-				new TestServerStreamWriter<UpdateSessionResponse>(AdminContext);
-			Task Call = RpcService.UpdateSession(RequestStream, ResponseStream, AdminContext);
-
-			RequestStream.AddMessage(new UpdateSessionRequest {AgentId = "does-not-exist", SessionId = SessionId});
-			StructuredRpcException Re = await Assert.ThrowsExceptionAsync<StructuredRpcException>(() => Call);
-			Assert.AreEqual(StatusCode.NotFound, Re.StatusCode);
-			Assert.IsTrue(Re.Message.Contains("Invalid agent name"));
+			CreateSessionResponse createRes = await RpcService.CreateSession(createReq, _adminContext);
+			string agentId = createRes.AgentId;
+			string sessionId = createRes.SessionId;
+
+			TestAsyncStreamReader<UpdateSessionRequest> requestStream =
+				new TestAsyncStreamReader<UpdateSessionRequest>(_adminContext);
+			TestServerStreamWriter<UpdateSessionResponse> responseStream =
+				new TestServerStreamWriter<UpdateSessionResponse>(_adminContext);
+			Task call = RpcService.UpdateSession(requestStream, responseStream, _adminContext);
+
+			requestStream.AddMessage(new UpdateSessionRequest {AgentId = "does-not-exist", SessionId = sessionId});
+			StructuredRpcException re = await Assert.ThrowsExceptionAsync<StructuredRpcException>(() => call);
+			Assert.AreEqual(StatusCode.NotFound, re.StatusCode);
+			Assert.IsTrue(re.Message.Contains("Invalid agent name", StringComparison.OrdinalIgnoreCase));
 		}
 		
 		[TestMethod]
@@ -338,75 +332,75 @@ namespace Horde.Build.Tests
 		{
 			RpcService.LongPollTimeout = TimeSpan.FromMilliseconds(200);
 
-			TestAsyncStreamReader<QueryServerStateRequest> RequestStream =
-				new TestAsyncStreamReader<QueryServerStateRequest>(AdminContext);
-			TestServerStreamWriter<QueryServerStateResponse> ResponseStream =
-				new TestServerStreamWriter<QueryServerStateResponse>(AdminContext);
-			Task Call = RpcService.QueryServerState(RequestStream, ResponseStream, AdminContext);
+			TestAsyncStreamReader<QueryServerStateRequest> requestStream =
+				new TestAsyncStreamReader<QueryServerStateRequest>(_adminContext);
+			TestServerStreamWriter<QueryServerStateResponse> responseStream =
+				new TestServerStreamWriter<QueryServerStateResponse>(_adminContext);
+			Task call = RpcService.QueryServerState(requestStream, responseStream, _adminContext);
 
-			RequestStream.AddMessage(new QueryServerStateRequest {Name = "bogusAgentName"});
-			QueryServerStateResponse? Res = await ResponseStream.ReadNextAsync();
-			Assert.IsNotNull(Res);
+			requestStream.AddMessage(new QueryServerStateRequest {Name = "bogusAgentName"});
+			QueryServerStateResponse? res = await responseStream.ReadNextAsync();
+			Assert.IsNotNull(res);
 			
-			Res = await ResponseStream.ReadNextAsync();
-			Assert.IsNotNull(Res);
+			res = await responseStream.ReadNextAsync();
+			Assert.IsNotNull(res);
 
 			// Should timeout after LongPollTimeout specified above
-			await Call;
+			await call;
 		}
 		
 		[TestMethod]
 		public async Task FinishBatchTest()
 		{
-			CreateSessionRequest CreateReq = new CreateSessionRequest
+			CreateSessionRequest createReq = new CreateSessionRequest
 			{
 				Name = "UpdateSessionTest1", Capabilities = new AgentCapabilities()
 			};
-			CreateSessionResponse CreateRes = await RpcService.CreateSession(CreateReq, AdminContext);
-			string AgentId = CreateRes.AgentId;
-			string SessionId = CreateRes.SessionId;
-
-			TestAsyncStreamReader<UpdateSessionRequest> RequestStream =
-				new TestAsyncStreamReader<UpdateSessionRequest>(AdminContext);
-			TestServerStreamWriter<UpdateSessionResponse> ResponseStream =
-				new TestServerStreamWriter<UpdateSessionResponse>(AdminContext);
-			Task Call = RpcService.UpdateSession(RequestStream, ResponseStream, AdminContext);
-
-			RequestStream.AddMessage(new UpdateSessionRequest {AgentId = "does-not-exist", SessionId = SessionId});
-			StructuredRpcException Re = await Assert.ThrowsExceptionAsync<StructuredRpcException>(() => Call);
-			Assert.AreEqual(StatusCode.NotFound, Re.StatusCode);
-			Assert.IsTrue(Re.Message.Contains("Invalid agent name"));
+			CreateSessionResponse createRes = await RpcService.CreateSession(createReq, _adminContext);
+			string agentId = createRes.AgentId;
+			string sessionId = createRes.SessionId;
+
+			TestAsyncStreamReader<UpdateSessionRequest> requestStream =
+				new TestAsyncStreamReader<UpdateSessionRequest>(_adminContext);
+			TestServerStreamWriter<UpdateSessionResponse> responseStream =
+				new TestServerStreamWriter<UpdateSessionResponse>(_adminContext);
+			Task call = RpcService.UpdateSession(requestStream, responseStream, _adminContext);
+
+			requestStream.AddMessage(new UpdateSessionRequest {AgentId = "does-not-exist", SessionId = sessionId});
+			StructuredRpcException re = await Assert.ThrowsExceptionAsync<StructuredRpcException>(() => call);
+			Assert.AreEqual(StatusCode.NotFound, re.StatusCode);
+			Assert.IsTrue(re.Message.Contains("Invalid agent name", StringComparison.OrdinalIgnoreCase));
 		}
 		
 		[TestMethod]
 		public async Task UploadArtifactTest()
 		{
-			Fixture Fixture = await CreateFixtureAsync();
+			Fixture fixture = await CreateFixtureAsync();
 
-			ObjectId SessionId = ObjectId.GenerateNewId();
-			ServerCallContext Context = new ServerCallContextStub(new ClaimsPrincipal(new ClaimsIdentity(new List<Claim>
+			ObjectId sessionId = ObjectId.GenerateNewId();
+			ServerCallContext context = new ServerCallContextStub(new ClaimsPrincipal(new ClaimsIdentity(new List<Claim>
 			{
 				new Claim(HordeClaimTypes.Role, "app-horde-admins"),
-				new Claim(HordeClaimTypes.AgentSessionId, SessionId.ToString()),
+				new Claim(HordeClaimTypes.AgentSessionId, sessionId.ToString()),
 			}, "TestAuthType")));
 
 
-			string[] Data = {"foo", "bar", "baz", "qux"};
-			string DataStr = string.Join("", Data);
+			string[] data = {"foo", "bar", "baz", "qux"};
+			string dataStr = String.Join("", data);
 			
-			UploadArtifactMetadata Metadata = new UploadArtifactMetadata
+			UploadArtifactMetadata metadata = new UploadArtifactMetadata
 			{
-				JobId = Fixture.Job1.Id.ToString(),
-				BatchId = Fixture.Job1.Batches[0].Id.ToString(),
-				StepId = Fixture.Job1.Batches[0].Steps[0].Id.ToString(),
+				JobId = fixture.Job1.Id.ToString(),
+				BatchId = fixture.Job1.Batches[0].Id.ToString(),
+				StepId = fixture.Job1.Batches[0].Steps[0].Id.ToString(),
 				Name = "testfile.txt",
 				MimeType = "text/plain",
-				Length = DataStr.Length
+				Length = dataStr.Length
 			};
 
 			// Set the session ID on the job batch to pass auth later
-			Deref(await JobCollection.TryAssignLeaseAsync(Fixture.Job1, 0, new PoolId("foo"),
-				new AgentId("test"), new ObjectId<Horde.Build.Models.ISession>(SessionId),
+			Deref(await JobCollection.TryAssignLeaseAsync(fixture.Job1, 0, new PoolId("foo"),
+				new AgentId("test"), new ObjectId<Horde.Build.Models.ISession>(sessionId),
 				LeaseId.GenerateNewId(), LogId.GenerateNewId()));
 /*
 			TestAsyncStreamReader<UploadArtifactRequest> RequestStream = new TestAsyncStreamReader<UploadArtifactRequest>(Context);
@@ -420,22 +414,22 @@ namespace Horde.Build.Tests
 			RequestStream.Complete();
 			await Task.Delay(500);
 */			
-			TestAsyncStreamReader<UploadArtifactRequest> RequestStream = new TestAsyncStreamReader<UploadArtifactRequest>(Context);
-			Task<UploadArtifactResponse> Call = RpcService.UploadArtifact(RequestStream,  Context);
-			RequestStream.AddMessage(new UploadArtifactRequest { Metadata = Metadata });
-			RequestStream.AddMessage(new UploadArtifactRequest { Data = ByteString.CopyFromUtf8(Data[0]) });
-			RequestStream.AddMessage(new UploadArtifactRequest { Data = ByteString.CopyFromUtf8(Data[1]) });
-			RequestStream.AddMessage(new UploadArtifactRequest { Data = ByteString.CopyFromUtf8(Data[2]) });
-			RequestStream.AddMessage(new UploadArtifactRequest { Data = ByteString.CopyFromUtf8(Data[3]) });
-			UploadArtifactResponse Res = await Call;
-
-
-			IArtifact? Artifact = await ArtifactCollection.GetArtifactAsync(ObjectId.Parse(Res.Id));
-			Assert.IsNotNull(Artifact);
-			Stream Stream = await ArtifactCollection.OpenArtifactReadStreamAsync(Artifact!);
-			StreamReader Reader = new StreamReader(Stream);
-			string text = Reader.ReadToEnd();
-			Assert.AreEqual(DataStr, text);
+			TestAsyncStreamReader<UploadArtifactRequest> requestStream = new TestAsyncStreamReader<UploadArtifactRequest>(context);
+			Task<UploadArtifactResponse> call = RpcService.UploadArtifact(requestStream,  context);
+			requestStream.AddMessage(new UploadArtifactRequest { Metadata = metadata });
+			requestStream.AddMessage(new UploadArtifactRequest { Data = ByteString.CopyFromUtf8(data[0]) });
+			requestStream.AddMessage(new UploadArtifactRequest { Data = ByteString.CopyFromUtf8(data[1]) });
+			requestStream.AddMessage(new UploadArtifactRequest { Data = ByteString.CopyFromUtf8(data[2]) });
+			requestStream.AddMessage(new UploadArtifactRequest { Data = ByteString.CopyFromUtf8(data[3]) });
+			UploadArtifactResponse res = await call;
+
+
+			IArtifact? artifact = await ArtifactCollection.GetArtifactAsync(ObjectId.Parse(res.Id));
+			Assert.IsNotNull(artifact);
+			Stream stream = await ArtifactCollection.OpenArtifactReadStreamAsync(artifact!);
+			using StreamReader reader = new StreamReader(stream);
+			string text = await reader.ReadToEndAsync();
+			Assert.AreEqual(dataStr, text);
 		}
 		/*
 		[TestMethod]
