Horde: Update Horde.Build.Tests to use standard coding conventions, and fix a number of issues caught by static analysis.

diff --git a/Engine/Source/Programs/Horde/Horde.Build.Tests/SchedulerTests.cs b/Engine/Source/Programs/Horde/Horde.Build.Tests/SchedulerTests.cs
--- a/Engine/Source/Programs/Horde/Horde.Build.Tests/SchedulerTests.cs
+++ b/Engine/Source/Programs/Horde/Horde.Build.Tests/SchedulerTests.cs
@@ -4,22 +4,17 @@ using System;
 using System.Collections.Generic;
 using System.Threading.Tasks;
 using HordeCommon;
-using Horde.Build.Collections.Impl;
 using Horde.Build.Models;
 using Horde.Build.Services;
 using Microsoft.Extensions.DependencyInjection;
 using Microsoft.VisualStudio.TestTools.UnitTesting;
-using PoolId = Horde.Build.Utilities.StringId<Horde.Build.Models.IPool>;
 using StreamId = Horde.Build.Utilities.StringId<Horde.Build.Models.IStream>;
 using ProjectId = Horde.Build.Utilities.StringId<Horde.Build.Models.IProject>;
 using TemplateRefId = Horde.Build.Utilities.StringId<Horde.Build.Models.TemplateRef>;
 using Horde.Build.Api;
-using MongoDB.Bson;
 using System.Linq;
-using Horde.Build.Tests.Stubs.Services;
 using Horde.Build.Collections;
 using Horde.Build.Utilities;
-using System.Threading;
 
 namespace Horde.Build.Tests
 {
@@ -33,62 +28,62 @@ namespace Horde.Build.Tests
 		StreamId StreamId { get; } = new StreamId("ue5-main");
 		TemplateRefId TemplateRefId { get; } = new TemplateRefId("template1");
 
-		ITemplate Template;
-		HashSet<JobId> InitialJobIds;
+		ITemplate _template;
+		HashSet<JobId> _initialJobIds;
 
 		public SchedulerTests()
 		{
-			IUser Bob = UserCollection.FindOrAddUserByLoginAsync("Bob").Result;
+			IUser bob = UserCollection.FindOrAddUserByLoginAsync("Bob").Result;
 
-			IProject ? Project = ProjectService.Collection.AddOrUpdateAsync(ProjectId, "", "", 0, new ProjectConfig { Name = "UE5" }).Result;
-			Assert.IsNotNull(Project);
+			IProject ? project = ProjectService.Collection.AddOrUpdateAsync(ProjectId, "", "", 0, new ProjectConfig { Name = "UE5" }).Result;
+			Assert.IsNotNull(project);
 
-			Template = TemplateCollection.AddAsync("Test template").Result;
+			_template = TemplateCollection.AddAsync("Test template").Result;
 
-			InitialJobIds = new HashSet<JobId>(JobCollection.FindAsync().Result.Select(x => x.Id));
+			_initialJobIds = new HashSet<JobId>(JobCollection.FindAsync().Result.Select(x => x.Id));
 
 			PerforceService.Changes.Clear();
-			PerforceService.AddChange("//UE5/Main", 100, Bob, "", new[] { "code.cpp" });
-			PerforceService.AddChange("//UE5/Main", 101, Bob, "", new[] { "content.uasset" });
-			PerforceService.AddChange("//UE5/Main", 102, Bob, "", new[] { "content.uasset" });
+			PerforceService.AddChange("//UE5/Main", 100, bob, "", new[] { "code.cpp" });
+			PerforceService.AddChange("//UE5/Main", 101, bob, "", new[] { "content.uasset" });
+			PerforceService.AddChange("//UE5/Main", 102, bob, "", new[] { "content.uasset" });
 		}
 
-		async Task<IStream> SetScheduleAsync(CreateScheduleRequest Schedule)
+		async Task<IStream> SetScheduleAsync(CreateScheduleRequest schedule)
 		{
 			await ScheduleService.ResetAsync();
 
-			IStream? Stream = await StreamService.GetStreamAsync(StreamId);
+			IStream? stream = await StreamService.GetStreamAsync(StreamId);
 
-			StreamConfig Config = new StreamConfig();
-			Config.Name = "//UE5/Main";
-			Config.Tabs.Add(new CreateJobsTabRequest { Title = "foo", Templates = new List<string> { TemplateRefId.ToString() } });
-			Config.Templates.Add(new CreateTemplateRefRequest { Id = TemplateRefId.ToString(), Name = "Test", Schedule = Schedule });
+			StreamConfig config = new StreamConfig();
+			config.Name = "//UE5/Main";
+			config.Tabs.Add(new CreateJobsTabRequest { Title = "foo", Templates = new List<string> { TemplateRefId.ToString() } });
+			config.Templates.Add(new CreateTemplateRefRequest { Id = TemplateRefId.ToString(), Name = "Test", Schedule = schedule });
 
-			return (await StreamService.StreamCollection.TryCreateOrReplaceAsync(StreamId, Stream, "", "", ProjectId, Config))!; 
+			return (await StreamService.StreamCollection.TryCreateOrReplaceAsync(StreamId, stream, "", "", ProjectId, config))!; 
 		}
 
-		public async Task<List<IJob>> FileTestHelperAsync(params string[] Files)
+		public async Task<List<IJob>> FileTestHelperAsync(params string[] files)
 		{
-			IUser Bob = UserCollection.FindOrAddUserByLoginAsync("Bob", "Bob").Result;
+			IUser bob = await UserCollection.FindOrAddUserByLoginAsync("Bob", "Bob");
 
 			PerforceService.Changes.Clear();
-			PerforceService.AddChange("//UE5/Main", 100, Bob, "", new[] { "code.cpp" });
-			PerforceService.AddChange("//UE5/Main", 101, Bob, "", new[] { "content.uasset" });
-			PerforceService.AddChange("//UE5/Main", 102, Bob, "", new[] { "content.uasset" });
-			PerforceService.AddChange("//UE5/Main", 103, Bob, "", new[] { "foo/code.cpp" });
-			PerforceService.AddChange("//UE5/Main", 104, Bob, "", new[] { "bar/code.cpp" });
-			PerforceService.AddChange("//UE5/Main", 105, Bob, "", new[] { "foo/bar/content.uasset" });
-			PerforceService.AddChange("//UE5/Main", 106, Bob, "", new[] { "bar/foo/content.uasset" });
-
-			DateTime StartTime = new DateTime(2021, 1, 1, 12, 0, 0, DateTimeKind.Local); // Friday Jan 1, 2021 
-			Clock.UtcNow = StartTime;
-
-			CreateScheduleRequest Schedule = new CreateScheduleRequest();
-			Schedule.Enabled = true;
-			Schedule.MaxChanges = 10;
-			Schedule.Patterns.Add(new CreateSchedulePatternRequest { Interval = 1 });
-			Schedule.Files = Files.ToList();
-			await SetScheduleAsync(Schedule);
+			PerforceService.AddChange("//UE5/Main", 100, bob, "", new[] { "code.cpp" });
+			PerforceService.AddChange("//UE5/Main", 101, bob, "", new[] { "content.uasset" });
+			PerforceService.AddChange("//UE5/Main", 102, bob, "", new[] { "content.uasset" });
+			PerforceService.AddChange("//UE5/Main", 103, bob, "", new[] { "foo/code.cpp" });
+			PerforceService.AddChange("//UE5/Main", 104, bob, "", new[] { "bar/code.cpp" });
+			PerforceService.AddChange("//UE5/Main", 105, bob, "", new[] { "foo/bar/content.uasset" });
+			PerforceService.AddChange("//UE5/Main", 106, bob, "", new[] { "bar/foo/content.uasset" });
+
+			DateTime startTime = new DateTime(2021, 1, 1, 12, 0, 0, DateTimeKind.Local); // Friday Jan 1, 2021 
+			Clock.UtcNow = startTime;
+
+			CreateScheduleRequest schedule = new CreateScheduleRequest();
+			schedule.Enabled = true;
+			schedule.MaxChanges = 10;
+			schedule.Patterns.Add(new CreateSchedulePatternRequest { Interval = 1 });
+			schedule.Files = files.ToList();
+			await SetScheduleAsync(schedule);
 
 			await Clock.AdvanceAsync(TimeSpan.FromHours(1.25));
 			await ScheduleService.TickForTestingAsync();
@@ -99,527 +94,527 @@ namespace Horde.Build.Tests
 		[TestMethod]
 		public async Task FileTest1Async()
 		{
-			List<IJob> Jobs = await FileTestHelperAsync("....cpp");
-			Assert.AreEqual(3, Jobs.Count);
-			Assert.AreEqual(100, Jobs[0].Change);
-			Assert.AreEqual(103, Jobs[1].Change);
-			Assert.AreEqual(104, Jobs[2].Change);
+			List<IJob> jobs = await FileTestHelperAsync("....cpp");
+			Assert.AreEqual(3, jobs.Count);
+			Assert.AreEqual(100, jobs[0].Change);
+			Assert.AreEqual(103, jobs[1].Change);
+			Assert.AreEqual(104, jobs[2].Change);
 		}
 
 		[TestMethod]
 		public async Task FileTest2Async()
 		{
-			List<IJob> Jobs = await FileTestHelperAsync("/foo/...");
-			Assert.AreEqual(2, Jobs.Count);
-			Assert.AreEqual(103, Jobs[0].Change);
-			Assert.AreEqual(105, Jobs[1].Change);
+			List<IJob> jobs = await FileTestHelperAsync("/foo/...");
+			Assert.AreEqual(2, jobs.Count);
+			Assert.AreEqual(103, jobs[0].Change);
+			Assert.AreEqual(105, jobs[1].Change);
 		}
 
 		[TestMethod]
 		public async Task FileTest3Async()
 		{
-			List<IJob> Jobs = await FileTestHelperAsync("....uasset", "-/bar/...");
-			Assert.AreEqual(3, Jobs.Count);
-			Assert.AreEqual(101, Jobs[0].Change);
-			Assert.AreEqual(102, Jobs[1].Change);
-			Assert.AreEqual(105, Jobs[2].Change);
+			List<IJob> jobs = await FileTestHelperAsync("....uasset", "-/bar/...");
+			Assert.AreEqual(3, jobs.Count);
+			Assert.AreEqual(101, jobs[0].Change);
+			Assert.AreEqual(102, jobs[1].Change);
+			Assert.AreEqual(105, jobs[2].Change);
 		}
 
 		[TestMethod]
 		public void DayScheduleTest()
 		{
-			DateTime StartTime = new DateTime(2021, 1, 1, 12, 0, 0, DateTimeKind.Utc); // Friday Jan 1, 2021 
-			Clock.UtcNow = StartTime;
+			DateTime startTime = new DateTime(2021, 1, 1, 12, 0, 0, DateTimeKind.Utc); // Friday Jan 1, 2021 
+			Clock.UtcNow = startTime;
 
-			Schedule Schedule = new Schedule(Clock.UtcNow, RequireSubmittedChange: false);
-			Schedule.Patterns.Add(new SchedulePattern(new List<DayOfWeek> { DayOfWeek.Friday, DayOfWeek.Sunday }, 13 * 60, null, null));
+			Schedule schedule = new Schedule(Clock.UtcNow, RequireSubmittedChange: false);
+			schedule.Patterns.Add(new SchedulePattern(new List<DayOfWeek> { DayOfWeek.Friday, DayOfWeek.Sunday }, 13 * 60, null, null));
 
-			DateTime? NextTime = Schedule.GetNextTriggerTimeUtc(StartTime, TimeZoneInfo.Utc);
-			Assert.AreEqual(StartTime + TimeSpan.FromHours(1.0), NextTime!.Value);
+			DateTime? nextTime = schedule.GetNextTriggerTimeUtc(startTime, TimeZoneInfo.Utc);
+			Assert.AreEqual(startTime + TimeSpan.FromHours(1.0), nextTime!.Value);
 
-			NextTime = Schedule.GetNextTriggerTimeUtc(NextTime.Value, TimeZoneInfo.Utc);
-			Assert.AreEqual(StartTime + TimeSpan.FromHours(1.0 + 24.0 * 2.0), NextTime!.Value);
+			nextTime = schedule.GetNextTriggerTimeUtc(nextTime.Value, TimeZoneInfo.Utc);
+			Assert.AreEqual(startTime + TimeSpan.FromHours(1.0 + 24.0 * 2.0), nextTime!.Value);
 
-			NextTime = Schedule.GetNextTriggerTimeUtc(NextTime.Value, TimeZoneInfo.Utc);
-			Assert.AreEqual(StartTime + TimeSpan.FromHours(1.0 + 24.0 * 7.0), NextTime!.Value);
+			nextTime = schedule.GetNextTriggerTimeUtc(nextTime.Value, TimeZoneInfo.Utc);
+			Assert.AreEqual(startTime + TimeSpan.FromHours(1.0 + 24.0 * 7.0), nextTime!.Value);
 
-			NextTime = Schedule.GetNextTriggerTimeUtc(NextTime.Value, TimeZoneInfo.Utc);
-			Assert.AreEqual(StartTime + TimeSpan.FromHours(1.0 + 24.0 * 9.0), NextTime!.Value);
+			nextTime = schedule.GetNextTriggerTimeUtc(nextTime.Value, TimeZoneInfo.Utc);
+			Assert.AreEqual(startTime + TimeSpan.FromHours(1.0 + 24.0 * 9.0), nextTime!.Value);
 		}
 
 		[TestMethod]
 		public void MulticheduleTest()
 		{
-			DateTime StartTime = new DateTime(2021, 1, 1, 12, 0, 0, DateTimeKind.Utc); // Friday Jan 1, 2021 
-			Clock.UtcNow = StartTime;
+			DateTime startTime = new DateTime(2021, 1, 1, 12, 0, 0, DateTimeKind.Utc); // Friday Jan 1, 2021 
+			Clock.UtcNow = startTime;
 
-			Schedule Schedule = new Schedule(Clock.UtcNow, RequireSubmittedChange: false);
-			Schedule.Patterns.Add(new SchedulePattern(null, 13 * 60, 14 * 60, 15));
+			Schedule schedule = new Schedule(Clock.UtcNow, RequireSubmittedChange: false);
+			schedule.Patterns.Add(new SchedulePattern(null, 13 * 60, 14 * 60, 15));
 
-			DateTime? NextTime = Schedule.GetNextTriggerTimeUtc(StartTime, TimeZoneInfo.Utc);
-			Assert.AreEqual(StartTime + TimeSpan.FromHours(1.0), NextTime!.Value);
+			DateTime? nextTime = schedule.GetNextTriggerTimeUtc(startTime, TimeZoneInfo.Utc);
+			Assert.AreEqual(startTime + TimeSpan.FromHours(1.0), nextTime!.Value);
 
-			NextTime = Schedule.GetNextTriggerTimeUtc(NextTime.Value, TimeZoneInfo.Utc);
-			Assert.AreEqual(StartTime + TimeSpan.FromHours(1.25), NextTime!.Value);
+			nextTime = schedule.GetNextTriggerTimeUtc(nextTime.Value, TimeZoneInfo.Utc);
+			Assert.AreEqual(startTime + TimeSpan.FromHours(1.25), nextTime!.Value);
 
-			NextTime = Schedule.GetNextTriggerTimeUtc(NextTime.Value, TimeZoneInfo.Utc);
-			Assert.AreEqual(StartTime + TimeSpan.FromHours(1.5), NextTime!.Value);
+			nextTime = schedule.GetNextTriggerTimeUtc(nextTime.Value, TimeZoneInfo.Utc);
+			Assert.AreEqual(startTime + TimeSpan.FromHours(1.5), nextTime!.Value);
 
-			NextTime = Schedule.GetNextTriggerTimeUtc(NextTime.Value, TimeZoneInfo.Utc);
-			Assert.AreEqual(StartTime + TimeSpan.FromHours(1.75), NextTime!.Value);
+			nextTime = schedule.GetNextTriggerTimeUtc(nextTime.Value, TimeZoneInfo.Utc);
+			Assert.AreEqual(startTime + TimeSpan.FromHours(1.75), nextTime!.Value);
 
-			NextTime = Schedule.GetNextTriggerTimeUtc(NextTime.Value, TimeZoneInfo.Utc);
-			Assert.AreEqual(StartTime + TimeSpan.FromHours(2.0), NextTime!.Value);
+			nextTime = schedule.GetNextTriggerTimeUtc(nextTime.Value, TimeZoneInfo.Utc);
+			Assert.AreEqual(startTime + TimeSpan.FromHours(2.0), nextTime!.Value);
 
-			NextTime = Schedule.GetNextTriggerTimeUtc(NextTime.Value, TimeZoneInfo.Utc);
-			Assert.AreEqual(StartTime + TimeSpan.FromHours(1.0 + 24.0), NextTime!.Value);
+			nextTime = schedule.GetNextTriggerTimeUtc(nextTime.Value, TimeZoneInfo.Utc);
+			Assert.AreEqual(startTime + TimeSpan.FromHours(1.0 + 24.0), nextTime!.Value);
 		}
 
 		[TestMethod]
 		public void MultiPatternTest()
 		{
-			DateTime StartTime = new DateTime(2021, 1, 1, 0, 0, 0, DateTimeKind.Utc); // Friday Jan 1, 2021 
-			Clock.UtcNow = StartTime;
+			DateTime startTime = new DateTime(2021, 1, 1, 0, 0, 0, DateTimeKind.Utc); // Friday Jan 1, 2021 
+			Clock.UtcNow = startTime;
 
-			Schedule Schedule = new Schedule(Clock.UtcNow, RequireSubmittedChange: false);
-			Schedule.Patterns.Add(new SchedulePattern(null, 11 * 60, 0, 0));
-			Schedule.Patterns.Add(new SchedulePattern(null, 19 * 60, 0, 0));
+			Schedule schedule = new Schedule(Clock.UtcNow, RequireSubmittedChange: false);
+			schedule.Patterns.Add(new SchedulePattern(null, 11 * 60, 0, 0));
+			schedule.Patterns.Add(new SchedulePattern(null, 19 * 60, 0, 0));
 
-			DateTime? NextTime = Schedule.GetNextTriggerTimeUtc(StartTime, TimeZoneInfo.Utc);
-			Assert.AreEqual(StartTime + TimeSpan.FromHours(11), NextTime!.Value);
+			DateTime? nextTime = schedule.GetNextTriggerTimeUtc(startTime, TimeZoneInfo.Utc);
+			Assert.AreEqual(startTime + TimeSpan.FromHours(11), nextTime!.Value);
 
-			NextTime = Schedule.GetNextTriggerTimeUtc(NextTime.Value, TimeZoneInfo.Utc);
-			Assert.AreEqual(StartTime + TimeSpan.FromHours(19), NextTime!.Value);
+			nextTime = schedule.GetNextTriggerTimeUtc(nextTime.Value, TimeZoneInfo.Utc);
+			Assert.AreEqual(startTime + TimeSpan.FromHours(19), nextTime!.Value);
 		}
 
 		[TestMethod]
 		public async Task NoSubmittedChangeTestAsync()
 		{
-			DateTime StartTime = new DateTime(2021, 1, 1, 12, 0, 0, DateTimeKind.Local); // Friday Jan 1, 2021 
-			Clock.UtcNow = StartTime;
+			DateTime startTime = new DateTime(2021, 1, 1, 12, 0, 0, DateTimeKind.Local); // Friday Jan 1, 2021 
+			Clock.UtcNow = startTime;
 
-			CreateScheduleRequest Schedule = new CreateScheduleRequest();
-			Schedule.Enabled = true;
-			Schedule.Patterns.Add(new CreateSchedulePatternRequest { MinTime = 13 * 60, MaxTime = 14 * 60, Interval = 15 });
+			CreateScheduleRequest schedule = new CreateScheduleRequest();
+			schedule.Enabled = true;
+			schedule.Patterns.Add(new CreateSchedulePatternRequest { MinTime = 13 * 60, MaxTime = 14 * 60, Interval = 15 });
 //			Schedule.LastTriggerTime = StartTime;
-			await SetScheduleAsync(Schedule);
+			await SetScheduleAsync(schedule);
 
 			// Initial tick
 			await ScheduleService.TickForTestingAsync();
 
-			List<IJob> Jobs1 = await GetNewJobs();
-			Assert.AreEqual(0, Jobs1.Count);
+			List<IJob> jobs1 = await GetNewJobs();
+			Assert.AreEqual(0, jobs1.Count);
 
 			// Trigger a job
 			await Clock.AdvanceAsync(TimeSpan.FromHours(1.25));
 			await ScheduleService.TickForTestingAsync();
 
-			List<IJob> Jobs2 = await GetNewJobs();
-			Assert.AreEqual(1, Jobs2.Count);
-			Assert.AreEqual(102, Jobs2[0].Change);
-			Assert.AreEqual(100, Jobs2[0].CodeChange);
+			List<IJob> jobs2 = await GetNewJobs();
+			Assert.AreEqual(1, jobs2.Count);
+			Assert.AreEqual(102, jobs2[0].Change);
+			Assert.AreEqual(100, jobs2[0].CodeChange);
 		}
 
 		[TestMethod]
 		public async Task RequireSubmittedChangeTestAsync()
 		{
-			DateTime StartTime = new DateTime(2021, 1, 1, 12, 0, 0, DateTimeKind.Local); // Friday Jan 1, 2021 
-			Clock.UtcNow = StartTime;
+			DateTime startTime = new DateTime(2021, 1, 1, 12, 0, 0, DateTimeKind.Local); // Friday Jan 1, 2021 
+			Clock.UtcNow = startTime;
 
-			CreateScheduleRequest Schedule = new CreateScheduleRequest();
-			Schedule.Enabled = true;
-			Schedule.Patterns.Add(new CreateSchedulePatternRequest { MinTime = 13 * 60, MaxTime = 14 * 60, Interval = 15 });
-			IStream Stream = await SetScheduleAsync(Schedule);
+			CreateScheduleRequest schedule = new CreateScheduleRequest();
+			schedule.Enabled = true;
+			schedule.Patterns.Add(new CreateSchedulePatternRequest { MinTime = 13 * 60, MaxTime = 14 * 60, Interval = 15 });
+			IStream stream = await SetScheduleAsync(schedule);
 
 			// Initial tick
 			await ScheduleService.TickForTestingAsync();
 
-			List<IJob> Jobs1 = await GetNewJobs();
-			Assert.AreEqual(0, Jobs1.Count);
+			List<IJob> jobs1 = await GetNewJobs();
+			Assert.AreEqual(0, jobs1.Count);
 
 			// Trigger a job
 			await Clock.AdvanceAsync(TimeSpan.FromHours(1.25));
 			await ScheduleService.TickForTestingAsync();
 
-			List<IJob> Jobs2 = await GetNewJobs();
-			Assert.AreEqual(1, Jobs2.Count);
-			Assert.AreEqual(102, Jobs2[0].Change);
-			Assert.AreEqual(100, Jobs2[0].CodeChange);
+			List<IJob> jobs2 = await GetNewJobs();
+			Assert.AreEqual(1, jobs2.Count);
+			Assert.AreEqual(102, jobs2[0].Change);
+			Assert.AreEqual(100, jobs2[0].CodeChange);
 
-			IStream Stream2 = (await StreamCollection.GetAsync(Stream.Id))!;
-			Schedule Schedule2 = Stream2.Templates.First().Value.Schedule!;
-			Assert.AreEqual(102, Schedule2.LastTriggerChange);
-			Assert.AreEqual(Clock.UtcNow, Schedule2.LastTriggerTime);
+			IStream stream2 = (await StreamCollection.GetAsync(stream.Id))!;
+			Schedule schedule2 = stream2.Templates.First().Value.Schedule!;
+			Assert.AreEqual(102, schedule2.LastTriggerChange);
+			Assert.AreEqual(Clock.UtcNow, schedule2.LastTriggerTime);
 
 			// Trigger another job
 			await Clock.AdvanceAsync(TimeSpan.FromHours(0.5));
 			await ScheduleService.TickForTestingAsync();
 
-			List<IJob> Jobs3 = await GetNewJobs();
-			Assert.AreEqual(0, Jobs3.Count);
+			List<IJob> jobs3 = await GetNewJobs();
+			Assert.AreEqual(0, jobs3.Count);
 		}
 
 		[TestMethod]
 		public async Task MultipleJobsTestAsync()
 		{
-			DateTime StartTime = new DateTime(2021, 1, 1, 12, 0, 0, DateTimeKind.Local); // Friday Jan 1, 2021 
-			Clock.UtcNow = StartTime;
+			DateTime startTime = new DateTime(2021, 1, 1, 12, 0, 0, DateTimeKind.Local); // Friday Jan 1, 2021 
+			Clock.UtcNow = startTime;
 
-			CreateScheduleRequest Schedule = new CreateScheduleRequest();
-			Schedule.Enabled = true;
-			Schedule.Patterns.Add(new CreateSchedulePatternRequest { MinTime = 13 * 60, MaxTime = 14 * 60, Interval = 15 });
-			Schedule.MaxChanges = 2;
-			Schedule.Filter = new List<ChangeContentFlags> { ChangeContentFlags.ContainsCode };
-			await SetScheduleAsync(Schedule);
+			CreateScheduleRequest schedule = new CreateScheduleRequest();
+			schedule.Enabled = true;
+			schedule.Patterns.Add(new CreateSchedulePatternRequest { MinTime = 13 * 60, MaxTime = 14 * 60, Interval = 15 });
+			schedule.MaxChanges = 2;
+			schedule.Filter = new List<ChangeContentFlags> { ChangeContentFlags.ContainsCode };
+			await SetScheduleAsync(schedule);
 
 			// Initial tick
 			await ScheduleService.TickForTestingAsync();
 
-			List<IJob> Jobs1 = await GetNewJobs();
-			Assert.AreEqual(0, Jobs1.Count);
+			List<IJob> jobs1 = await GetNewJobs();
+			Assert.AreEqual(0, jobs1.Count);
 
 			// Trigger some jobs
-			IUser Bob = UserCollection.FindOrAddUserByLoginAsync("Bob").Result;
-			PerforceService.AddChange("//UE5/Main", 103, Bob, "", new string[] { "foo.cpp" });
-			PerforceService.AddChange("//UE5/Main", 104, Bob, "", new string[] { "foo.cpp" });
-			PerforceService.AddChange("//UE5/Main", 105, Bob, "", new string[] { "foo.uasset" });
-			PerforceService.AddChange("//UE5/Main", 106, Bob, "", new string[] { "foo.cpp" });
+			IUser bob = await UserCollection.FindOrAddUserByLoginAsync("Bob");
+			PerforceService.AddChange("//UE5/Main", 103, bob, "", new string[] { "foo.cpp" });
+			PerforceService.AddChange("//UE5/Main", 104, bob, "", new string[] { "foo.cpp" });
+			PerforceService.AddChange("//UE5/Main", 105, bob, "", new string[] { "foo.uasset" });
+			PerforceService.AddChange("//UE5/Main", 106, bob, "", new string[] { "foo.cpp" });
 
 			await Clock.AdvanceAsync(TimeSpan.FromHours(1.25));
 			await ScheduleService.TickForTestingAsync();
 
-			List<IJob> Jobs2 = await GetNewJobs();
-			Assert.AreEqual(2, Jobs2.Count);
-			Assert.AreEqual(104, Jobs2[0].Change);
-			Assert.AreEqual(104, Jobs2[0].CodeChange);
-			Assert.AreEqual(106, Jobs2[1].Change);
-			Assert.AreEqual(106, Jobs2[1].CodeChange);
+			List<IJob> jobs2 = await GetNewJobs();
+			Assert.AreEqual(2, jobs2.Count);
+			Assert.AreEqual(104, jobs2[0].Change);
+			Assert.AreEqual(104, jobs2[0].CodeChange);
+			Assert.AreEqual(106, jobs2[1].Change);
+			Assert.AreEqual(106, jobs2[1].CodeChange);
 		}
 
 		[TestMethod]
 		public async Task SkipCiTestAsync()
 		{
-			DateTime StartTime = new DateTime(2021, 1, 1, 12, 0, 0, DateTimeKind.Local); // Friday Jan 1, 2021 
-			Clock.UtcNow = StartTime;
+			DateTime startTime = new DateTime(2021, 1, 1, 12, 0, 0, DateTimeKind.Local); // Friday Jan 1, 2021 
+			Clock.UtcNow = startTime;
 
-			CreateScheduleRequest Schedule = new CreateScheduleRequest();
-			Schedule.Enabled = true;
-			Schedule.Patterns.Add(new CreateSchedulePatternRequest { MinTime = 13 * 60, MaxTime = 14 * 60, Interval = 15 });
-			Schedule.MaxChanges = 2;
-			Schedule.Filter = new List<ChangeContentFlags> { ChangeContentFlags.ContainsCode };
-			await SetScheduleAsync(Schedule);
+			CreateScheduleRequest schedule = new CreateScheduleRequest();
+			schedule.Enabled = true;
+			schedule.Patterns.Add(new CreateSchedulePatternRequest { MinTime = 13 * 60, MaxTime = 14 * 60, Interval = 15 });
+			schedule.MaxChanges = 2;
+			schedule.Filter = new List<ChangeContentFlags> { ChangeContentFlags.ContainsCode };
+			await SetScheduleAsync(schedule);
 
 			// Initial tick
 			await ScheduleService.TickForTestingAsync();
 
-			List<IJob> Jobs1 = await GetNewJobs();
-			Assert.AreEqual(0, Jobs1.Count);
+			List<IJob> jobs1 = await GetNewJobs();
+			Assert.AreEqual(0, jobs1.Count);
 
 			// Trigger some jobs
-			IUser Bob = UserCollection.FindOrAddUserByLoginAsync("Bob").Result;
-			PerforceService.AddChange("//UE5/Main", 103, Bob, "", new string[] { "foo.cpp" });
-			PerforceService.AddChange("//UE5/Main", 104, Bob, "Don't build this change!\n#skipci", new string[] { "foo.cpp" });
-			PerforceService.AddChange("//UE5/Main", 105, Bob, "", new string[] { "foo.uasset" });
-			PerforceService.AddChange("//UE5/Main", 106, Bob, "", new string[] { "foo.cpp" });
+			IUser bob = await UserCollection.FindOrAddUserByLoginAsync("Bob");
+			PerforceService.AddChange("//UE5/Main", 103, bob, "", new string[] { "foo.cpp" });
+			PerforceService.AddChange("//UE5/Main", 104, bob, "Don't build this change!\n#skipci", new string[] { "foo.cpp" });
+			PerforceService.AddChange("//UE5/Main", 105, bob, "", new string[] { "foo.uasset" });
+			PerforceService.AddChange("//UE5/Main", 106, bob, "", new string[] { "foo.cpp" });
 
 			await Clock.AdvanceAsync(TimeSpan.FromHours(1.25));
 			await ScheduleService.TickForTestingAsync();
 
-			List<IJob> Jobs2 = await GetNewJobs();
-			Assert.AreEqual(2, Jobs2.Count);
-			Assert.AreEqual(103, Jobs2[0].Change);
-			Assert.AreEqual(103, Jobs2[0].CodeChange);
-			Assert.AreEqual(106, Jobs2[1].Change);
-			Assert.AreEqual(106, Jobs2[1].CodeChange);
+			List<IJob> jobs2 = await GetNewJobs();
+			Assert.AreEqual(2, jobs2.Count);
+			Assert.AreEqual(103, jobs2[0].Change);
+			Assert.AreEqual(103, jobs2[0].CodeChange);
+			Assert.AreEqual(106, jobs2[1].Change);
+			Assert.AreEqual(106, jobs2[1].CodeChange);
 		}
 
 		[TestMethod]
 		public async Task MaxActiveTestAsync()
 		{
-			DateTime StartTime = new DateTime(2021, 1, 1, 12, 0, 0, DateTimeKind.Local); // Friday Jan 1, 2021 
-			Clock.UtcNow = StartTime;
+			DateTime startTime = new DateTime(2021, 1, 1, 12, 0, 0, DateTimeKind.Local); // Friday Jan 1, 2021 
+			Clock.UtcNow = startTime;
 
-			CreateScheduleRequest Schedule = new CreateScheduleRequest();
-			Schedule.Enabled = true;
-			Schedule.RequireSubmittedChange = false;
-			Schedule.Patterns.Add(new CreateSchedulePatternRequest { MinTime = 13 * 60, MaxTime = 14 * 60, Interval = 15 });
-			Schedule.MaxActive = 1;
-			await SetScheduleAsync(Schedule);
+			CreateScheduleRequest schedule = new CreateScheduleRequest();
+			schedule.Enabled = true;
+			schedule.RequireSubmittedChange = false;
+			schedule.Patterns.Add(new CreateSchedulePatternRequest { MinTime = 13 * 60, MaxTime = 14 * 60, Interval = 15 });
+			schedule.MaxActive = 1;
+			await SetScheduleAsync(schedule);
 
 			// Trigger a job
 			await Clock.AdvanceAsync(TimeSpan.FromHours(1.25));
 			await ScheduleService.TickForTestingAsync();
 
-			List<IJob> Jobs2 = await GetNewJobs();
-			Assert.AreEqual(1, Jobs2.Count);
-			Assert.AreEqual(102, Jobs2[0].Change);
-			Assert.AreEqual(100, Jobs2[0].CodeChange);
+			List<IJob> jobs2 = await GetNewJobs();
+			Assert.AreEqual(1, jobs2.Count);
+			Assert.AreEqual(102, jobs2[0].Change);
+			Assert.AreEqual(100, jobs2[0].CodeChange);
 
 			// Test that another job does not trigger
 			await Clock.AdvanceAsync(TimeSpan.FromHours(0.5));
 			await ScheduleService.TickForTestingAsync();
 
-			List<IJob> Jobs3 = await GetNewJobs();
-			Assert.AreEqual(0, Jobs3.Count);
+			List<IJob> jobs3 = await GetNewJobs();
+			Assert.AreEqual(0, jobs3.Count);
 
 			// Mark the original job as complete
-			await JobService.UpdateJobAsync(Jobs2[0], AbortedByUserId: KnownUsers.System);
+			await JobService.UpdateJobAsync(jobs2[0], AbortedByUserId: KnownUsers.System);
 
 			// Test that another job does not trigger
 			await Clock.AdvanceAsync(TimeSpan.FromHours(0.5));
 			await ScheduleService.TickForTestingAsync();
 
-			List<IJob> Jobs4 = await GetNewJobs();
-			Assert.AreEqual(1, Jobs4.Count);
-			Assert.AreEqual(102, Jobs4[0].Change);
-			Assert.AreEqual(100, Jobs4[0].CodeChange);
+			List<IJob> jobs4 = await GetNewJobs();
+			Assert.AreEqual(1, jobs4.Count);
+			Assert.AreEqual(102, jobs4[0].Change);
+			Assert.AreEqual(100, jobs4[0].CodeChange);
 		}
 
 		[TestMethod]
 		public async Task CreateNewChangeTestAsync()
 		{
-			DateTime StartTime = new DateTime(2021, 1, 1, 12, 0, 0, DateTimeKind.Local); // Friday Jan 1, 2021 
-			Clock.UtcNow = StartTime;
+			DateTime startTime = new DateTime(2021, 1, 1, 12, 0, 0, DateTimeKind.Local); // Friday Jan 1, 2021 
+			Clock.UtcNow = startTime;
 
-			CreateScheduleRequest Schedule = new CreateScheduleRequest();
-			Schedule.Enabled = true;
-			Schedule.Patterns.Add(new CreateSchedulePatternRequest { MinTime = 13 * 60, MaxTime = 14 * 60, Interval = 15 });
-			IStream Stream = await SetScheduleAsync(Schedule);
+			CreateScheduleRequest schedule = new CreateScheduleRequest();
+			schedule.Enabled = true;
+			schedule.Patterns.Add(new CreateSchedulePatternRequest { MinTime = 13 * 60, MaxTime = 14 * 60, Interval = 15 });
+			IStream stream = await SetScheduleAsync(schedule);
 
 			// Trigger a job
 			await Clock.AdvanceAsync(TimeSpan.FromHours(1.25));
 			await ScheduleService.TickForTestingAsync();
 
-			List<IJob> Jobs2 = await GetNewJobs();
-			Assert.AreEqual(1, Jobs2.Count);
-			Assert.AreEqual(102, Jobs2[0].Change);
-			Assert.AreEqual(100, Jobs2[0].CodeChange);
+			List<IJob> jobs2 = await GetNewJobs();
+			Assert.AreEqual(1, jobs2.Count);
+			Assert.AreEqual(102, jobs2[0].Change);
+			Assert.AreEqual(100, jobs2[0].CodeChange);
 
 			// Check another job does not trigger due to the change above
 			await Clock.AdvanceAsync(TimeSpan.FromHours(1.25));
 			await ScheduleService.TickForTestingAsync();
 
-			List<IJob> Jobs3 = await GetNewJobs();
-			Assert.AreEqual(0, Jobs3.Count);
+			List<IJob> jobs3 = await GetNewJobs();
+			Assert.AreEqual(0, jobs3.Count);
 		}
 
 		[TestMethod]
 		public async Task GateTestAsync()
 		{
-			DateTime StartTime = new DateTime(2021, 1, 1, 12, 0, 0, DateTimeKind.Local); // Friday Jan 1, 2021 
-			Clock.UtcNow = StartTime;
+			DateTime startTime = new DateTime(2021, 1, 1, 12, 0, 0, DateTimeKind.Local); // Friday Jan 1, 2021 
+			Clock.UtcNow = startTime;
 
 			// Create two templates, the second dependent on the first
-			ITemplate? NewTemplate1 = await TemplateCollection.AddAsync("Test template 1");
-			TemplateRef NewTemplateRef1 = new TemplateRef(NewTemplate1);
-			TemplateRefId NewTemplateRefId1 = new TemplateRefId("new-template-1");
+			ITemplate? newTemplate1 = await TemplateCollection.AddAsync("Test template 1");
+			TemplateRef newTemplateRef1 = new TemplateRef(newTemplate1);
+			TemplateRefId newTemplateRefId1 = new TemplateRefId("new-template-1");
 
-			ITemplate? NewTemplate2 = await TemplateCollection.AddAsync("Test template 2");
-			TemplateRef NewTemplateRef2 = new TemplateRef(NewTemplate2);
-			NewTemplateRef2.Schedule = new Schedule(Clock.UtcNow);
-			NewTemplateRef2.Schedule.Gate = new ScheduleGate(NewTemplateRefId1, "TriggerNext");
-			NewTemplateRef2.Schedule.Patterns.Add(new SchedulePattern(null, 0, null, 10));
-			NewTemplateRef2.Schedule.LastTriggerTime = StartTime;
-			TemplateRefId NewTemplateRefId2 = new TemplateRefId("new-template-2");
+			ITemplate? newTemplate2 = await TemplateCollection.AddAsync("Test template 2");
+			TemplateRef newTemplateRef2 = new TemplateRef(newTemplate2);
+			newTemplateRef2.Schedule = new Schedule(Clock.UtcNow);
+			newTemplateRef2.Schedule.Gate = new ScheduleGate(newTemplateRefId1, "TriggerNext");
+			newTemplateRef2.Schedule.Patterns.Add(new SchedulePattern(null, 0, null, 10));
+			newTemplateRef2.Schedule.LastTriggerTime = startTime;
+			TemplateRefId newTemplateRefId2 = new TemplateRefId("new-template-2");
 
-			IStream? Stream = await StreamService.GetStreamAsync(StreamId);
+			IStream? stream = await StreamService.GetStreamAsync(StreamId);
 
-			StreamConfig Config = new StreamConfig();
-			Config.Name = "//UE5/Main";
-			Config.Tabs.Add(new CreateJobsTabRequest { Title = "foo", Templates = new List<string> { NewTemplateRefId1.ToString(), NewTemplateRefId2.ToString() } });
+			StreamConfig config = new StreamConfig();
+			config.Name = "//UE5/Main";
+			config.Tabs.Add(new CreateJobsTabRequest { Title = "foo", Templates = new List<string> { newTemplateRefId1.ToString(), newTemplateRefId2.ToString() } });
 
-			Stream = (await StreamService.StreamCollection.TryCreateOrReplaceAsync(StreamId, Stream, "", "", ProjectId, Config))!;
+			stream = (await StreamService.StreamCollection.TryCreateOrReplaceAsync(StreamId, stream, "", "", ProjectId, config))!;
 
 			// Create the TriggerNext step and mark it as complete
-			IGraph GraphA = await GraphCollection.AddAsync(NewTemplate1);
-			NewGroup GroupA = new NewGroup("win", new List<NewNode> { new NewNode("TriggerNext") });
-			GraphA = await GraphCollection.AppendAsync(GraphA, new List<NewGroup> { GroupA });
+			IGraph graphA = await GraphCollection.AddAsync(newTemplate1);
+			NewGroup groupA = new NewGroup("win", new List<NewNode> { new NewNode("TriggerNext") });
+			graphA = await GraphCollection.AppendAsync(graphA, new List<NewGroup> { groupA });
 
 			// Tick the schedule and make sure it doesn't trigger
 			await ScheduleService.TickForTestingAsync();
-			List<IJob> Jobs2 = await GetNewJobs();
-			Assert.AreEqual(0, Jobs2.Count);
+			List<IJob> jobs2 = await GetNewJobs();
+			Assert.AreEqual(0, jobs2.Count);
 
 			// Create a job and fail it
-			IJob Job1 = await JobService.CreateJobAsync(null, Stream, NewTemplateRefId1, Template.Id, GraphA, "Hello", 1234, 1233, 999, null, null, null, null, null, null, null, true, true, null, null, new List<string> { "-Target=TriggerNext" });
-			SubResourceId BatchId1 = Job1.Batches[0].Id;
-			SubResourceId StepId1 = Job1.Batches[0].Steps[0].Id;
-			Job1 = Deref(await JobService.UpdateBatchAsync(Job1, BatchId1, LogId.GenerateNewId(), JobStepBatchState.Running));
-			Job1 = Deref(await JobService.UpdateStepAsync(Job1, BatchId1, StepId1, JobStepState.Completed, JobStepOutcome.Failure));
-			Job1 = Deref(await JobService.UpdateBatchAsync(Job1, BatchId1, LogId.GenerateNewId(), JobStepBatchState.Complete));
+			IJob job1 = await JobService.CreateJobAsync(null, stream, newTemplateRefId1, _template.Id, graphA, "Hello", 1234, 1233, 999, null, null, null, null, null, null, null, true, true, null, null, new List<string> { "-Target=TriggerNext" });
+			SubResourceId batchId1 = job1.Batches[0].Id;
+			SubResourceId stepId1 = job1.Batches[0].Steps[0].Id;
+			job1 = Deref(await JobService.UpdateBatchAsync(job1, batchId1, LogId.GenerateNewId(), JobStepBatchState.Running));
+			job1 = Deref(await JobService.UpdateStepAsync(job1, batchId1, stepId1, JobStepState.Completed, JobStepOutcome.Failure));
+			job1 = Deref(await JobService.UpdateBatchAsync(job1, batchId1, LogId.GenerateNewId(), JobStepBatchState.Complete));
 			await GetNewJobs();
 
 			// Tick the schedule and make sure it doesn't trigger
 			await Clock.AdvanceAsync(TimeSpan.FromMinutes(30.0));
 			await ScheduleService.TickForTestingAsync();
-			List<IJob> Jobs3 = await GetNewJobs();
-			Assert.AreEqual(0, Jobs3.Count);
+			List<IJob> jobs3 = await GetNewJobs();
+			Assert.AreEqual(0, jobs3.Count);
 
 			// Create a job and make it succeed
-			IJob Job2 = await JobService.CreateJobAsync(null, Stream, NewTemplateRefId1, Template.Id, GraphA, "Hello", 1234, 1233, 999, null, null, null, null, null, null, null, true, true, null, null, new List<string> { "-Target=TriggerNext" });
-			SubResourceId BatchId2 = Job2.Batches[0].Id;
-			SubResourceId StepId2 = Job2.Batches[0].Steps[0].Id;
-			Job2 = Deref(await JobService.UpdateBatchAsync(Job2, BatchId2, LogId.GenerateNewId(), JobStepBatchState.Running));
-			Job2 = Deref(await JobService.UpdateStepAsync(Job2, BatchId2, StepId2, JobStepState.Completed, JobStepOutcome.Success));
+			IJob job2 = await JobService.CreateJobAsync(null, stream, newTemplateRefId1, _template.Id, graphA, "Hello", 1234, 1233, 999, null, null, null, null, null, null, null, true, true, null, null, new List<string> { "-Target=TriggerNext" });
+			SubResourceId batchId2 = job2.Batches[0].Id;
+			SubResourceId stepId2 = job2.Batches[0].Steps[0].Id;
+			job2 = Deref(await JobService.UpdateBatchAsync(job2, batchId2, LogId.GenerateNewId(), JobStepBatchState.Running));
+			job2 = Deref(await JobService.UpdateStepAsync(job2, batchId2, stepId2, JobStepState.Completed, JobStepOutcome.Success));
 
 			// Tick the schedule and make sure it does trigger
 			await ScheduleService.TickForTestingAsync();
-			List<IJob> Jobs4 = await GetNewJobs();
-			Assert.AreEqual(1, Jobs4.Count);
-			Assert.AreEqual(1234, Jobs4[0].Change);
-			Assert.AreEqual(1233, Jobs4[0].CodeChange);
+			List<IJob> jobs4 = await GetNewJobs();
+			Assert.AreEqual(1, jobs4.Count);
+			Assert.AreEqual(1234, jobs4[0].Change);
+			Assert.AreEqual(1233, jobs4[0].CodeChange);
 		}
 
 		[TestMethod]
 		public async Task GateTest2Async()
 		{
-			IUser Bob = await UserCollection.FindOrAddUserByLoginAsync("Bob");
+			IUser bob = await UserCollection.FindOrAddUserByLoginAsync("Bob");
 
-			DateTime StartTime = new DateTime(2021, 1, 1, 12, 0, 0, DateTimeKind.Local); // Friday Jan 1, 2021 
-			Clock.UtcNow = StartTime;
+			DateTime startTime = new DateTime(2021, 1, 1, 12, 0, 0, DateTimeKind.Local); // Friday Jan 1, 2021 
+			Clock.UtcNow = startTime;
 
 			PerforceService.Changes.Clear();
-			PerforceService.AddChange("//UE5/Main", 1230, Bob, "", new[] { "code.cpp" });
-			PerforceService.AddChange("//UE5/Main", 1231, Bob, "", new[] { "content.uasset" });
-			PerforceService.AddChange("//UE5/Main", 1232, Bob, "", new[] { "content.uasset" });
-			PerforceService.AddChange("//UE5/Main", 1233, Bob, "", new[] { "code.cpp" });
+			PerforceService.AddChange("//UE5/Main", 1230, bob, "", new[] { "code.cpp" });
+			PerforceService.AddChange("//UE5/Main", 1231, bob, "", new[] { "content.uasset" });
+			PerforceService.AddChange("//UE5/Main", 1232, bob, "", new[] { "content.uasset" });
+			PerforceService.AddChange("//UE5/Main", 1233, bob, "", new[] { "code.cpp" });
 
 			// Create two templates, the second dependent on the first
-			TemplateRefId NewTemplateRefId1 = new TemplateRefId("new-template-1");
-			CreateTemplateRefRequest NewTemplate1 = new CreateTemplateRefRequest();
-			NewTemplate1.Id = NewTemplateRefId1.ToString();
-
-			TemplateRefId NewTemplateRefId2 = new TemplateRefId("new-template-2");
-			CreateTemplateRefRequest NewTemplate2 = new CreateTemplateRefRequest();
-			NewTemplate2.Id = NewTemplateRefId2.ToString();
-			NewTemplate2.Name = "Test template 2";
-			NewTemplate2.Schedule = new CreateScheduleRequest();
-			NewTemplate2.Schedule.MaxChanges = 4;
-			NewTemplate2.Schedule.Filter = new List<ChangeContentFlags> { ChangeContentFlags.ContainsCode };
-			NewTemplate2.Schedule.Gate = new CreateScheduleGateRequest { TemplateId = NewTemplateRefId1.ToString(), Target = "TriggerNext" };
-			NewTemplate2.Schedule.Patterns.Add(new CreateSchedulePatternRequest { Interval = 10 });// (null, 0, null, 10));
+			TemplateRefId newTemplateRefId1 = new TemplateRefId("new-template-1");
+			CreateTemplateRefRequest newTemplate1 = new CreateTemplateRefRequest();
+			newTemplate1.Id = newTemplateRefId1.ToString();
+
+			TemplateRefId newTemplateRefId2 = new TemplateRefId("new-template-2");
+			CreateTemplateRefRequest newTemplate2 = new CreateTemplateRefRequest();
+			newTemplate2.Id = newTemplateRefId2.ToString();
+			newTemplate2.Name = "Test template 2";
+			newTemplate2.Schedule = new CreateScheduleRequest();
+			newTemplate2.Schedule.MaxChanges = 4;
+			newTemplate2.Schedule.Filter = new List<ChangeContentFlags> { ChangeContentFlags.ContainsCode };
+			newTemplate2.Schedule.Gate = new CreateScheduleGateRequest { TemplateId = newTemplateRefId1.ToString(), Target = "TriggerNext" };
+			newTemplate2.Schedule.Patterns.Add(new CreateSchedulePatternRequest { Interval = 10 });// (null, 0, null, 10));
 //			NewTemplate2.Schedule.LastTriggerTime = StartTime;
 
-			IStream? Stream = await StreamService.GetStreamAsync(StreamId);
+			IStream? stream = await StreamService.GetStreamAsync(StreamId);
 
-			StreamConfig Config = new StreamConfig();
-			Config.Name = "//UE5/Main";
-			Config.Tabs.Add(new CreateJobsTabRequest { Title = "foo", Templates = new List<string> { NewTemplateRefId1.ToString(), NewTemplateRefId2.ToString() } });
-			Config.Templates.Add(NewTemplate1);
-			Config.Templates.Add(NewTemplate2);
+			StreamConfig config = new StreamConfig();
+			config.Name = "//UE5/Main";
+			config.Tabs.Add(new CreateJobsTabRequest { Title = "foo", Templates = new List<string> { newTemplateRefId1.ToString(), newTemplateRefId2.ToString() } });
+			config.Templates.Add(newTemplate1);
+			config.Templates.Add(newTemplate2);
 
-			Stream = (await StreamService.StreamCollection.TryCreateOrReplaceAsync(StreamId, Stream, "", "", ProjectId, Config))!;
+			stream = (await StreamService.StreamCollection.TryCreateOrReplaceAsync(StreamId, stream, "", "", ProjectId, config))!;
 
-			ITemplate Template1 = (await TemplateCollection.GetAsync(Stream.Templates[NewTemplateRefId1].Hash))!;
-			ITemplate Template2 = (await TemplateCollection.GetAsync(Stream.Templates[NewTemplateRefId2].Hash))!;
+			ITemplate template1 = (await TemplateCollection.GetAsync(stream.Templates[newTemplateRefId1].Hash))!;
+			ITemplate template2 = (await TemplateCollection.GetAsync(stream.Templates[newTemplateRefId2].Hash))!;
 
 			// Create the graph
-			IGraph GraphA = await GraphCollection.AddAsync(Template1);
-			NewGroup GroupA = new NewGroup("win", new List<NewNode> { new NewNode("TriggerNext") });
-			GraphA = await GraphCollection.AppendAsync(GraphA, new List<NewGroup> { GroupA });
+			IGraph graphA = await GraphCollection.AddAsync(template1);
+			NewGroup groupA = new NewGroup("win", new List<NewNode> { new NewNode("TriggerNext") });
+			graphA = await GraphCollection.AppendAsync(graphA, new List<NewGroup> { groupA });
 
 			// Create successful jobs for all the changes we added above
-			for (int Change = 1230; Change <= 1233; Change++)
+			for (int change = 1230; change <= 1233; change++)
 			{
-				int CodeChange = (Change < 1233) ? 1230 : 1233;
+				int codeChange = (change < 1233) ? 1230 : 1233;
 
-				IJob Job1 = await JobService.CreateJobAsync(null, Stream, NewTemplateRefId1, Template.Id, GraphA, "Hello", Change, CodeChange, null, null, null, null, null, false, null, null, true, true, null, null, new List<string> { "-Target=TriggerNext" });
-				for (int BatchIdx = 0; BatchIdx < Job1.Batches.Count; BatchIdx++)
+				IJob job1 = await JobService.CreateJobAsync(null, stream, newTemplateRefId1, _template.Id, graphA, "Hello", change, codeChange, null, null, null, null, null, false, null, null, true, true, null, null, new List<string> { "-Target=TriggerNext" });
+				for (int batchIdx = 0; batchIdx < job1.Batches.Count; batchIdx++)
 				{
-					SubResourceId BatchId1 = Job1.Batches[BatchIdx].Id;
-					Job1 = Deref(await JobService.UpdateBatchAsync(Job1, BatchId1, LogId.GenerateNewId(), JobStepBatchState.Running));
-					for (int StepIdx = 0; StepIdx < Job1.Batches[BatchIdx].Steps.Count; StepIdx++)
+					SubResourceId batchId1 = job1.Batches[batchIdx].Id;
+					job1 = Deref(await JobService.UpdateBatchAsync(job1, batchId1, LogId.GenerateNewId(), JobStepBatchState.Running));
+					for (int stepIdx = 0; stepIdx < job1.Batches[batchIdx].Steps.Count; stepIdx++)
 					{
-						SubResourceId StepId1 = Job1.Batches[BatchIdx].Steps[StepIdx].Id;
-						Job1 = Deref(await JobService.UpdateStepAsync(Job1, BatchId1, StepId1, JobStepState.Completed, JobStepOutcome.Success, NewLogId: LogId.GenerateNewId()));
+						SubResourceId stepId1 = job1.Batches[batchIdx].Steps[stepIdx].Id;
+						job1 = Deref(await JobService.UpdateStepAsync(job1, batchId1, stepId1, JobStepState.Completed, JobStepOutcome.Success, NewLogId: LogId.GenerateNewId()));
 					}
-					Job1 = Deref(await JobService.UpdateBatchAsync(Job1, BatchId1, LogId.GenerateNewId(), JobStepBatchState.Complete));
+					job1 = Deref(await JobService.UpdateBatchAsync(job1, batchId1, LogId.GenerateNewId(), JobStepBatchState.Complete));
 				}
 			}
 			await GetNewJobs();
 
 			// Tick the schedule and make sure it doesn't trigger
 			await Clock.AdvanceAsync(TimeSpan.FromMinutes(30.0));
-			await ScheduleService.TriggerAsync(StreamId, NewTemplateRefId2, Clock.UtcNow, default);
-			List<IJob> Jobs3 = await GetNewJobs();
-			Assert.AreEqual(2, Jobs3.Count);
-			Assert.AreEqual(1230, Jobs3[0].Change);
-			Assert.AreEqual(1233, Jobs3[1].Change);
+			await ScheduleService.TriggerAsync(StreamId, newTemplateRefId2, Clock.UtcNow, default);
+			List<IJob> jobs3 = await GetNewJobs();
+			Assert.AreEqual(2, jobs3.Count);
+			Assert.AreEqual(1230, jobs3[0].Change);
+			Assert.AreEqual(1233, jobs3[1].Change);
 		}
 
 		[TestMethod]
 		public async Task UpdateConfigAsync()
 		{
-			DateTime StartTime = new DateTime(2021, 1, 1, 12, 0, 0, DateTimeKind.Local); // Friday Jan 1, 2021 
-			Clock.UtcNow = StartTime;
+			DateTime startTime = new DateTime(2021, 1, 1, 12, 0, 0, DateTimeKind.Local); // Friday Jan 1, 2021 
+			Clock.UtcNow = startTime;
 
-			CreateScheduleRequest Schedule = new CreateScheduleRequest();
-			Schedule.Enabled = true;
-			Schedule.RequireSubmittedChange = false;
-			Schedule.Patterns.Add(new CreateSchedulePatternRequest { MinTime = 13 * 60, MaxTime = 14 * 60, Interval = 15 });
-			Schedule.MaxActive = 2;
-			await SetScheduleAsync(Schedule);
+			CreateScheduleRequest schedule = new CreateScheduleRequest();
+			schedule.Enabled = true;
+			schedule.RequireSubmittedChange = false;
+			schedule.Patterns.Add(new CreateSchedulePatternRequest { MinTime = 13 * 60, MaxTime = 14 * 60, Interval = 15 });
+			schedule.MaxActive = 2;
+			await SetScheduleAsync(schedule);
 
 			await Clock.AdvanceAsync(TimeSpan.FromHours(1.25));
 			await ScheduleService.TickForTestingAsync();
 
-			List<IJob> Jobs1 = await GetNewJobs();
-			Assert.AreEqual(1, Jobs1.Count);
-			Assert.AreEqual(102, Jobs1[0].Change);
-			Assert.AreEqual(100, Jobs1[0].CodeChange);
+			List<IJob> jobs1 = await GetNewJobs();
+			Assert.AreEqual(1, jobs1.Count);
+			Assert.AreEqual(102, jobs1[0].Change);
+			Assert.AreEqual(100, jobs1[0].CodeChange);
 
 			// Make sure the job is registered
-			IStream? Stream1 = await StreamService.GetStreamAsync(StreamId);
-			TemplateRef TemplateRef1 = Stream1!.Templates.First().Value;
-			Assert.AreEqual(1, TemplateRef1.Schedule!.ActiveJobs.Count);
-			Assert.AreEqual(Jobs1[0].Id, TemplateRef1.Schedule!.ActiveJobs[0]);
+			IStream? stream1 = await StreamService.GetStreamAsync(StreamId);
+			TemplateRef templateRef1 = stream1!.Templates.First().Value;
+			Assert.AreEqual(1, templateRef1.Schedule!.ActiveJobs.Count);
+			Assert.AreEqual(jobs1[0].Id, templateRef1.Schedule!.ActiveJobs[0]);
 
 			// Test that another job does not trigger
-			await SetScheduleAsync(Schedule);
+			await SetScheduleAsync(schedule);
 
 			// Make sure the job is still registered
-			IStream? Stream2 = await StreamService.GetStreamAsync(StreamId);
-			TemplateRef TemplateRef2 = Stream2!.Templates.First().Value;
-			Assert.AreEqual(1, TemplateRef2.Schedule!.ActiveJobs.Count);
-			Assert.AreEqual(Jobs1[0].Id, TemplateRef2.Schedule!.ActiveJobs[0]);
+			IStream? stream2 = await StreamService.GetStreamAsync(StreamId);
+			TemplateRef templateRef2 = stream2!.Templates.First().Value;
+			Assert.AreEqual(1, templateRef2.Schedule!.ActiveJobs.Count);
+			Assert.AreEqual(jobs1[0].Id, templateRef2.Schedule!.ActiveJobs[0]);
 		}
 
 		[TestMethod]
 		public async Task StreamPausing()
 		{
-			DateTime StartTime = new DateTime(2021, 1, 1, 12, 0, 0, DateTimeKind.Utc); // Friday Jan 1, 2021 
-			Clock.UtcNow = StartTime;
+			DateTime startTime = new DateTime(2021, 1, 1, 12, 0, 0, DateTimeKind.Utc); // Friday Jan 1, 2021 
+			Clock.UtcNow = startTime;
 
-			CreateScheduleRequest Schedule = new CreateScheduleRequest();
-			Schedule.Enabled = true;
-			Schedule.Patterns.Add(new CreateSchedulePatternRequest { MinTime = 13 * 60, MaxTime = 14 * 60, Interval = 15 });
-			IStream Stream = await SetScheduleAsync(Schedule);
+			CreateScheduleRequest schedule = new CreateScheduleRequest();
+			schedule.Enabled = true;
+			schedule.Patterns.Add(new CreateSchedulePatternRequest { MinTime = 13 * 60, MaxTime = 14 * 60, Interval = 15 });
+			IStream stream = await SetScheduleAsync(schedule);
 
-			IStreamCollection StreamCollection = ServiceProvider.GetRequiredService<IStreamCollection>();
-			await StreamCollection.TryUpdatePauseStateAsync(Stream, NewPausedUntil: StartTime.AddHours(5), NewPauseComment: "testing");
+			IStreamCollection streamCollection = ServiceProvider.GetRequiredService<IStreamCollection>();
+			await streamCollection.TryUpdatePauseStateAsync(stream, NewPausedUntil: startTime.AddHours(5), NewPauseComment: "testing");
 
 			// Try trigger a job. No job should be scheduled as the stream is paused
 			await Clock.AdvanceAsync(TimeSpan.FromHours(1.25));
 			await ScheduleService.TickForTestingAsync();
-			List<IJob> Jobs2 = await GetNewJobs();
-			Assert.AreEqual(0, Jobs2.Count);
+			List<IJob> jobs2 = await GetNewJobs();
+			Assert.AreEqual(0, jobs2.Count);
 
 			// Advance time beyond the pause period. A build should now trigger
 			await Clock.AdvanceAsync(TimeSpan.FromHours(5.25));
 			await ScheduleService.TickForTestingAsync();
 
-			List<IJob> Jobs3 = await GetNewJobs();
-			Assert.AreEqual(1, Jobs3.Count);
-			Assert.AreEqual(102, Jobs3[0].Change);
-			Assert.AreEqual(100, Jobs3[0].CodeChange);
+			List<IJob> jobs3 = await GetNewJobs();
+			Assert.AreEqual(1, jobs3.Count);
+			Assert.AreEqual(102, jobs3[0].Change);
+			Assert.AreEqual(100, jobs3[0].CodeChange);
 		}
 
 		async Task<List<IJob>> GetNewJobs()
 		{
-			List<IJob> Jobs = await JobCollection.FindAsync();
-			Jobs.RemoveAll(x => InitialJobIds.Contains(x.Id));
-			InitialJobIds.UnionWith(Jobs.Select(x => x.Id));
-			return Jobs.OrderBy(x => x.Change).ToList();
+			List<IJob> jobs = await JobCollection.FindAsync();
+			jobs.RemoveAll(x => _initialJobIds.Contains(x.Id));
+			_initialJobIds.UnionWith(jobs.Select(x => x.Id));
+			return jobs.OrderBy(x => x.Change).ToList();
 		}
 	}
 }
\ No newline at end of file
diff --git a/Engine/Source/Programs/Horde/Horde.Build.Tests/ServiceAccountAuthTest.cs b/Engine/Source/Programs/Horde/Horde.Build.Tests/ServiceAccountAuthTest.cs
--- a/Engine/Source/Programs/Horde/Horde.Build.Tests/ServiceAccountAuthTest.cs
+++ b/Engine/Source/Programs/Horde/Horde.Build.Tests/ServiceAccountAuthTest.cs
@@ -7,10 +7,10 @@ using System.Threading.Tasks;
 using Horde.Build.Authentication;
 using Horde.Build.Collections;
 using Horde.Build.Collections.Impl;
-using Horde.Build.Models;
 using Horde.Build.Services;
 using Microsoft.AspNetCore.Authentication;
 using Microsoft.AspNetCore.Http;
+using Microsoft.Extensions.DependencyInjection;
 using Microsoft.Extensions.Logging;
 using Microsoft.Extensions.Primitives;
 using Microsoft.Extensions.WebEncoders.Testing;
@@ -22,85 +22,92 @@ namespace Horde.Build.Tests
 	[TestClass]
 	public class ServiceAccountAuthTest : DatabaseIntegrationTest
 	{
-		private readonly IServiceAccountCollection ServiceAccounts;
-		private readonly IServiceAccount ServiceAccount;
-
 		public ServiceAccountAuthTest()
 		{
-			DatabaseService DatabaseService = GetDatabaseService();
-			ServiceAccounts = new ServiceAccountCollection(DatabaseService);
-			ServiceAccount = ServiceAccounts.AddAsync("mytoken", new List<string> {"myclaim###myvalue", "foo###bar"}, "mydesc").Result;
  		}
 
-		private async Task<ServiceAccountAuthHandler> GetAuthHandlerAsync(string? HeaderValue)
+		protected override void ConfigureServices(IServiceCollection services)
 		{
-			ServiceAccountAuthOptions Options = new ServiceAccountAuthOptions();
-			ILoggerFactory LoggerFactory = new LoggerFactory();
-			ServiceAccountAuthHandler Handler = new ServiceAccountAuthHandler(new TestOptionsMonitor<ServiceAccountAuthOptions>(Options), LoggerFactory, new UrlTestEncoder(), new SystemClock(), ServiceAccounts);
-			AuthenticationScheme Scheme = new AuthenticationScheme(ServiceAccountAuthHandler.AuthenticationScheme, "ServiceAccountAuth", Handler.GetType());
+			base.ConfigureServices(services);
+
+			services.AddSingleton<ILoggerFactory, LoggerFactory>();
+			services.AddSingleton<ServiceAccountCollection>();
+			services.AddSingleton<IServiceAccountCollection>(SP => SP.GetRequiredService<ServiceAccountCollection>());
+		}
+
+		private async Task<ServiceAccountAuthHandler> GetAuthHandlerAsync(string? headerValue)
+		{
+			ServiceAccountAuthOptions options = new ServiceAccountAuthOptions();
+
+			ServiceAccountCollection serviceAccounts = ServiceProvider.GetRequiredService<ServiceAccountCollection>();
+			await serviceAccounts.AddAsync("mytoken", new List<string> { "myclaim###myvalue", "foo###bar" }, "mydesc");
+
+			ILoggerFactory loggerFactory = ServiceProvider.GetRequiredService<ILoggerFactory>();
+			ServiceAccountAuthHandler handler = new ServiceAccountAuthHandler(new TestOptionsMonitor<ServiceAccountAuthOptions>(options), loggerFactory, new UrlTestEncoder(), new SystemClock(), serviceAccounts);
+			AuthenticationScheme scheme = new AuthenticationScheme(ServiceAccountAuthHandler.AuthenticationScheme, "ServiceAccountAuth", handler.GetType());
 			
-			HttpContext HttpContext = new DefaultHttpContext();
-			if (HeaderValue != null)
+			HttpContext httpContext = new DefaultHttpContext();
+			if (headerValue != null)
 			{
-				HttpContext.Request.Headers.Add(HeaderNames.Authorization, new StringValues(HeaderValue));	
+				httpContext.Request.Headers.Add(HeaderNames.Authorization, new StringValues(headerValue));	
 			}
 			
-			await Handler.InitializeAsync(Scheme, HttpContext);
+			await handler.InitializeAsync(scheme, httpContext);
 
-			return Handler;
+			return handler;
 		}
 
 		[TestMethod]
 		public async Task ValidToken()
 		{
-			ServiceAccountAuthHandler Handler = await GetAuthHandlerAsync("ServiceAccount mytoken");
-			AuthenticateResult Result = await Handler.AuthenticateAsync();
-			Assert.IsTrue(Result.Succeeded);
+			ServiceAccountAuthHandler handler = await GetAuthHandlerAsync("ServiceAccount mytoken");
+			AuthenticateResult result = await handler.AuthenticateAsync();
+			Assert.IsTrue(result.Succeeded);
 			
-			Handler = await GetAuthHandlerAsync("ServiceAccount   mytoken        ");
-			Result = await Handler.AuthenticateAsync();
-			Assert.IsTrue(Result.Succeeded);
+			handler = await GetAuthHandlerAsync("ServiceAccount   mytoken        ");
+			result = await handler.AuthenticateAsync();
+			Assert.IsTrue(result.Succeeded);
 		}
 		
 		[TestMethod]
 		public async Task InvalidToken()
 		{
-			ServiceAccountAuthHandler Handler = await GetAuthHandlerAsync("ServiceAccount doesNotExist");
-			AuthenticateResult Result = await Handler.AuthenticateAsync();
-			Assert.IsFalse(Result.Succeeded);
+			ServiceAccountAuthHandler handler = await GetAuthHandlerAsync("ServiceAccount doesNotExist");
+			AuthenticateResult result = await handler.AuthenticateAsync();
+			Assert.IsFalse(result.Succeeded);
 			
 			// Valid token but bad prefix
-			Handler = await GetAuthHandlerAsync("SriceAcct mytoken");
-			Result = await Handler.AuthenticateAsync();
-			Assert.IsFalse(Result.Succeeded);
+			handler = await GetAuthHandlerAsync("SriceAcct mytoken");
+			result = await handler.AuthenticateAsync();
+			Assert.IsFalse(result.Succeeded);
 		}
 		
 		[TestMethod]
 		public async Task NoResult()
 		{
 			// Valid token but bad prefix
-			ServiceAccountAuthHandler Handler = await GetAuthHandlerAsync("Bogus mytoken");
-			AuthenticateResult Result = await Handler.AuthenticateAsync();
-			Assert.IsFalse(Result.Succeeded);
-			Assert.IsTrue(Result.None);
+			ServiceAccountAuthHandler handler = await GetAuthHandlerAsync("Bogus mytoken");
+			AuthenticateResult result = await handler.AuthenticateAsync();
+			Assert.IsFalse(result.Succeeded);
+			Assert.IsTrue(result.None);
 			
 			// No Authorization header at all
-			Handler = await GetAuthHandlerAsync(null);
-			Result = await Handler.AuthenticateAsync();
-			Assert.IsFalse(Result.Succeeded);
-			Assert.IsTrue(Result.None);
+			handler = await GetAuthHandlerAsync(null);
+			result = await handler.AuthenticateAsync();
+			Assert.IsFalse(result.Succeeded);
+			Assert.IsTrue(result.None);
 		}
 		
 		[TestMethod]
 		public async Task Claims()
 		{
-			ServiceAccountAuthHandler Handler = await GetAuthHandlerAsync("ServiceAccount mytoken");
-			AuthenticateResult Result = await Handler.AuthenticateAsync();
-			Assert.IsTrue(Result.Succeeded);
-			Assert.AreEqual(3, Result.Ticket!.Principal.Claims.Count());
-			Assert.AreEqual(ServiceAccountAuthHandler.AuthenticationScheme, Result.Ticket.Principal.FindFirst(ClaimTypes.Name)!.Value);
-			Assert.AreEqual("myvalue", Result.Ticket!.Principal.FindFirst("myclaim")!.Value);
-			Assert.AreEqual("bar", Result.Ticket!.Principal.FindFirst("foo")!.Value);
+			ServiceAccountAuthHandler handler = await GetAuthHandlerAsync("ServiceAccount mytoken");
+			AuthenticateResult result = await handler.AuthenticateAsync();
+			Assert.IsTrue(result.Succeeded);
+			Assert.AreEqual(3, result.Ticket!.Principal.Claims.Count());
+			Assert.AreEqual(ServiceAccountAuthHandler.AuthenticationScheme, result.Ticket.Principal.FindFirst(ClaimTypes.Name)!.Value);
+			Assert.AreEqual("myvalue", result.Ticket!.Principal.FindFirst("myclaim")!.Value);
+			Assert.AreEqual("bar", result.Ticket!.Principal.FindFirst("foo")!.Value);
 		}
 	}
 }
\ No newline at end of file
diff --git a/Engine/Source/Programs/Horde/Horde.Build.Tests/ServiceAccountCollectionTests.cs b/Engine/Source/Programs/Horde/Horde.Build.Tests/ServiceAccountCollectionTests.cs
--- a/Engine/Source/Programs/Horde/Horde.Build.Tests/ServiceAccountCollectionTests.cs
+++ b/Engine/Source/Programs/Horde/Horde.Build.Tests/ServiceAccountCollectionTests.cs
@@ -14,62 +14,62 @@ namespace Horde.Build.Tests
 	[TestClass]
 	public class ServiceAccountCollectionTests : DatabaseIntegrationTest
 	{
-		private readonly IServiceAccountCollection ServiceAccounts;
-		private readonly IServiceAccount ServiceAccount;
+		private readonly IServiceAccountCollection _serviceAccounts;
+		private readonly IServiceAccount _serviceAccount;
 		
 		public ServiceAccountCollectionTests()
 		{
-			DatabaseService DatabaseService = GetDatabaseService();
-			ServiceAccounts = new ServiceAccountCollection(DatabaseService);
-			ServiceAccount = ServiceAccounts.AddAsync(ObjectId.GenerateNewId().ToString(), new List<string> {"myclaim###myvalue"}, "mydesc").Result;
+			DatabaseService databaseService = GetDatabaseServiceSingleton();
+			_serviceAccounts = new ServiceAccountCollection(databaseService);
+			_serviceAccount = _serviceAccounts.AddAsync(ObjectId.GenerateNewId().ToString(), new List<string> {"myclaim###myvalue"}, "mydesc").Result;
 		}
 
 		[TestMethod]
 		public async Task Add()
 		{
-			IServiceAccount Sa = await ServiceAccounts.AddAsync("addToken", new List<string> {"myclaim###myvalue"}, "mydesc");
-			Assert.AreEqual("addToken", Sa.SecretToken);
-			Assert.AreEqual(1, Sa.GetClaims().Count);
-			Assert.AreEqual("myvalue", Sa.GetClaims()[0].Value);
-			Assert.IsTrue(Sa.Enabled);
-			Assert.AreEqual("mydesc", Sa.Description);
+			IServiceAccount sa = await _serviceAccounts.AddAsync("addToken", new List<string> {"myclaim###myvalue"}, "mydesc");
+			Assert.AreEqual("addToken", sa.SecretToken);
+			Assert.AreEqual(1, sa.GetClaims().Count);
+			Assert.AreEqual("myvalue", sa.GetClaims()[0].Value);
+			Assert.IsTrue(sa.Enabled);
+			Assert.AreEqual("mydesc", sa.Description);
 		}
 		
 		[TestMethod]
 		public async Task Get()
 		{
-			IServiceAccount Sa = (await ServiceAccounts.GetAsync(ServiceAccount.Id))!;
-			Assert.AreEqual(ServiceAccount, Sa);
+			IServiceAccount sa = (await _serviceAccounts.GetAsync(_serviceAccount.Id))!;
+			Assert.AreEqual(_serviceAccount, sa);
 		}
 		
 		[TestMethod]
 		public async Task GetBySecretToken()
 		{
-			IServiceAccount Sa = (await ServiceAccounts.GetBySecretTokenAsync(ServiceAccount.SecretToken))!;
-			Assert.AreEqual(ServiceAccount, Sa);
+			IServiceAccount sa = (await _serviceAccounts.GetBySecretTokenAsync(_serviceAccount.SecretToken))!;
+			Assert.AreEqual(_serviceAccount, sa);
 		}
 		
 		[TestMethod]
 		public async Task Update()
 		{
-			List<string> NewClaims = new List<string> {{"newclaim1###newvalue1"}, {"newclaim2###newvalue2"}};
-			await ServiceAccounts.UpdateAsync(ServiceAccount.Id, "newtoken", NewClaims, false, "newdesc");
-			IServiceAccount GetSa = (await ServiceAccounts.GetAsync(ServiceAccount.Id))!;
+			List<string> newClaims = new List<string> {{"newclaim1###newvalue1"}, {"newclaim2###newvalue2"}};
+			await _serviceAccounts.UpdateAsync(_serviceAccount.Id, "newtoken", newClaims, false, "newdesc");
+			IServiceAccount getSa = (await _serviceAccounts.GetAsync(_serviceAccount.Id))!;
 			
-			Assert.AreEqual("newtoken", GetSa.SecretToken);
-			Assert.AreEqual(2, GetSa.GetClaims().Count);
-			Assert.AreEqual("newvalue1", GetSa.GetClaims()[0].Value);
-			Assert.AreEqual("newvalue2", GetSa.GetClaims()[1].Value);
-			Assert.AreEqual(false, GetSa.Enabled);
-			Assert.AreEqual("newdesc", GetSa.Description);
+			Assert.AreEqual("newtoken", getSa.SecretToken);
+			Assert.AreEqual(2, getSa.GetClaims().Count);
+			Assert.AreEqual("newvalue1", getSa.GetClaims()[0].Value);
+			Assert.AreEqual("newvalue2", getSa.GetClaims()[1].Value);
+			Assert.AreEqual(false, getSa.Enabled);
+			Assert.AreEqual("newdesc", getSa.Description);
 		}
 		
 		[TestMethod]
 		public async Task Delete()
 		{
-			await ServiceAccounts.DeleteAsync(ServiceAccount.Id);
-			IServiceAccount? Result = await ServiceAccounts.GetAsync(ServiceAccount.Id);
-			Assert.IsNull(Result);
+			await _serviceAccounts.DeleteAsync(_serviceAccount.Id);
+			IServiceAccount? result = await _serviceAccounts.GetAsync(_serviceAccount.Id);
+			Assert.IsNull(result);
 		}
 	}
 }
\ No newline at end of file
diff --git a/Engine/Source/Programs/Horde/Horde.Build.Tests/StorageTests.cs b/Engine/Source/Programs/Horde/Horde.Build.Tests/StorageTests.cs
--- a/Engine/Source/Programs/Horde/Horde.Build.Tests/StorageTests.cs
+++ b/Engine/Source/Programs/Horde/Horde.Build.Tests/StorageTests.cs
@@ -3,13 +3,10 @@
 using EpicGames.Core;
 using EpicGames.Horde.Storage;
 using EpicGames.Serialization;
-using Horde.Build.Services;
-using Horde.Build.Utilities;
 using Microsoft.Extensions.DependencyInjection;
 using Microsoft.VisualStudio.TestTools.UnitTesting;
 using System;
 using System.Collections.Generic;
-using System.IO;
 using System.Text;
 using System.Threading.Tasks;
 
@@ -21,118 +18,118 @@ namespace Horde.Build.Tests
 		[TestMethod]
 		public async Task BlobCollectionTest()
 		{
-			NamespaceId NamespaceId = new NamespaceId("test-ns");
+			NamespaceId namespaceId = new NamespaceId("test-ns");
 
-			IStorageClient StorageClient = ServiceProvider.GetRequiredService<IStorageClient>();
+			IStorageClient storageClient = ServiceProvider.GetRequiredService<IStorageClient>();
 
-			byte[] TestData = Encoding.UTF8.GetBytes("Hello world");
-			IoHash Hash = IoHash.Compute(TestData);
+			byte[] testData = Encoding.UTF8.GetBytes("Hello world");
+			IoHash hash = IoHash.Compute(testData);
 
-			await Assert.ThrowsExceptionAsync<BlobNotFoundException>(() => StorageClient.ReadBlobAsync(NamespaceId, Hash));
+			await Assert.ThrowsExceptionAsync<BlobNotFoundException>(() => storageClient.ReadBlobAsync(namespaceId, hash));
 
-			IoHash ReturnedHash = await StorageClient.WriteBlobFromMemoryAsync(NamespaceId, TestData);
-			Assert.AreEqual(Hash, ReturnedHash);
-			Assert.IsNotNull(await StorageClient.ReadBlobToMemoryAsync(NamespaceId, Hash));
+			IoHash returnedHash = await storageClient.WriteBlobFromMemoryAsync(namespaceId, testData);
+			Assert.AreEqual(hash, returnedHash);
+			Assert.IsNotNull(await storageClient.ReadBlobToMemoryAsync(namespaceId, hash));
 
-			ReadOnlyMemory<byte> StoredData = await StorageClient.ReadBlobToMemoryAsync(NamespaceId, Hash);
-			Assert.IsTrue(TestData.AsSpan().SequenceEqual(StoredData.Span));
+			ReadOnlyMemory<byte> storedData = await storageClient.ReadBlobToMemoryAsync(namespaceId, hash);
+			Assert.IsTrue(testData.AsSpan().SequenceEqual(storedData.Span));
 
-			NamespaceId OtherNamespaceId = new NamespaceId("other-ns");
-			await Assert.ThrowsExceptionAsync<BlobNotFoundException>(() => StorageClient.ReadBlobAsync(OtherNamespaceId, Hash));
+			NamespaceId otherNamespaceId = new NamespaceId("other-ns");
+			await Assert.ThrowsExceptionAsync<BlobNotFoundException>(() => storageClient.ReadBlobAsync(otherNamespaceId, hash));
 		}
 
 		[TestMethod]
 		public async Task ObjectCollectionTest()
 		{
-			NamespaceId NamespaceId = new NamespaceId("test-ns");
+			NamespaceId namespaceId = new NamespaceId("test-ns");
 
-			IStorageClient StorageClient = ServiceProvider.GetRequiredService<IStorageClient>();
+			IStorageClient storageClient = ServiceProvider.GetRequiredService<IStorageClient>();
 
-			CbObject ObjectD = CbObject.Build(Writer =>
+			CbObject objectD = CbObject.Build(writer =>
 			{
-				Writer.WriteString("hello", "field");
+				writer.WriteString("hello", "field");
 			});
-			IoHash HashD = ObjectD.GetHash();
+			IoHash hashD = objectD.GetHash();
 
-			await Assert.ThrowsExceptionAsync<BlobNotFoundException>(() => StorageClient.ReadBlobToMemoryAsync(NamespaceId, HashD));
+			await Assert.ThrowsExceptionAsync<BlobNotFoundException>(() => storageClient.ReadBlobToMemoryAsync(namespaceId, hashD));
 
-			await StorageClient.WriteBlobFromMemoryAsync(NamespaceId, HashD, ObjectD.GetView());
+			await storageClient.WriteBlobFromMemoryAsync(namespaceId, hashD, objectD.GetView());
 
-			CbObject ReturnedObjectD = new CbObject(await StorageClient.ReadBlobToMemoryAsync(NamespaceId, HashD));
-			Assert.IsTrue(ReturnedObjectD!.GetView().Span.SequenceEqual(ObjectD.GetView().Span));
+			CbObject returnedObjectD = new CbObject(await storageClient.ReadBlobToMemoryAsync(namespaceId, hashD));
+			Assert.IsTrue(returnedObjectD!.GetView().Span.SequenceEqual(objectD.GetView().Span));
 		}
 
 		[TestMethod]
 		public async Task RefCollectionTest()
 		{
-			IStorageClient StorageClient = ServiceProvider.GetRequiredService<IStorageClient>();
+			IStorageClient storageClient = ServiceProvider.GetRequiredService<IStorageClient>();
 
-			byte[] BlobA = Encoding.UTF8.GetBytes("This is blob A");
-			IoHash HashA = IoHash.Compute(BlobA);
+			byte[] blobA = Encoding.UTF8.GetBytes("This is blob A");
+			IoHash hashA = IoHash.Compute(blobA);
 
-			byte[] BlobB = Encoding.UTF8.GetBytes("This is blob B");
-			IoHash HashB = IoHash.Compute(BlobB);
+			byte[] blobB = Encoding.UTF8.GetBytes("This is blob B");
+			IoHash hashB = IoHash.Compute(blobB);
 
-			byte[] BlobC = Encoding.UTF8.GetBytes("This is blob C");
-			IoHash HashC = IoHash.Compute(BlobC);
+			byte[] blobC = Encoding.UTF8.GetBytes("This is blob C");
+			IoHash hashC = IoHash.Compute(blobC);
 
-			CbObject ObjectD = CbObject.Build(Writer =>
+			CbObject objectD = CbObject.Build(writer =>
 			{
-				Writer.WriteBinaryAttachment("A", HashA);
-				Writer.WriteBinaryAttachment("B", HashB);
-				Writer.WriteBinaryAttachment("C", HashC);
+				writer.WriteBinaryAttachment("A", hashA);
+				writer.WriteBinaryAttachment("B", hashB);
+				writer.WriteBinaryAttachment("C", hashC);
 			});
-			IoHash HashD = ObjectD.GetHash();
+			IoHash hashD = objectD.GetHash();
 
-			CbObject ObjectE = CbObject.Build(Writer =>
+			CbObject objectE = CbObject.Build(writer =>
 			{
-				Writer.WriteBinaryAttachment("A", HashA);
-				Writer.WriteObjectAttachment("D", HashD);
+				writer.WriteBinaryAttachment("A", hashA);
+				writer.WriteObjectAttachment("D", hashD);
 			});
-			IoHash HashE = ObjectE.GetHash();
+			IoHash hashE = objectE.GetHash();
 
-			NamespaceId NamespaceId = new NamespaceId("test-ns");
-			BucketId BucketId = new BucketId("test-bkt");
-			RefId RefId = new RefId("refname");
+			NamespaceId namespaceId = new NamespaceId("test-ns");
+			BucketId bucketId = new BucketId("test-bkt");
+			RefId refId = new RefId("refname");
 
 			// Check that setting the ref to E fails with the correct missing hashes
-			List<IoHash> MissingHashes = await StorageClient.TrySetRefAsync(NamespaceId, BucketId, RefId, ObjectE);
-			Assert.AreEqual(2, MissingHashes.Count);
-			Assert.IsTrue(MissingHashes.Contains(HashA));
-			Assert.IsTrue(MissingHashes.Contains(HashD));
+			List<IoHash> missingHashes = await storageClient.TrySetRefAsync(namespaceId, bucketId, refId, objectE);
+			Assert.AreEqual(2, missingHashes.Count);
+			Assert.IsTrue(missingHashes.Contains(hashA));
+			Assert.IsTrue(missingHashes.Contains(hashD));
 
-			await Assert.ThrowsExceptionAsync<RefNotFoundException>(() => StorageClient.GetRefAsync(NamespaceId, BucketId, RefId));
+			await Assert.ThrowsExceptionAsync<RefNotFoundException>(() => storageClient.GetRefAsync(namespaceId, bucketId, refId));
 
 			// Add object D, and check that changes the missing hashes to just the blobs
-			await StorageClient.WriteBlobFromMemoryAsync(NamespaceId, HashD, ObjectD.GetView());
+			await storageClient.WriteBlobFromMemoryAsync(namespaceId, hashD, objectD.GetView());
 
-			MissingHashes = await StorageClient.TryFinalizeRefAsync(NamespaceId, BucketId, RefId, HashE);
-			Assert.AreEqual(3, MissingHashes.Count);
-			Assert.IsTrue(MissingHashes.Contains(HashA));
-			Assert.IsTrue(MissingHashes.Contains(HashB));
-			Assert.IsTrue(MissingHashes.Contains(HashC));
+			missingHashes = await storageClient.TryFinalizeRefAsync(namespaceId, bucketId, refId, hashE);
+			Assert.AreEqual(3, missingHashes.Count);
+			Assert.IsTrue(missingHashes.Contains(hashA));
+			Assert.IsTrue(missingHashes.Contains(hashB));
+			Assert.IsTrue(missingHashes.Contains(hashC));
 
-			await Assert.ThrowsExceptionAsync<RefNotFoundException>(() => StorageClient.GetRefAsync(NamespaceId, BucketId, RefId));
+			await Assert.ThrowsExceptionAsync<RefNotFoundException>(() => storageClient.GetRefAsync(namespaceId, bucketId, refId));
 
 			// Add blobs A, B and C and check that the object can be finalized
-			await StorageClient.WriteBlobFromMemoryAsync(NamespaceId, HashA, BlobA);
-			await StorageClient.WriteBlobFromMemoryAsync(NamespaceId, HashB, BlobB);
-			await StorageClient.WriteBlobFromMemoryAsync(NamespaceId, HashC, BlobC);
+			await storageClient.WriteBlobFromMemoryAsync(namespaceId, hashA, blobA);
+			await storageClient.WriteBlobFromMemoryAsync(namespaceId, hashB, blobB);
+			await storageClient.WriteBlobFromMemoryAsync(namespaceId, hashC, blobC);
 
-			MissingHashes = await StorageClient.TryFinalizeRefAsync(NamespaceId, BucketId, RefId, HashE);
-			Assert.AreEqual(0, MissingHashes.Count);
+			missingHashes = await storageClient.TryFinalizeRefAsync(namespaceId, bucketId, refId, hashE);
+			Assert.AreEqual(0, missingHashes.Count);
 
-			IRef Ref = await StorageClient.GetRefAsync(NamespaceId, BucketId, RefId);
-			Assert.IsNotNull(Ref);
+			IRef refValue = await storageClient.GetRefAsync(namespaceId, bucketId, refId);
+			Assert.IsNotNull(refValue);
 
 			// Add a new ref to existing objects and check it finalizes correctly 
-			RefId RefId2 = new RefId("refname");
+			RefId refId2 = new RefId("refname");
 
-			MissingHashes = await StorageClient.TrySetRefAsync(NamespaceId, BucketId, RefId2, ObjectE);
-			Assert.AreEqual(0, MissingHashes.Count);
+			missingHashes = await storageClient.TrySetRefAsync(namespaceId, bucketId, refId2, objectE);
+			Assert.AreEqual(0, missingHashes.Count);
 
-			Ref = await StorageClient.GetRefAsync(NamespaceId, BucketId, RefId2);
-			Assert.IsNotNull(Ref);
+			refValue = await storageClient.GetRefAsync(namespaceId, bucketId, refId2);
+			Assert.IsNotNull(refValue);
 		}
 	}
 }
diff --git a/Engine/Source/Programs/Horde/Horde.Build.Tests/StreamServiceTest.cs b/Engine/Source/Programs/Horde/Horde.Build.Tests/StreamServiceTest.cs
--- a/Engine/Source/Programs/Horde/Horde.Build.Tests/StreamServiceTest.cs
+++ b/Engine/Source/Programs/Horde/Horde.Build.Tests/StreamServiceTest.cs
@@ -5,36 +5,35 @@ using System.Globalization;
 using System.Threading.Tasks;
 using Horde.Build.Models;
 using Microsoft.VisualStudio.TestTools.UnitTesting;
-using PoolId = Horde.Build.Utilities.StringId<Horde.Build.Models.IPool>;
 
 namespace Horde.Build.Tests
 {
-    [TestClass]
+	[TestClass]
     public class StreamServiceTests : TestSetup
     {
         [TestMethod]
         public async Task Pausing()
         {
-			Fixture Fixture = await CreateFixtureAsync();
+			Fixture fixture = await CreateFixtureAsync();
 
-	        IStream Stream = (await StreamService.GetStreamAsync(Fixture!.Stream!.Id))!;
-	        Assert.IsFalse(Stream.IsPaused(DateTime.UtcNow));
-	        Assert.IsNull(Stream.PausedUntil);
-	        Assert.IsNull(Stream.PauseComment);
+	        IStream stream = (await StreamService.GetStreamAsync(fixture!.Stream!.Id))!;
+	        Assert.IsFalse(stream.IsPaused(DateTime.UtcNow));
+	        Assert.IsNull(stream.PausedUntil);
+	        Assert.IsNull(stream.PauseComment);
 
-	        DateTime PausedUntil = DateTime.UtcNow.AddHours(1);
-	        await StreamService.UpdatePauseStateAsync(Stream, NewPausedUntil: PausedUntil, NewPauseComment: "mycomment");
-	        Stream = (await StreamService.GetStreamAsync(Fixture!.Stream!.Id))!;
+	        DateTime pausedUntil = DateTime.UtcNow.AddHours(1);
+	        await StreamService.UpdatePauseStateAsync(stream, NewPausedUntil: pausedUntil, NewPauseComment: "mycomment");
+	        stream = (await StreamService.GetStreamAsync(fixture!.Stream!.Id))!;
 	        // Comparing by string to avoid comparing exact milliseconds as those are not persisted in MongoDB fields
-	        Assert.IsTrue(Stream.IsPaused(DateTime.UtcNow));
-	        Assert.AreEqual(PausedUntil.ToString(CultureInfo.InvariantCulture), Stream.PausedUntil!.Value.ToString(CultureInfo.InvariantCulture));
-	        Assert.AreEqual("mycomment", Stream.PauseComment);
+	        Assert.IsTrue(stream.IsPaused(DateTime.UtcNow));
+	        Assert.AreEqual(pausedUntil.ToString(CultureInfo.InvariantCulture), stream.PausedUntil!.Value.ToString(CultureInfo.InvariantCulture));
+	        Assert.AreEqual("mycomment", stream.PauseComment);
 	        
-	        await StreamService.UpdatePauseStateAsync(Stream, NewPausedUntil: null, NewPauseComment: null);
-	        Stream = (await StreamService.GetStreamAsync(Fixture!.Stream!.Id))!;
-	        Assert.IsFalse(Stream.IsPaused(DateTime.UtcNow));
-	        Assert.IsNull(Stream.PausedUntil);
-	        Assert.IsNull(Stream.PauseComment);
+	        await StreamService.UpdatePauseStateAsync(stream, NewPausedUntil: null, NewPauseComment: null);
+	        stream = (await StreamService.GetStreamAsync(fixture!.Stream!.Id))!;
+	        Assert.IsFalse(stream.IsPaused(DateTime.UtcNow));
+	        Assert.IsNull(stream.PausedUntil);
+	        Assert.IsNull(stream.PauseComment);
         }
     }
 }
\ No newline at end of file
diff --git a/Engine/Source/Programs/Horde/Horde.Build.Tests/Stubs/Collections/JobStepRefCollectionStub.cs b/Engine/Source/Programs/Horde/Horde.Build.Tests/Stubs/Collections/JobStepRefCollectionStub.cs
--- a/Engine/Source/Programs/Horde/Horde.Build.Tests/Stubs/Collections/JobStepRefCollectionStub.cs
+++ b/Engine/Source/Programs/Horde/Horde.Build.Tests/Stubs/Collections/JobStepRefCollectionStub.cs
@@ -4,11 +4,8 @@ using Horde.Build.Collections;
 using HordeCommon;
 using Horde.Build.Models;
 using Horde.Build.Utilities;
-using MongoDB.Bson;
-using Moq;
 using System;
 using System.Collections.Generic;
-using System.Text;
 using System.Threading.Tasks;
 
 using StreamId = Horde.Build.Utilities.StringId<Horde.Build.Models.IStream>;
@@ -40,65 +37,65 @@ namespace Horde.Build.Tests.Stubs.Collections
 		public virtual DateTime StartTimeUtc => throw new NotImplementedException();
 		public virtual DateTime? FinishTimeUtc => throw new NotImplementedException();
 
-		public JobStepRefStub(JobId JobId, SubResourceId BatchId, SubResourceId StepId, string JobName, string NodeName, StreamId StreamId, TemplateRefId TemplateId, int Change, JobStepOutcome? Outcome)
+		public JobStepRefStub(JobId jobId, SubResourceId batchId, SubResourceId stepId, string jobName, string nodeName, StreamId streamId, TemplateRefId templateId, int change, JobStepOutcome? outcome)
 		{
-			this.Id = new JobStepRefId(JobId, BatchId, StepId);
-			this.JobName = JobName;
-			this.NodeName = NodeName;
-			this.StreamId = StreamId;
-			this.TemplateId = TemplateId;
-			this.Change = Change;
-			this.Outcome = Outcome;
+			Id = new JobStepRefId(jobId, batchId, stepId);
+			this.JobName = jobName;
+			this.NodeName = nodeName;
+			this.StreamId = streamId;
+			this.TemplateId = templateId;
+			this.Change = change;
+			this.Outcome = outcome;
 		}
 	}
 
 	class JobStepRefCollectionStub : IJobStepRefCollection
 	{
-		List<IJobStepRef> Refs = new List<IJobStepRef>();
+		List<IJobStepRef> _refs = new List<IJobStepRef>();
 
-		public void Add(IJobStepRef JobStepRef)
+		public void Add(IJobStepRef jobStepRef)
 		{
-			Refs.Add(JobStepRef);
+			_refs.Add(jobStepRef);
 		}
 
-		public Task<IJobStepRef?> GetNextStepForNodeAsync(StreamId StreamId, TemplateRefId TemplateId, string NodeName, int Change)
+		public Task<IJobStepRef?> GetNextStepForNodeAsync(StreamId streamId, TemplateRefId templateId, string nodeName, int change)
 		{
-			IJobStepRef? NextRef = null;
-			foreach (IJobStepRef Ref in Refs)
+			IJobStepRef? nextRef = null;
+			foreach (IJobStepRef jobStepRef in _refs)
 			{
-				if (Ref.StreamId == StreamId && Ref.TemplateId == TemplateId && Ref.NodeName == NodeName && Ref.Change > Change)
+				if (jobStepRef.StreamId == streamId && jobStepRef.TemplateId == templateId && jobStepRef.NodeName == nodeName && jobStepRef.Change > change)
 				{
-					if (NextRef == null || Ref.Change < NextRef.Change)
+					if (nextRef == null || jobStepRef.Change < nextRef.Change)
 					{
-						NextRef = Ref;
+						nextRef = jobStepRef;
 					}
 				}
 			}
-			return Task.FromResult(NextRef);
+			return Task.FromResult(nextRef);
 		}
 
-		public Task<IJobStepRef?> GetPrevStepForNodeAsync(StreamId StreamId, TemplateRefId TemplateId, string NodeName, int Change)
+		public Task<IJobStepRef?> GetPrevStepForNodeAsync(StreamId streamId, TemplateRefId templateId, string nodeName, int change)
 		{
-			IJobStepRef? PrevRef = null;
-			foreach (IJobStepRef Ref in Refs)
+			IJobStepRef? prevRef = null;
+			foreach (IJobStepRef jobStepRef in _refs)
 			{
-				if (Ref.StreamId == StreamId && Ref.TemplateId == TemplateId && Ref.NodeName == NodeName && Ref.Change < Change)
+				if (jobStepRef.StreamId == streamId && jobStepRef.TemplateId == templateId && jobStepRef.NodeName == nodeName && jobStepRef.Change < change)
 				{
-					if (PrevRef == null || Ref.Change > PrevRef.Change)
+					if (prevRef == null || jobStepRef.Change > prevRef.Change)
 					{
-						PrevRef = Ref;
+						prevRef = jobStepRef;
 					}
 				}
 			}
-			return Task.FromResult(PrevRef);
+			return Task.FromResult(prevRef);
 		}
 
-		Task<List<IJobStepRef>> IJobStepRefCollection.GetStepsForNodeAsync(StreamId StreamId, TemplateRefId TemplateId, string NodeName, int? Change, bool IncludeFailed, int Count)
+		Task<List<IJobStepRef>> IJobStepRefCollection.GetStepsForNodeAsync(StreamId streamId, TemplateRefId templateId, string nodeName, int? change, bool includeFailed, int count)
 		{
 			throw new NotImplementedException();
 		}
 
-		Task<IJobStepRef> IJobStepRefCollection.InsertOrReplaceAsync(JobStepRefId Id, string JobName, string NodeName, StreamId StreamId, TemplateRefId TemplateId, int Change, LogId? LogId, PoolId? PoolId, AgentId? AgentId, JobStepOutcome? Outcome, int? LastSuccess, int? LastWarning, float WaitTime, float InitTime, DateTime StartTime, DateTime? FinishTime)
+		Task<IJobStepRef> IJobStepRefCollection.InsertOrReplaceAsync(JobStepRefId id, string jobName, string nodeName, StreamId streamId, TemplateRefId templateId, int change, LogId? logId, PoolId? poolId, AgentId? agentId, JobStepOutcome? outcome, int? lastSuccess, int? lastWarning, float waitTime, float initTime, DateTime startTime, DateTime? finishTime)
 		{
 			throw new NotImplementedException();
 		}
diff --git a/Engine/Source/Programs/Horde/Horde.Build.Tests/Stubs/Collections/SingletonDocumentStub.cs b/Engine/Source/Programs/Horde/Horde.Build.Tests/Stubs/Collections/SingletonDocumentStub.cs
--- a/Engine/Source/Programs/Horde/Horde.Build.Tests/Stubs/Collections/SingletonDocumentStub.cs
+++ b/Engine/Source/Programs/Horde/Horde.Build.Tests/Stubs/Collections/SingletonDocumentStub.cs
@@ -2,25 +2,22 @@
 
 using Horde.Build.Models;
 using Horde.Build.Utilities;
-using System;
-using System.Collections.Generic;
-using System.Text;
 using System.Threading.Tasks;
 
 namespace Horde.Build.Tests.Stubs.Collections
 {
 	class SingletonDocumentStub<T> : ISingletonDocument<T> where T : SingletonBase, new()
 	{
-		T Document = new T();
+		T _document = new T();
 
 		public Task<T> GetAsync()
 		{
-			return Task.FromResult(Document);
+			return Task.FromResult(_document);
 		}
 
-		public Task<bool> TryUpdateAsync(T Value)
+		public Task<bool> TryUpdateAsync(T value)
 		{
-			Document = Value;
+			_document = value;
 			return Task.FromResult(true);
 		}
 	}
diff --git a/Engine/Source/Programs/Horde/Horde.Build.Tests/Stubs/Services/PerforceServiceStub.cs b/Engine/Source/Programs/Horde/Horde.Build.Tests/Stubs/Services/PerforceServiceStub.cs
--- a/Engine/Source/Programs/Horde/Horde.Build.Tests/Stubs/Services/PerforceServiceStub.cs
+++ b/Engine/Source/Programs/Horde/Horde.Build.Tests/Stubs/Services/PerforceServiceStub.cs
@@ -7,28 +7,23 @@ using Horde.Build.Services;
 using Horde.Build.Utilities;
 using System;
 using System.Collections.Generic;
-using System.Diagnostics.CodeAnalysis;
 using System.Linq;
-using System.Security.Cryptography.X509Certificates;
-using System.Text;
 using System.Threading.Tasks;
 
-using StreamId = Horde.Build.Utilities.StringId<Horde.Build.Models.IStream>;
-
 namespace Horde.Build.Tests.Stubs.Services
 {
 	using UserId = ObjectId<IUser>;
 
 	static class PerforceExtensions
 	{
-		public static ChangeFile CreateChangeFile(string Path)
+		public static ChangeFile CreateChangeFile(string path)
 		{
-			return new ChangeFile(Path, null!, 0, 0, EpicGames.Core.Md5Hash.Zero, null!);
+			return new ChangeFile(path, null!, 0, 0, EpicGames.Core.Md5Hash.Zero, null!);
 		}
 
-		public static void Add(this List<ChangeFile> Files, string Path)
+		public static void Add(this List<ChangeFile> files, string path)
 		{
-			Files.Add(CreateChangeFile(Path));
+			files.Add(CreateChangeFile(path));
 		}
 	}
 
@@ -41,185 +36,184 @@ namespace Horde.Build.Tests.Stubs.Services
 			public string Email { get; set; }
 			public string Login { get; set; }
 
-			public User(string Login)
+			public User(string login)
 			{
-				this.Id = UserId.GenerateNewId();
-				this.Name = Login.ToUpperInvariant();
-				this.Email = $"{Login}@server";
-				this.Login = Login;
+				Id = UserId.GenerateNewId();
+				Name = login.ToUpperInvariant();
+				Email = $"{login}@server";
+				Login = login;
 			}
 		}
 
-		User TestUser = new User("TestUser");
+		User _testUser = new User("TestUser");
 
 		class ChangeComparer : IComparer<int>
 		{
-			public int Compare(int X, int Y) => Y.CompareTo(X);
+			public int Compare(int x, int y) => y.CompareTo(x);
 		}
 
-		Dictionary<string, IUser> NameToUser = new Dictionary<string, IUser>(StringComparer.OrdinalIgnoreCase);
 		public Dictionary<string, SortedDictionary<int, ChangeDetails>> Changes { get; } = new Dictionary<string, SortedDictionary<int, ChangeDetails>>(StringComparer.OrdinalIgnoreCase);
 
-		IUserCollection UserCollection;
+		IUserCollection _userCollection;
 
-		public PerforceServiceStub(IUserCollection UserCollection)
+		public PerforceServiceStub(IUserCollection userCollection)
 		{
-			this.UserCollection = UserCollection;
+			this._userCollection = userCollection;
 		}
 
-		public Task<IPerforceConnection?> GetServiceUserConnection(string? ClusterName)
+		public Task<IPerforceConnection?> GetServiceUserConnection(string? clusterName)
 		{
 			throw new NotImplementedException();
 		}
 
-		public async ValueTask<IUser> FindOrAddUserAsync(string ClusterName, string UserName)
+		public async ValueTask<IUser> FindOrAddUserAsync(string clusterName, string userName)
 		{
-			return await UserCollection.FindOrAddUserByLoginAsync(UserName);
+			return await _userCollection.FindOrAddUserByLoginAsync(userName);
 		}
 
-		public void AddChange(string StreamName, int Number, IUser Author, string Description, IEnumerable<string> Files)
+		public void AddChange(string streamName, int number, IUser author, string description, IEnumerable<string> files)
 		{
-			SortedDictionary<int, ChangeDetails>? StreamChanges;
-			if (!Changes.TryGetValue(StreamName, out StreamChanges))
+			SortedDictionary<int, ChangeDetails>? streamChanges;
+			if (!Changes.TryGetValue(streamName, out streamChanges))
 			{
-				StreamChanges = new SortedDictionary<int, ChangeDetails>(new ChangeComparer());
-				Changes[StreamName] = StreamChanges;
+				streamChanges = new SortedDictionary<int, ChangeDetails>(new ChangeComparer());
+				Changes[streamName] = streamChanges;
 			}
-			StreamChanges.Add(Number, new ChangeDetails(Number, Author, null!, Description, Files.Select(x => PerforceExtensions.CreateChangeFile(x)).ToList(), DateTime.Now));
+			streamChanges.Add(number, new ChangeDetails(number, author, null!, description, files.Select(x => PerforceExtensions.CreateChangeFile(x)).ToList(), DateTime.Now));
 		}
 
-		public Task<List<ChangeSummary>> GetChangesAsync(string ClusterName, string StreamName, int? MinChange, int? MaxChange, int NumResults, string? ImpersonateUser)
+		public Task<List<ChangeSummary>> GetChangesAsync(string clusterName, string streamName, int? minChange, int? maxChange, int numResults, string? impersonateUser)
 		{
-			List<ChangeSummary> Results = new List<ChangeSummary>();
+			List<ChangeSummary> results = new List<ChangeSummary>();
 
-			SortedDictionary<int, ChangeDetails>? StreamChanges;
-			if (Changes.TryGetValue(StreamName, out StreamChanges) && StreamChanges.Count > 0)
+			SortedDictionary<int, ChangeDetails>? streamChanges;
+			if (Changes.TryGetValue(streamName, out streamChanges) && streamChanges.Count > 0)
 			{
-				foreach (ChangeDetails Details in StreamChanges.Values)
+				foreach (ChangeDetails details in streamChanges.Values)
 				{
-					if (MinChange.HasValue && Details.Number < MinChange)
+					if (minChange.HasValue && details.Number < minChange)
 					{
 						break;
 					}
-					if (!MaxChange.HasValue || Details.Number <= MaxChange.Value)
+					if (!maxChange.HasValue || details.Number <= maxChange.Value)
 					{
-						Results.Add(new ChangeSummary(Details.Number, Details.Author, "//...", Details.Description));
+						results.Add(new ChangeSummary(details.Number, details.Author, "//...", details.Description));
 					}
-					if (NumResults > 0 && Results.Count >= NumResults)
+					if (numResults > 0 && results.Count >= numResults)
 					{
 						break;
 					}
 				}
 			}
 
-			return Task.FromResult(Results);
+			return Task.FromResult(results);
 		}
 
-		public Task<CheckShelfResult> CheckShelfAsync(string ClusterName, string StreamName, int ChangeNumber, string? ImpersonateUser)
+		public Task<CheckShelfResult> CheckShelfAsync(string clusterName, string streamName, int changeNumber, string? impersonateUser)
 		{
 			throw new NotImplementedException();
 		}
 
-		public Task<PerforceUserInfo?> GetUserInfoAsync(string ClusterName, string UserName)
+		public Task<PerforceUserInfo?> GetUserInfoAsync(string clusterName, string userName)
 		{
-			return Task.FromResult<PerforceUserInfo?>(new PerforceUserInfo { Login = UserName, FullName = UserName, Email = $"[REDACTED_EMAIL]" });
+			return Task.FromResult<PerforceUserInfo?>(new PerforceUserInfo { Login = userName, FullName = userName, Email = $"[REDACTED_EMAIL]" });
 		}
 
-		public Task<List<ChangeDetails>> GetChangeDetailsAsync(string ClusterName, string StreamName, IReadOnlyList<int> ChangeNumbers, string? ImpersonateUser)
+		public Task<List<ChangeDetails>> GetChangeDetailsAsync(string clusterName, string streamName, IReadOnlyList<int> changeNumbers, string? impersonateUser)
 		{
-			List<ChangeDetails> Results = new List<ChangeDetails>();
-			foreach (int ChangeNumber in ChangeNumbers)
+			List<ChangeDetails> results = new List<ChangeDetails>();
+			foreach (int changeNumber in changeNumbers)
 			{
-				Results.Add(Changes[StreamName][ChangeNumber]);
+				results.Add(Changes[streamName][changeNumber]);
 			}
-			return Task.FromResult(Results);
+			return Task.FromResult(results);
 		}
 
-		public Task<string> CreateTicket(string ClusterName, string ImpersonateUser)
+		public static Task<string> CreateTicket()
 		{
 			return Task.FromResult("bogus-ticket");
 		}
 
-		public Task<int> GetCodeChangeAsync(string ClusterName, string StreamName, int Change)
+		public Task<int> GetCodeChangeAsync(string clusterName, string streamName, int change)
 		{
-			int CodeChange = 0;
+			int codeChange = 0;
 
-			SortedDictionary<int, ChangeDetails>? StreamChanges;
-			if (Changes.TryGetValue(StreamName, out StreamChanges))
+			SortedDictionary<int, ChangeDetails>? streamChanges;
+			if (Changes.TryGetValue(streamName, out streamChanges))
 			{
-				foreach (ChangeDetails Details in StreamChanges.Values)
+				foreach (ChangeDetails details in streamChanges.Values)
 				{
-					if (Details.Number <= Change && Details.Files.Any(x => x.Path.EndsWith(".h") || x.Path.EndsWith(".cpp")))
+					if (details.Number <= change && details.Files.Any(x => x.Path.EndsWith(".h", StringComparison.OrdinalIgnoreCase) || x.Path.EndsWith(".cpp", StringComparison.OrdinalIgnoreCase)))
 					{
-						CodeChange = Details.Number;
+						codeChange = details.Number;
 						break;
 					}
 				}
 			}
 
-			return Task.FromResult(CodeChange);
+			return Task.FromResult(codeChange);
 		}
 
-		public Task<int> CreateNewChangeAsync(string ClusterName, string StreamName, string Path, string Description)
+		public Task<int> CreateNewChangeAsync(string clusterName, string streamName, string path, string description)
 		{
-			ChangeDetails NewChange = new ChangeDetails(Changes[StreamName].First().Key + 1, TestUser, null!, Description, new List<ChangeFile> { PerforceExtensions.CreateChangeFile(Path) }, DateTime.Now);
-			Changes[StreamName].Add(NewChange.Number, NewChange);
-			return Task.FromResult(NewChange.Number);
+			ChangeDetails newChange = new ChangeDetails(Changes[streamName].First().Key + 1, _testUser, null!, description, new List<ChangeFile> { PerforceExtensions.CreateChangeFile(path) }, DateTime.Now);
+			Changes[streamName].Add(newChange.Number, newChange);
+			return Task.FromResult(newChange.Number);
 		}
 
-		public Task<List<FileSummary>> FindFilesAsync(string ClusterName, IEnumerable<string> Paths)
+		public Task<List<FileSummary>> FindFilesAsync(string clusterName, IEnumerable<string> paths)
 		{
 			throw new NotImplementedException();
 		}
 
-		public Task<byte[]> PrintAsync(string ClusterName, string Path)
+		public Task<byte[]> PrintAsync(string clusterName, string path)
 		{
 			throw new NotImplementedException();
 		}
 
-		public Task<(int? Change, string Message)> SubmitShelvedChangeAsync(string ClusterName, int ShelvedChange, int OriginalChange)
+		public Task<(int? Change, string Message)> SubmitShelvedChangeAsync(string clusterName, int shelvedChange, int originalChange)
 		{
 			throw new NotImplementedException();
 		}
 
-		public Task<int> DuplicateShelvedChangeAsync(string ClusterName, int ShelvedChange)
+		public Task<int> DuplicateShelvedChangeAsync(string clusterName, int shelvedChange)
 		{
-			return Task.FromResult(ShelvedChange);
+			return Task.FromResult(shelvedChange);
 		}
 
-		public Task DeleteShelvedChangeAsync(string ClusterName, int ShelvedChange)
+		public Task DeleteShelvedChangeAsync(string clusterName, int shelvedChange)
 		{
 			return Task.CompletedTask;
 		}
 
-		public Task UpdateChangelistDescription(string ClusterName, int Change, string Description)
+		public Task UpdateChangelistDescription(string clusterName, int change, string description)
 		{
 			return Task.CompletedTask;
 		}
 
-		public Task<IStreamView> GetStreamViewAsync(string ClusterName, string StreamName)
+		public Task<IStreamView> GetStreamViewAsync(string clusterName, string streamName)
 		{
 			throw new NotImplementedException();
 		}
 
-		public Task<List<ChangeSummary>> GetChangesAsync(string ClusterName, int? MinChange, int? MaxChange, int MaxResults)
+		public Task<List<ChangeSummary>> GetChangesAsync(string clusterName, int? minChange, int? maxChange, int maxResults)
 		{
 			throw new NotImplementedException();
 		}
 
-		public Task<ChangeDetails> GetChangeDetailsAsync(string ClusterName, int ChangeNumber)
+		public Task<ChangeDetails> GetChangeDetailsAsync(string clusterName, int changeNumber)
 		{
 			throw new NotImplementedException();
 		}
 
-		public Task<List<ChangeFile>> GetStreamSnapshotAsync(string ClusterName, string StreamName, int Change)
+		public Task<List<ChangeFile>> GetStreamSnapshotAsync(string clusterName, string streamName, int change)
 		{
 			throw new NotImplementedException();
 		}
 
-		public Task<ChangeDetails> GetChangeDetailsAsync(string ClusterName, string StreamName, int ChangeNumber)
+		public Task<ChangeDetails> GetChangeDetailsAsync(string clusterName, string streamName, int changeNumber)
 		{
-			return Task.FromResult(Changes[StreamName][ChangeNumber]);
+			return Task.FromResult(Changes[streamName][changeNumber]);
 		}
 	}
 }
diff --git a/Engine/Source/Programs/Horde/Horde.Build.Tests/SwaggerTest.cs b/Engine/Source/Programs/Horde/Horde.Build.Tests/SwaggerTest.cs
--- a/Engine/Source/Programs/Horde/Horde.Build.Tests/SwaggerTest.cs
+++ b/Engine/Source/Programs/Horde/Horde.Build.Tests/SwaggerTest.cs
@@ -1,28 +1,28 @@
 ﻿// Copyright Epic Games, Inc. All Rights Reserved.
 
 using System;
+using System.Net.Http;
 using System.Text.Json;
 using System.Threading.Tasks;
-using Horde.Build.Tests;
 using Microsoft.VisualStudio.TestTools.UnitTesting;
 
 namespace Horde.Build.Tests
 {
-    [TestClass]
+	[TestClass]
     public class SwaggerTest : ControllerIntegrationTest
     {
         [TestMethod]
         [Ignore("Until clashes with Horde.Build and Horde.Storage are resolved")]
         public async Task ValidateSwagger()
         {
-	        var Res = await client.GetAsync("/swagger/v1/swagger.json");
-	        if (!Res.IsSuccessStatusCode)
+	        HttpResponseMessage res = await Client.GetAsync(new Uri("/swagger/v1/swagger.json"));
+	        if (!res.IsSuccessStatusCode)
 	        {
-		        string RawJson = await Res.Content.ReadAsStringAsync();
-		        JsonElement TempElement = JsonSerializer.Deserialize<JsonElement>(RawJson);
-		        string FormattedJson = JsonSerializer.Serialize(TempElement, new JsonSerializerOptions { WriteIndented = true });
-		        Console.Error.WriteLine("Error result:\n" + FormattedJson);
-		        Res.EnsureSuccessStatusCode();
+		        string rawJson = await res.Content.ReadAsStringAsync();
+		        JsonElement tempElement = JsonSerializer.Deserialize<JsonElement>(rawJson);
+		        string formattedJson = JsonSerializer.Serialize(tempElement, new JsonSerializerOptions { WriteIndented = true });
+		        await Console.Error.WriteLineAsync("Error result:\n" + formattedJson);
+		        res.EnsureSuccessStatusCode();
 	        }
         }
     }
