Horde: Update Horde.Build.Tests to use standard coding conventions, and fix a number of issues caught by static analysis.

diff --git a/Engine/Source/Programs/Horde/Horde.Build.Tests/TestSetup.cs b/Engine/Source/Programs/Horde/Horde.Build.Tests/TestSetup.cs
--- a/Engine/Source/Programs/Horde/Horde.Build.Tests/TestSetup.cs
+++ b/Engine/Source/Programs/Horde/Horde.Build.Tests/TestSetup.cs
@@ -11,7 +11,6 @@ using EpicGames.Horde.Storage;
 using HordeCommon;
 using Horde.Build.Acls;
 using Horde.Build.Storage.Services;
-using Horde.Build;
 using Horde.Build.Collections;
 using Horde.Build.Collections.Impl;
 using Horde.Build.Controllers;
@@ -26,7 +25,6 @@ using Horde.Build.Notifications.Impl;
 using Horde.Build.Services;
 using Horde.Build.Services.Impl;
 using Horde.Build.Storage;
-using Horde.Build.Storage.Backends;
 using Horde.Build.Tasks.Impl;
 using Horde.Build.Utilities;
 using Horde.Build.Tests.Stubs.Collections;
@@ -37,8 +35,6 @@ using Microsoft.Extensions.Configuration;
 using Microsoft.Extensions.DependencyInjection;
 using Microsoft.Extensions.Hosting;
 using Microsoft.Extensions.Logging;
-using Microsoft.Extensions.Logging.Abstractions;
-using Microsoft.Extensions.Logging.Console;
 using Microsoft.Extensions.Options;
 using Moq;
 using StatsdClient;
@@ -52,8 +48,7 @@ namespace Horde.Build.Tests
 	/// </summary>
 	public class TestSetup : DatabaseIntegrationTest
 	{
-		public IServiceProvider ServiceProvider { get; }
-		public FakeClock Clock { get; set; } = null!;
+		public FakeClock Clock => ServiceProvider.GetRequiredService<FakeClock>();
 
 		public IGraphCollection GraphCollection => ServiceProvider.GetRequiredService<IGraphCollection>();
 		public INotificationTriggerCollection NotificationTriggerCollection => ServiceProvider.GetRequiredService<INotificationTriggerCollection>();
@@ -102,168 +97,162 @@ namespace Horde.Build.Tests
 		public PoolsController PoolsController => GetPoolsController();
 		public LeasesController LeasesController => GetLeasesController();
 
-		private static bool IsDatadogWriterPatched;
+		private static bool s_datadogWriterPatched;
 
 		public TestSetup()
 		{
 			PatchDatadogWriter();
-
-			IConfiguration Config = new ConfigurationBuilder().Build();
-
-			IServiceCollection Services = new ServiceCollection();
-			Services.AddSingleton(GetDatabaseService());
-			Services.Configure<ServerSettings>(ConfigureSettings);
-			Services.AddSingleton<IConfiguration>(Config);
-
-			ConfigureServices(Services);
-
-			ServiceProvider = Services.BuildServiceProvider();
 		}
 
-		protected virtual void ConfigureSettings(ServerSettings Settings)
+		protected override void ConfigureSettings(ServerSettings settings)
 		{
-			DirectoryReference BaseDir = DirectoryReference.Combine(Program.DataDir, "Tests");
+			DirectoryReference baseDir = DirectoryReference.Combine(Program.DataDir, "Tests");
 			try
 			{
-				FileUtils.ForceDeleteDirectoryContents(BaseDir);
+				FileUtils.ForceDeleteDirectoryContents(baseDir);
 			}
 			catch
 			{
 			}
 
-			Settings.AdminClaimType = HordeClaimTypes.Role;
-			Settings.AdminClaimValue = "app-horde-admins";
+			settings.AdminClaimType = HordeClaimTypes.Role;
+			settings.AdminClaimValue = "app-horde-admins";
 		}
 
-		protected virtual void ConfigureServices(IServiceCollection Services)
+		protected override void ConfigureServices(IServiceCollection services)
 		{
-			Services.AddLogging(Builder => Builder.AddConsole());
+			base.ConfigureServices(services);
+
+			IConfiguration config = new ConfigurationBuilder().Build();
+			services.Configure<ServerSettings>(ConfigureSettings);
+			services.AddSingleton<IConfiguration>(config);
+
+			services.AddLogging(builder => builder.AddConsole());
 
-			Services.AddSingleton<RedisService>();
-			Services.AddSingleton<StackExchange.Redis.IDatabase>(new Mock<StackExchange.Redis.IDatabase>().Object);
+			services.AddSingleton<RedisService>();
 
-			Services.AddSingleton(typeof(IAuditLogFactory<>), typeof(AuditLogFactory<>));
-			Services.AddSingleton<IAuditLog<AgentId>>(SP => SP.GetRequiredService<IAuditLogFactory<AgentId>>().Create("Agents.Log", "AgentId"));
+			services.AddSingleton(typeof(IAuditLogFactory<>), typeof(AuditLogFactory<>));
+			services.AddSingleton<IAuditLog<AgentId>>(sp => sp.GetRequiredService<IAuditLogFactory<AgentId>>().Create("Agents.Log", "AgentId"));
 
-			Services.AddSingleton<IAgentCollection, AgentCollection>();
-			Services.AddSingleton<IAgentSoftwareCollection, AgentSoftwareCollection>();
-			Services.AddSingleton<IArtifactCollection, ArtifactCollection>();
-			Services.AddSingleton<IGraphCollection, GraphCollection>();
-			Services.AddSingleton<IIssueCollection, IssueCollection>();
-			Services.AddSingleton<IJobCollection, JobCollection>();
-			Services.AddSingleton<IJobStepRefCollection, JobStepRefCollection>();
-			Services.AddSingleton<IJobTimingCollection, JobTimingCollection>();
-			Services.AddSingleton<ILeaseCollection, LeaseCollection>();
-			Services.AddSingleton<ILogEventCollection, LogEventCollection>();
-			Services.AddSingleton<ILogFileCollection, LogFileCollection>();
-			Services.AddSingleton<INotificationTriggerCollection, NotificationTriggerCollection>();
-			Services.AddSingleton<IPoolCollection, PoolCollection>();
-			Services.AddSingleton<IProjectCollection, ProjectCollection>();
-			Services.AddSingleton<ISessionCollection, SessionCollection>();
-			Services.AddSingleton<ISubscriptionCollection, SubscriptionCollection>();
-			Services.AddSingleton<IStreamCollection, StreamCollection>();
-			Services.AddSingleton<ITemplateCollection, TemplateCollection>();
-			Services.AddSingleton<ITestDataCollection, TestDataCollection>();
-			Services.AddSingleton<ITelemetryCollection, TelemetryCollection>();
-			Services.AddSingleton<ITemplateCollection, TemplateCollection>();
-			Services.AddSingleton<IUgsMetadataCollection, UgsMetadataCollection>();
-			Services.AddSingleton<IUserCollection, UserCollectionV1>();
+			services.AddSingleton<IAgentCollection, AgentCollection>();
+			services.AddSingleton<IAgentSoftwareCollection, AgentSoftwareCollection>();
+			services.AddSingleton<IArtifactCollection, ArtifactCollection>();
+			services.AddSingleton<IGraphCollection, GraphCollection>();
+			services.AddSingleton<IIssueCollection, IssueCollection>();
+			services.AddSingleton<IJobCollection, JobCollection>();
+			services.AddSingleton<IJobStepRefCollection, JobStepRefCollection>();
+			services.AddSingleton<IJobTimingCollection, JobTimingCollection>();
+			services.AddSingleton<ILeaseCollection, LeaseCollection>();
+			services.AddSingleton<ILogEventCollection, LogEventCollection>();
+			services.AddSingleton<ILogFileCollection, LogFileCollection>();
+			services.AddSingleton<INotificationTriggerCollection, NotificationTriggerCollection>();
+			services.AddSingleton<IPoolCollection, PoolCollection>();
+			services.AddSingleton<IProjectCollection, ProjectCollection>();
+			services.AddSingleton<ISessionCollection, SessionCollection>();
+			services.AddSingleton<ISubscriptionCollection, SubscriptionCollection>();
+			services.AddSingleton<IStreamCollection, StreamCollection>();
+			services.AddSingleton<ITemplateCollection, TemplateCollection>();
+			services.AddSingleton<ITestDataCollection, TestDataCollection>();
+			services.AddSingleton<ITelemetryCollection, TelemetryCollection>();
+			services.AddSingleton<ITemplateCollection, TemplateCollection>();
+			services.AddSingleton<IUgsMetadataCollection, UgsMetadataCollection>();
+			services.AddSingleton<IUserCollection, UserCollectionV1>();
 
-			Clock = new FakeClock();
-			Services.AddSingleton<IClock>(Clock);
-			Services.AddSingleton<IHostApplicationLifetime, AppLifetimeStub>();
+			services.AddSingleton<FakeClock>();
+			services.AddSingleton<IClock>(SP => SP.GetRequiredService<FakeClock>());
+			services.AddSingleton<IHostApplicationLifetime, AppLifetimeStub>();
 
-			Services.AddSingleton<AclService>();
-			Services.AddSingleton<AgentService>();
-			Services.AddSingleton<AgentSoftwareService>();
-			Services.AddSingleton<AutoscaleService>();
-			Services.AddSingleton<AwsFleetManager, AwsFleetManager>();
-			Services.AddSingleton<ConsistencyService>();
-			Services.AddSingleton<RequestTrackerService>();
-			Services.AddSingleton<CredentialService>();
-			Services.AddSingleton<JobTaskSource>();
-			Services.AddSingleton<IDowntimeService, DowntimeServiceStub>();
-			Services.AddSingleton<IDogStatsd, NoOpDogStatsd>();
-			Services.AddSingleton<IIssueService, IssueService>();
-			Services.AddSingleton<IFleetManager, AwsFleetManager>();
-			Services.AddSingleton<JobService>();
-			Services.AddSingleton<LifetimeService>();
-			Services.AddSingleton<ILogStorage, NullLogStorage>();
-			Services.AddSingleton<ILogBuilder, LocalLogBuilder>();
-			Services.AddSingleton<ILogFileService, LogFileService>();
-			Services.AddSingleton<INotificationService, NotificationService>();
-			Services.AddSingleton<IPerforceService, PerforceServiceStub>();
-			Services.AddSingleton<PerforceLoadBalancer>();
-			Services.AddSingleton<PoolService>();
-			Services.AddSingleton<ProjectService>();
-			Services.AddSingleton<RpcService>();
-			Services.AddSingleton<ScheduleService>();
-			Services.AddSingleton<StreamService>();
-			Services.AddSingleton<UpgradeService>();
+			services.AddSingleton<AclService>();
+			services.AddSingleton<AgentService>();
+			services.AddSingleton<AgentSoftwareService>();
+			services.AddSingleton<AutoscaleService>();
+			services.AddSingleton<AwsFleetManager, AwsFleetManager>();
+			services.AddSingleton<ConsistencyService>();
+			services.AddSingleton<RequestTrackerService>();
+			services.AddSingleton<CredentialService>();
+			services.AddSingleton<JobTaskSource>();
+			services.AddSingleton<IDowntimeService, DowntimeServiceStub>();
+			services.AddSingleton<IDogStatsd, NoOpDogStatsd>();
+			services.AddSingleton<IIssueService, IssueService>();
+			services.AddSingleton<IFleetManager, AwsFleetManager>();
+			services.AddSingleton<JobService>();
+			services.AddSingleton<LifetimeService>();
+			services.AddSingleton<ILogStorage, NullLogStorage>();
+			services.AddSingleton<ILogBuilder, LocalLogBuilder>();
+			services.AddSingleton<ILogFileService, LogFileService>();
+			services.AddSingleton<INotificationService, NotificationService>();
+			services.AddSingleton<IPerforceService, PerforceServiceStub>();
+			services.AddSingleton<PerforceLoadBalancer>();
+			services.AddSingleton<PoolService>();
+			services.AddSingleton<ProjectService>();
+			services.AddSingleton<RpcService>();
+			services.AddSingleton<ScheduleService>();
+			services.AddSingleton<StreamService>();
+			services.AddSingleton<UpgradeService>();
 
-			Services.AddSingleton<ConformTaskSource>();
+			services.AddSingleton<ConformTaskSource>();
 
-			Services.AddSingleton(new Startup.StorageBackendSettings<PersistentLogStorage> { Type = StorageProviderType.Transient });
-			Services.AddSingleton(new Startup.StorageBackendSettings<ArtifactCollection> { Type = StorageProviderType.Transient });
-			Services.AddSingleton(new Startup.StorageBackendSettings<BasicStorageClient> { Type = StorageProviderType.Transient });
-			Services.AddSingleton(typeof(IStorageBackend<>), typeof(Startup.StorageBackendFactory<>));
+			services.AddSingleton(new Startup.StorageBackendSettings<PersistentLogStorage> { Type = StorageProviderType.Transient });
+			services.AddSingleton(new Startup.StorageBackendSettings<ArtifactCollection> { Type = StorageProviderType.Transient });
+			services.AddSingleton(new Startup.StorageBackendSettings<BasicStorageClient> { Type = StorageProviderType.Transient });
+			services.AddSingleton(typeof(IStorageBackend<>), typeof(Startup.StorageBackendFactory<>));
 
-			Services.AddSingleton<IStorageClient, BasicStorageClient>();
+			services.AddSingleton<IStorageClient, BasicStorageClient>();
 
-			Services.AddSingleton<ISingletonDocument<AgentSoftwareChannels>>(new SingletonDocumentStub<AgentSoftwareChannels>());
+			services.AddSingleton<ISingletonDocument<AgentSoftwareChannels>>(new SingletonDocumentStub<AgentSoftwareChannels>());
 		}
 
 		public Task<Fixture> CreateFixtureAsync()
 		{
-			return Fixture.Create(GraphCollection, TemplateCollection, JobService, ArtifactCollection, StreamService, AgentService, PerforceService);
+			return Fixture.Create(GraphCollection, TemplateCollection, JobService, ArtifactCollection, StreamService, AgentService);
 		}
 
 		private JobsController GetJobsController()
         {
-			ILogger<JobsController> Logger = ServiceProvider.GetRequiredService<ILogger<JobsController>>();
-			JobsController JobsCtrl = new JobsController(AclService, GraphCollection, PerforceService, StreamService, JobService,
-		        TemplateCollection, ArtifactCollection, UserCollection, NotificationService, AgentService, Logger);
-	        JobsCtrl.ControllerContext = GetControllerContext();
-	        return JobsCtrl;
+			ILogger<JobsController> logger = ServiceProvider.GetRequiredService<ILogger<JobsController>>();
+			JobsController jobsCtrl = new JobsController(AclService, GraphCollection, PerforceService, StreamService, JobService,
+		        TemplateCollection, ArtifactCollection, UserCollection, NotificationService, AgentService, logger);
+	        jobsCtrl.ControllerContext = GetControllerContext();
+	        return jobsCtrl;
         }
 
 		private AgentsController GetAgentsController()
 		{
-			AgentsController AgentCtrl = new AgentsController(AclService, AgentService, PoolService);
-			AgentCtrl.ControllerContext = GetControllerContext();
-			return AgentCtrl;
+			AgentsController agentCtrl = new AgentsController(AclService, AgentService, PoolService);
+			agentCtrl.ControllerContext = GetControllerContext();
+			return agentCtrl;
 		}
 		
 		private PoolsController GetPoolsController()
 		{
-			PoolsController Controller = new PoolsController(AclService, PoolService);
-			Controller.ControllerContext = GetControllerContext();
-			return Controller;
+			PoolsController controller = new PoolsController(AclService, PoolService);
+			controller.ControllerContext = GetControllerContext();
+			return controller;
 		}
 		
 		private LeasesController GetLeasesController()
 		{
-			ILogger<LeasesController> Logger = ServiceProvider.GetRequiredService<ILogger<LeasesController>>();
-			LeasesController Controller = new LeasesController(AclService, AgentService, Logger);
-			Controller.ControllerContext = GetControllerContext();
-			return Controller;
+			ILogger<LeasesController> logger = ServiceProvider.GetRequiredService<ILogger<LeasesController>>();
+			LeasesController controller = new LeasesController(AclService, AgentService, logger);
+			controller.ControllerContext = GetControllerContext();
+			return controller;
 		}
 		
 		private ControllerContext GetControllerContext()
 		{
-			ControllerContext ControllerContext = new ControllerContext();
-			ControllerContext.HttpContext = new DefaultHttpContext();
-			ControllerContext.HttpContext.User = new ClaimsPrincipal(new ClaimsIdentity(
+			ControllerContext controllerContext = new ControllerContext();
+			controllerContext.HttpContext = new DefaultHttpContext();
+			controllerContext.HttpContext.User = new ClaimsPrincipal(new ClaimsIdentity(
 				new List<Claim> {new Claim(ServerSettings.AdminClaimType, ServerSettings.AdminClaimValue)}, "TestAuthType"));
-			return ControllerContext;
+			return controllerContext;
 		}
 		
-		private static int AgentIdCounter = 1;
-		public async Task<IAgent> CreateAgentAsync(IPool Pool)
+		private static int s_agentIdCounter = 1;
+		public async Task<IAgent> CreateAgentAsync(IPool pool)
 		{
-			IAgent Agent = await AgentService.CreateAgentAsync("TestAgent" + AgentIdCounter++, true, null, new List<StringId<IPool>> { Pool.Id });
-			Agent = await AgentService.CreateSessionAsync(Agent, AgentStatus.Ok, new List<string>(), new Dictionary<string, int>(), null);
-			return Agent;
+			IAgent agent = await AgentService.CreateAgentAsync("TestAgent" + s_agentIdCounter++, true, null, new List<StringId<IPool>> { pool.Id });
+			agent = await AgentService.CreateSessionAsync(agent, AgentStatus.Ok, new List<string>(), new Dictionary<string, int>(), null);
+			return agent;
 		}
 
 		/// <summary>
@@ -272,46 +261,46 @@ namespace Horde.Build.Tests
 		///
 		/// Since so many of the interfaces and classes in the lib are internal it was difficult to replace Tracer.Instance
 		/// </summary>
-		private void PatchDatadogWriter()
+		private static void PatchDatadogWriter()
 		{
-			if (IsDatadogWriterPatched)
+			if (s_datadogWriterPatched)
 			{
 				return;
 			}
 
-			IsDatadogWriterPatched = true;
+			s_datadogWriterPatched = true;
 
-			string Msg = "Unable to patch Datadog agent writer! Tests will still work, but shutdown will block for +20 seconds.";
+			string msg = "Unable to patch Datadog agent writer! Tests will still work, but shutdown will block for +20 seconds.";
 			
-			FieldInfo? AgentWriterField = Tracer.Instance.GetType().GetField("_agentWriter", BindingFlags.NonPublic | BindingFlags.Instance);
-			if (AgentWriterField == null)
+			FieldInfo? agentWriterField = Tracer.Instance.GetType().GetField("_agentWriter", BindingFlags.NonPublic | BindingFlags.Instance);
+			if (agentWriterField == null)
 			{
-				Console.Error.WriteLine(Msg);
+				Console.Error.WriteLine(msg);
 				return;
 			}
 			
-			object? AgentWriterInstance = AgentWriterField.GetValue(Tracer.Instance);
-			if (AgentWriterInstance == null)
+			object? agentWriterInstance = agentWriterField.GetValue(Tracer.Instance);
+			if (agentWriterInstance == null)
 			{
-				Console.Error.WriteLine(Msg);
+				Console.Error.WriteLine(msg);
 				return;	
 			}
 	        
-			FieldInfo? ProcessExitField = AgentWriterInstance.GetType().GetField("_processExit", BindingFlags.NonPublic | BindingFlags.Instance);
-			if (ProcessExitField == null)
+			FieldInfo? processExitField = agentWriterInstance.GetType().GetField("_processExit", BindingFlags.NonPublic | BindingFlags.Instance);
+			if (processExitField == null)
 			{
-				Console.Error.WriteLine(Msg);
+				Console.Error.WriteLine(msg);
 				return;
 			}
 			
-			TaskCompletionSource<bool>? ProcessExitInstance = (TaskCompletionSource<bool>?) ProcessExitField.GetValue(AgentWriterInstance);
-			if (ProcessExitInstance == null)
+			TaskCompletionSource<bool>? processExitInstance = (TaskCompletionSource<bool>?) processExitField.GetValue(agentWriterInstance);
+			if (processExitInstance == null)
 			{
-				Console.Error.WriteLine(Msg);
+				Console.Error.WriteLine(msg);
 				return;
 			}
 			
-			ProcessExitInstance.TrySetResult(true);
+			processExitInstance.TrySetResult(true);
 		}
 		
 		private class DowntimeServiceStub : IDowntimeService
diff --git a/Engine/Source/Programs/Horde/Horde.Build.Tests/VersionedCollectionTests.cs b/Engine/Source/Programs/Horde/Horde.Build.Tests/VersionedCollectionTests.cs
--- a/Engine/Source/Programs/Horde/Horde.Build.Tests/VersionedCollectionTests.cs
+++ b/Engine/Source/Programs/Horde/Horde.Build.Tests/VersionedCollectionTests.cs
@@ -49,16 +49,16 @@ namespace Horde.Build.Tests
 
 		DatabaseService _databaseService;
 		IDatabase _redis;
-		RedisKey baseKey = new RedisKey("versioned/");
+		RedisKey _baseKey = new RedisKey("versioned/");
 		IMongoCollection<VersionedDocument<string, DocumentV2>> _baseCollection;
 		VersionedCollection<string, DocumentV2> _collection;
 
 		public VersionedCollectionTests()
 		{
-			_databaseService = GetDatabaseService();
+			_databaseService = GetDatabaseServiceSingleton();
 			_redis = GetRedisDatabase();
 			_baseCollection = _databaseService.GetCollection<VersionedDocument<string, DocumentV2>>("versioned");
-			_collection = new VersionedCollection<string, DocumentV2>(_baseCollection, _redis, baseKey);
+			_collection = new VersionedCollection<string, DocumentV2>(_baseCollection, _redis, _baseKey);
 		}
 
 		[TestMethod]
