Horde: Few more minor code style fixes.

diff --git a/Engine/Source/Programs/Horde/Horde.Build.Tests/DatabaseRunner.cs b/Engine/Source/Programs/Horde/Horde.Build.Tests/DatabaseRunner.cs
--- a/Engine/Source/Programs/Horde/Horde.Build.Tests/DatabaseRunner.cs
+++ b/Engine/Source/Programs/Horde/Horde.Build.Tests/DatabaseRunner.cs
@@ -27,12 +27,12 @@ public abstract class DatabaseRunner
 	protected DatabaseRunner(string name, string binName, int defaultPort, bool reuseProcess, bool printStdOut = false,
 		bool printStdErr = true)
 	{
-		this._name = name;
-		this._binName = binName;
-		this._defaultPort = defaultPort;
-		this._reuseProcess = reuseProcess;
-		this._printStdOut = printStdOut;
-		this._printStdErr = printStdErr;
+		_name = name;
+		_binName = binName;
+		_defaultPort = defaultPort;
+		_reuseProcess = reuseProcess;
+		_printStdOut = printStdOut;
+		_printStdErr = printStdErr;
 		TempDir = GetTemporaryDirectory();
 	}
 
diff --git a/Engine/Source/Programs/Horde/Horde.Build.Tests/LeasesControllerDbTest.cs b/Engine/Source/Programs/Horde/Horde.Build.Tests/LeasesControllerDbTest.cs
--- a/Engine/Source/Programs/Horde/Horde.Build.Tests/LeasesControllerDbTest.cs
+++ b/Engine/Source/Programs/Horde/Horde.Build.Tests/LeasesControllerDbTest.cs
@@ -26,7 +26,7 @@ namespace Horde.Build.Tests
 
 			ILease lease1 = await CreateLease(Clock.UtcNow - TimeSpan.FromMinutes(10), TimeSpan.FromMinutes(6));
 			ILease lease2 = await CreateLease(Clock.UtcNow - TimeSpan.FromMinutes(7), TimeSpan.FromMinutes(3.1));
-			ILease outOfTimeWindow = await CreateLease(Clock.UtcNow - TimeSpan.FromMinutes(7), TimeSpan.FromMinutes(25));
+			/*ILease outOfTimeWindow = */await CreateLease(Clock.UtcNow - TimeSpan.FromMinutes(7), TimeSpan.FromMinutes(25));
 			
 			ActionResult<List<object>> res = await LeasesController.FindLeasesAsync(null, null, null, null, minTime, maxTime);
 			Assert.AreEqual(2, res.Value!.Count);
diff --git a/Engine/Source/Programs/Horde/Horde.Build.Tests/LogFileServiceTests.cs b/Engine/Source/Programs/Horde/Horde.Build.Tests/LogFileServiceTests.cs
--- a/Engine/Source/Programs/Horde/Horde.Build.Tests/LogFileServiceTests.cs
+++ b/Engine/Source/Programs/Horde/Horde.Build.Tests/LogFileServiceTests.cs
@@ -309,7 +309,7 @@ namespace Horde.Build.Tests
 
 			byte[] line5 = Encoding.UTF8.GetBytes("a\nb\n");
 			await _logFileService.WriteLogDataAsync(logFile, offset, 4, line5, false, MaxChunkSize, MaxSubChunkLineCount);
-			offset += line5.Length;
+			// offset += line5.Length;
 
 			await _logFileService.FlushAsync();
 			logFile = (await _logFileService.GetLogFileAsync(logFile.Id))!;
diff --git a/Engine/Source/Programs/Horde/Horde.Build.Tests/LogIndexingTests.cs b/Engine/Source/Programs/Horde/Horde.Build.Tests/LogIndexingTests.cs
--- a/Engine/Source/Programs/Horde/Horde.Build.Tests/LogIndexingTests.cs
+++ b/Engine/Source/Programs/Horde/Horde.Build.Tests/LogIndexingTests.cs
@@ -32,7 +32,7 @@ namespace Horde.Build.Tests
 		private readonly ILogFileService _logFileService;
 		private readonly NullLogStorage _nullLogStorage;
 		private readonly LocalLogStorage _logStorage;
-		private byte[] _data = Resources.TextFile;
+		private readonly byte[] _data = Resources.TextFile;
 
 		public LogIndexingTests()
 		{
diff --git a/Engine/Source/Programs/Horde/Horde.Build.Tests/NotificationServiceTests.cs b/Engine/Source/Programs/Horde/Horde.Build.Tests/NotificationServiceTests.cs
--- a/Engine/Source/Programs/Horde/Horde.Build.Tests/NotificationServiceTests.cs
+++ b/Engine/Source/Programs/Horde/Horde.Build.Tests/NotificationServiceTests.cs
@@ -10,15 +10,13 @@ using Horde.Build.Notifications;
 using Horde.Build.Notifications.Impl;
 using Horde.Build.Utilities;
 using Microsoft.Extensions.DependencyInjection;
-using Microsoft.Extensions.Logging;
 using Microsoft.VisualStudio.TestTools.UnitTesting;
 using Moq;
-using StackExchange.Redis;
-using StreamId = Horde.Build.Utilities.StringId<Horde.Build.Models.IStream>;
 
 namespace Horde.Build.Tests
 {
 	using JobId = ObjectId<IJob>;
+	using StreamId = StringId<IStream>;
 
 	public class FakeNotificationSink : INotificationSink
 	{
@@ -49,7 +47,7 @@ namespace Horde.Build.Tests
 			base.ConfigureServices(services);
 
 			services.AddSingleton<FakeNotificationSink>();
-			services.AddSingleton<INotificationSink>(SP => SP.GetRequiredService<FakeNotificationSink>());
+			services.AddSingleton<INotificationSink>(sp => sp.GetRequiredService<FakeNotificationSink>());
 		}
 
 		public static IJob CreateJob(StreamId streamId, int change, string name, IGraph graph)
diff --git a/Engine/Source/Programs/Horde/Horde.Build.Tests/RpcServiceTest.cs b/Engine/Source/Programs/Horde/Horde.Build.Tests/RpcServiceTest.cs
--- a/Engine/Source/Programs/Horde/Horde.Build.Tests/RpcServiceTest.cs
+++ b/Engine/Source/Programs/Horde/Horde.Build.Tests/RpcServiceTest.cs
@@ -23,7 +23,6 @@ using AgentCapabilities = HordeCommon.Rpc.Messages.AgentCapabilities;
 using ISession = Microsoft.AspNetCore.Http.ISession;
 using PoolId = Horde.Build.Utilities.StringId<Horde.Build.Models.IPool>;
 
-
 namespace Horde.Build.Tests
 {
 	using LeaseId = ObjectId<ILease>;
@@ -77,7 +76,7 @@ namespace Horde.Build.Tests
 			
 			public HttpContextStub(ClaimsPrincipal user)
 			{
-				this.User = user;
+				User = user;
 			}
 
 			public override void Abort()
@@ -151,8 +150,8 @@ namespace Horde.Build.Tests
 
 			public RpcServiceInvoker(RpcService rpcService, ServerCallContext serverCallContext)
 			{
-				this._rpcService = rpcService;
-				this._serverCallContext = serverCallContext;
+				_rpcService = rpcService;
+				_serverCallContext = serverCallContext;
 			}
 
 			public override TResponse BlockingUnaryCall<TRequest, TResponse>(Method<TRequest, TResponse> method, string host, CallOptions options, TRequest request)
@@ -384,7 +383,6 @@ namespace Horde.Build.Tests
 				new Claim(HordeClaimTypes.AgentSessionId, sessionId.ToString()),
 			}, "TestAuthType")));
 
-
 			string[] data = {"foo", "bar", "baz", "qux"};
 			string dataStr = String.Join("", data);
 			
@@ -423,7 +421,6 @@ namespace Horde.Build.Tests
 			requestStream.AddMessage(new UploadArtifactRequest { Data = ByteString.CopyFromUtf8(data[3]) });
 			UploadArtifactResponse res = await call;
 
-
 			IArtifact? artifact = await ArtifactCollection.GetArtifactAsync(ObjectId.Parse(res.Id));
 			Assert.IsNotNull(artifact);
 			Stream stream = await ArtifactCollection.OpenArtifactReadStreamAsync(artifact!);
diff --git a/Engine/Source/Programs/Horde/Horde.Build.Tests/SchedulerTests.cs b/Engine/Source/Programs/Horde/Horde.Build.Tests/SchedulerTests.cs
--- a/Engine/Source/Programs/Horde/Horde.Build.Tests/SchedulerTests.cs
+++ b/Engine/Source/Programs/Horde/Horde.Build.Tests/SchedulerTests.cs
@@ -28,8 +28,8 @@ namespace Horde.Build.Tests
 		StreamId StreamId { get; } = new StreamId("ue5-main");
 		TemplateRefId TemplateRefId { get; } = new TemplateRefId("template1");
 
-		ITemplate _template;
-		HashSet<JobId> _initialJobIds;
+		readonly ITemplate _template;
+		readonly HashSet<JobId> _initialJobIds;
 
 		public SchedulerTests()
 		{
@@ -379,7 +379,7 @@ namespace Horde.Build.Tests
 			CreateScheduleRequest schedule = new CreateScheduleRequest();
 			schedule.Enabled = true;
 			schedule.Patterns.Add(new CreateSchedulePatternRequest { MinTime = 13 * 60, MaxTime = 14 * 60, Interval = 15 });
-			IStream stream = await SetScheduleAsync(schedule);
+			/*IStream stream = */await SetScheduleAsync(schedule);
 
 			// Trigger a job
 			await Clock.AdvanceAsync(TimeSpan.FromHours(1.25));
@@ -406,7 +406,7 @@ namespace Horde.Build.Tests
 
 			// Create two templates, the second dependent on the first
 			ITemplate? newTemplate1 = await TemplateCollection.AddAsync("Test template 1");
-			TemplateRef newTemplateRef1 = new TemplateRef(newTemplate1);
+			//TemplateRef newTemplateRef1 = new TemplateRef(newTemplate1);
 			TemplateRefId newTemplateRefId1 = new TemplateRefId("new-template-1");
 
 			ITemplate? newTemplate2 = await TemplateCollection.AddAsync("Test template 2");
@@ -442,6 +442,7 @@ namespace Horde.Build.Tests
 			job1 = Deref(await JobService.UpdateBatchAsync(job1, batchId1, LogId.GenerateNewId(), JobStepBatchState.Running));
 			job1 = Deref(await JobService.UpdateStepAsync(job1, batchId1, stepId1, JobStepState.Completed, JobStepOutcome.Failure));
 			job1 = Deref(await JobService.UpdateBatchAsync(job1, batchId1, LogId.GenerateNewId(), JobStepBatchState.Complete));
+			Assert.IsNotNull(job1);
 			await GetNewJobs();
 
 			// Tick the schedule and make sure it doesn't trigger
@@ -456,6 +457,7 @@ namespace Horde.Build.Tests
 			SubResourceId stepId2 = job2.Batches[0].Steps[0].Id;
 			job2 = Deref(await JobService.UpdateBatchAsync(job2, batchId2, LogId.GenerateNewId(), JobStepBatchState.Running));
 			job2 = Deref(await JobService.UpdateStepAsync(job2, batchId2, stepId2, JobStepState.Completed, JobStepOutcome.Success));
+			Assert.IsNotNull(job2);
 
 			// Tick the schedule and make sure it does trigger
 			await ScheduleService.TickForTestingAsync();
@@ -506,7 +508,9 @@ namespace Horde.Build.Tests
 			stream = (await StreamService.StreamCollection.TryCreateOrReplaceAsync(StreamId, stream, "", "", ProjectId, config))!;
 
 			ITemplate template1 = (await TemplateCollection.GetAsync(stream.Templates[newTemplateRefId1].Hash))!;
+			Assert.IsNotNull(template1);
 			ITemplate template2 = (await TemplateCollection.GetAsync(stream.Templates[newTemplateRefId2].Hash))!;
+			Assert.IsNotNull(template2);
 
 			// Create the graph
 			IGraph graphA = await GraphCollection.AddAsync(template1);
diff --git a/Engine/Source/Programs/Horde/Horde.Build.Tests/ServiceAccountAuthTest.cs b/Engine/Source/Programs/Horde/Horde.Build.Tests/ServiceAccountAuthTest.cs
--- a/Engine/Source/Programs/Horde/Horde.Build.Tests/ServiceAccountAuthTest.cs
+++ b/Engine/Source/Programs/Horde/Horde.Build.Tests/ServiceAccountAuthTest.cs
@@ -7,7 +7,6 @@ using System.Threading.Tasks;
 using Horde.Build.Authentication;
 using Horde.Build.Collections;
 using Horde.Build.Collections.Impl;
-using Horde.Build.Services;
 using Microsoft.AspNetCore.Authentication;
 using Microsoft.AspNetCore.Http;
 using Microsoft.Extensions.DependencyInjection;
@@ -32,7 +31,7 @@ namespace Horde.Build.Tests
 
 			services.AddSingleton<ILoggerFactory, LoggerFactory>();
 			services.AddSingleton<ServiceAccountCollection>();
-			services.AddSingleton<IServiceAccountCollection>(SP => SP.GetRequiredService<ServiceAccountCollection>());
+			services.AddSingleton<IServiceAccountCollection>(sp => sp.GetRequiredService<ServiceAccountCollection>());
 		}
 
 		private async Task<ServiceAccountAuthHandler> GetAuthHandlerAsync(string? headerValue)
diff --git a/Engine/Source/Programs/Horde/Horde.Build.Tests/Stubs/Collections/JobStepRefCollectionStub.cs b/Engine/Source/Programs/Horde/Horde.Build.Tests/Stubs/Collections/JobStepRefCollectionStub.cs
--- a/Engine/Source/Programs/Horde/Horde.Build.Tests/Stubs/Collections/JobStepRefCollectionStub.cs
+++ b/Engine/Source/Programs/Horde/Horde.Build.Tests/Stubs/Collections/JobStepRefCollectionStub.cs
@@ -40,18 +40,18 @@ namespace Horde.Build.Tests.Stubs.Collections
 		public JobStepRefStub(JobId jobId, SubResourceId batchId, SubResourceId stepId, string jobName, string nodeName, StreamId streamId, TemplateRefId templateId, int change, JobStepOutcome? outcome)
 		{
 			Id = new JobStepRefId(jobId, batchId, stepId);
-			this.JobName = jobName;
-			this.NodeName = nodeName;
-			this.StreamId = streamId;
-			this.TemplateId = templateId;
-			this.Change = change;
-			this.Outcome = outcome;
+			JobName = jobName;
+			NodeName = nodeName;
+			StreamId = streamId;
+			TemplateId = templateId;
+			Change = change;
+			Outcome = outcome;
 		}
 	}
 
 	class JobStepRefCollectionStub : IJobStepRefCollection
 	{
-		List<IJobStepRef> _refs = new List<IJobStepRef>();
+		readonly List<IJobStepRef> _refs = new List<IJobStepRef>();
 
 		public void Add(IJobStepRef jobStepRef)
 		{
diff --git a/Engine/Source/Programs/Horde/Horde.Build.Tests/Stubs/Services/PerforceServiceStub.cs b/Engine/Source/Programs/Horde/Horde.Build.Tests/Stubs/Services/PerforceServiceStub.cs
--- a/Engine/Source/Programs/Horde/Horde.Build.Tests/Stubs/Services/PerforceServiceStub.cs
+++ b/Engine/Source/Programs/Horde/Horde.Build.Tests/Stubs/Services/PerforceServiceStub.cs
@@ -45,7 +45,7 @@ namespace Horde.Build.Tests.Stubs.Services
 			}
 		}
 
-		User _testUser = new User("TestUser");
+		readonly User _testUser = new User("TestUser");
 
 		class ChangeComparer : IComparer<int>
 		{
@@ -54,11 +54,11 @@ namespace Horde.Build.Tests.Stubs.Services
 
 		public Dictionary<string, SortedDictionary<int, ChangeDetails>> Changes { get; } = new Dictionary<string, SortedDictionary<int, ChangeDetails>>(StringComparer.OrdinalIgnoreCase);
 
-		IUserCollection _userCollection;
+		readonly IUserCollection _userCollection;
 
 		public PerforceServiceStub(IUserCollection userCollection)
 		{
-			this._userCollection = userCollection;
+			_userCollection = userCollection;
 		}
 
 		public Task<IPerforceConnection?> GetServiceUserConnection(string? clusterName)
diff --git a/Engine/Source/Programs/Horde/Horde.Build.Tests/TestSetup.cs b/Engine/Source/Programs/Horde/Horde.Build.Tests/TestSetup.cs
--- a/Engine/Source/Programs/Horde/Horde.Build.Tests/TestSetup.cs
+++ b/Engine/Source/Programs/Horde/Horde.Build.Tests/TestSetup.cs
@@ -36,7 +36,6 @@ using Microsoft.Extensions.DependencyInjection;
 using Microsoft.Extensions.Hosting;
 using Microsoft.Extensions.Logging;
 using Microsoft.Extensions.Options;
-using Moq;
 using StatsdClient;
 
 namespace Horde.Build.Tests
@@ -159,7 +158,7 @@ namespace Horde.Build.Tests
 			services.AddSingleton<IUserCollection, UserCollectionV1>();
 
 			services.AddSingleton<FakeClock>();
-			services.AddSingleton<IClock>(SP => SP.GetRequiredService<FakeClock>());
+			services.AddSingleton<IClock>(sp => sp.GetRequiredService<FakeClock>());
 			services.AddSingleton<IHostApplicationLifetime, AppLifetimeStub>();
 
 			services.AddSingleton<AclService>();
diff --git a/Engine/Source/Programs/Horde/Horde.Build.Tests/VersionedCollectionTests.cs b/Engine/Source/Programs/Horde/Horde.Build.Tests/VersionedCollectionTests.cs
--- a/Engine/Source/Programs/Horde/Horde.Build.Tests/VersionedCollectionTests.cs
+++ b/Engine/Source/Programs/Horde/Horde.Build.Tests/VersionedCollectionTests.cs
@@ -46,12 +46,11 @@ namespace Horde.Build.Tests
 			public override DocumentV2 UpgradeToLatest() => this;
 		}
 
-
-		DatabaseService _databaseService;
-		IDatabase _redis;
-		RedisKey _baseKey = new RedisKey("versioned/");
-		IMongoCollection<VersionedDocument<string, DocumentV2>> _baseCollection;
-		VersionedCollection<string, DocumentV2> _collection;
+		readonly DatabaseService _databaseService;
+		readonly IDatabase _redis;
+		readonly RedisKey _baseKey = new RedisKey("versioned/");
+		readonly IMongoCollection<VersionedDocument<string, DocumentV2>> _baseCollection;
+		readonly VersionedCollection<string, DocumentV2> _collection;
 
 		public VersionedCollectionTests()
 		{
