Horde Storage - Fixed build error in Callisto due to warning about md5. Its fine to use md5 as we do not use it for encryption but as a hash.

diff --git a/Engine/Source/Programs/Horde/HordeStorage/.editorconfig b/Engine/Source/Programs/Horde/HordeStorage/.editorconfig
--- a/Engine/Source/Programs/Horde/HordeStorage/.editorconfig
+++ b/Engine/Source/Programs/Horde/HordeStorage/.editorconfig
@@ -232,3 +232,6 @@ dotnet_style_prefer_compound_assignment = false:none
 
 # csharp_style_allow_blank_line_after_colon_in_constructor_initializer_experimental
 dotnet_diagnostic.IDE2004.severity = none
+
+# Do not generate build warning for broken cryptographic algorithms (MD5 in this case) we do not use it for encryption but as a hash
+dotnet_diagnostic.CA5351.severity = suggestion
