Horde Storage - Ignore case when attempting to resolve peer list to region that we can replicate from

diff --git a/Engine/Source/Programs/Horde/HordeStorage/Horde.Storage/Implementation/PeerStatusService.cs b/Engine/Source/Programs/Horde/HordeStorage/Horde.Storage/Implementation/PeerStatusService.cs
--- a/Engine/Source/Programs/Horde/HordeStorage/Horde.Storage/Implementation/PeerStatusService.cs
+++ b/Engine/Source/Programs/Horde/HordeStorage/Horde.Storage/Implementation/PeerStatusService.cs
@@ -35,7 +35,7 @@ namespace Horde.Storage.Implementation
         private readonly IOptionsMonitor<ClusterSettings> _clusterSettings;
         private readonly IOptionsMonitor<JupiterSettings> _jupiterSettings;
         private readonly IHttpClientFactory _clientFactory;
-        private readonly Dictionary<string, IPeerStatusService.PeerStatus> _peers = new Dictionary<string, IPeerStatusService.PeerStatus>();
+        private readonly Dictionary<string, IPeerStatusService.PeerStatus> _peers = new Dictionary<string, IPeerStatusService.PeerStatus>(StringComparer.InvariantCultureIgnoreCase);
         private volatile bool _alreadyPolling = false;
 
         public class PeerStatusServiceState
