Horde Storage - Fixed issue were the entire blob was removed from the index when it was being deleted rather then just removing the tracking for that region.
Also made sure we do not remove the blob unless we were successful in deleting the blob in the first place.

diff --git a/Engine/Source/Programs/Horde/HordeStorage/Horde.Storage/Implementation/Blob/IBlobService.cs b/Engine/Source/Programs/Horde/HordeStorage/Horde.Storage/Implementation/Blob/IBlobService.cs
--- a/Engine/Source/Programs/Horde/HordeStorage/Horde.Storage/Implementation/Blob/IBlobService.cs
+++ b/Engine/Source/Programs/Horde/HordeStorage/Horde.Storage/Implementation/Blob/IBlobService.cs
@@ -382,7 +382,6 @@ public class BlobService : IBlobService
     {
         bool blobNotFound = false;
         bool deletedAtLeastOnce = false;
-        Task removeFromIndexTask = _blobIndex.RemoveBlobFromIndex(ns, blob);
 
         foreach (IBlobStore store in _blobStores)
         {
@@ -404,7 +403,7 @@ public class BlobService : IBlobService
             }
         }
 
-        await removeFromIndexTask;
+        await _blobIndex.RemoveBlobFromRegion(ns, blob);
 
         if (deletedAtLeastOnce)
             return;
diff --git a/Engine/Source/Programs/Horde/HordeStorage/Horde.Storage/Implementation/GC/OrphanBlobCleanupRefs.cs b/Engine/Source/Programs/Horde/HordeStorage/Horde.Storage/Implementation/GC/OrphanBlobCleanupRefs.cs
--- a/Engine/Source/Programs/Horde/HordeStorage/Horde.Storage/Implementation/GC/OrphanBlobCleanupRefs.cs
+++ b/Engine/Source/Programs/Horde/HordeStorage/Horde.Storage/Implementation/GC/OrphanBlobCleanupRefs.cs
@@ -88,6 +88,9 @@ namespace Horde.Storage.Implementation
                     // check all other namespaces that share the same storage pool for presence of the blob
                     foreach (NamespaceId blobNamespace in namespaces.Where(ns => _namespacePolicyResolver.GetPoliciesForNs(ns).StoragePool == policy.StoragePool))
                     {
+                        if (cancellationToken.IsCancellationRequested)
+                            break;
+                        
                         IBlobIndex.BlobInfo? blobIndex = await _blobIndex.GetBlobInfo(blobNamespace, blob);
 
                         if (blobIndex == null)
