Horde: Sort using directives on Horde.Agent.

diff --git a/Engine/Source/Programs/Horde/Horde.Agent/Commands/Compute/ComputeCommand.cs b/Engine/Source/Programs/Horde/Horde.Agent/Commands/Compute/ComputeCommand.cs
--- a/Engine/Source/Programs/Horde/Horde.Agent/Commands/Compute/ComputeCommand.cs
+++ b/Engine/Source/Programs/Horde/Horde.Agent/Commands/Compute/ComputeCommand.cs
@@ -1,7 +1,5 @@
 // Copyright Epic Games, Inc. All Rights Reserved.
 
-using Google.Protobuf;
-using Microsoft.Extensions.Logging;
 using System;
 using System.Collections.Generic;
 using System.IO;
@@ -10,13 +8,15 @@ using System.Text;
 using System.Threading;
 using System.Threading.Tasks;
 using EpicGames.Core;
-using Microsoft.Extensions.Hosting;
-using Microsoft.Extensions.DependencyInjection;
-using Microsoft.Extensions.Configuration;
-using Serilog.Events;
-using EpicGames.Serialization;
 using EpicGames.Horde.Compute;
 using EpicGames.Horde.Storage;
+using EpicGames.Serialization;
+using Google.Protobuf;
+using Microsoft.Extensions.Configuration;
+using Microsoft.Extensions.DependencyInjection;
+using Microsoft.Extensions.Hosting;
+using Microsoft.Extensions.Logging;
+using Serilog.Events;
 
 namespace Horde.Agent.Commands
 {
diff --git a/Engine/Source/Programs/Horde/Horde.Agent/Commands/Service/RunCommand.cs b/Engine/Source/Programs/Horde/Horde.Agent/Commands/Service/RunCommand.cs
--- a/Engine/Source/Programs/Horde/Horde.Agent/Commands/Service/RunCommand.cs
+++ b/Engine/Source/Programs/Horde/Horde.Agent/Commands/Service/RunCommand.cs
@@ -1,22 +1,22 @@
 // Copyright Epic Games, Inc. All Rights Reserved.
 
+using System;
+using System.Collections.Generic;
+using System.Net.Http;
+using System.Threading.Tasks;
+using Datadog.Trace;
+using Datadog.Trace.Configuration;
+using EpicGames.Core;
+using EpicGames.Horde.Storage;
 using Horde.Agent.Services;
 using Horde.Agent.Utility;
 using Microsoft.Extensions.Configuration;
 using Microsoft.Extensions.DependencyInjection;
 using Microsoft.Extensions.Hosting;
 using Microsoft.Extensions.Logging;
-using Polly;
-using System;
-using System.Collections.Generic;
-using System.Net.Http;
-using System.Threading.Tasks;
-using EpicGames.Core;
 using OpenTracing;
 using OpenTracing.Util;
-using Datadog.Trace.Configuration;
-using Datadog.Trace;
-using EpicGames.Horde.Storage;
+using Polly;
 
 namespace Horde.Agent.Modes.Service
 {
diff --git a/Engine/Source/Programs/Horde/Horde.Agent/Commands/Service/UpgradeCommand.cs b/Engine/Source/Programs/Horde/Horde.Agent/Commands/Service/UpgradeCommand.cs
--- a/Engine/Source/Programs/Horde/Horde.Agent/Commands/Service/UpgradeCommand.cs
+++ b/Engine/Source/Programs/Horde/Horde.Agent/Commands/Service/UpgradeCommand.cs
@@ -1,6 +1,5 @@
 // Copyright Epic Games, Inc. All Rights Reserved.
 
-using Microsoft.Extensions.Logging;
 using System;
 using System.Collections.Generic;
 using System.Diagnostics;
@@ -8,11 +7,12 @@ using System.IO;
 using System.Linq;
 using System.Management;
 using System.Reflection;
+using System.Runtime.InteropServices;
 using System.ServiceProcess;
 using System.Text;
 using System.Threading.Tasks;
 using EpicGames.Core;
-using System.Runtime.InteropServices;
+using Microsoft.Extensions.Logging;
 
 namespace Horde.Agent.Commands
 {
diff --git a/Engine/Source/Programs/Horde/Horde.Agent/Commands/Utilities/CapsCommand.cs b/Engine/Source/Programs/Horde/Horde.Agent/Commands/Utilities/CapsCommand.cs
--- a/Engine/Source/Programs/Horde/Horde.Agent/Commands/Utilities/CapsCommand.cs
+++ b/Engine/Source/Programs/Horde/Horde.Agent/Commands/Utilities/CapsCommand.cs
@@ -1,10 +1,10 @@
 // Copyright Epic Games, Inc. All Rights Reserved.
 
-using Microsoft.Extensions.Logging;
 using System.Threading.Tasks;
 using EpicGames.Core;
-using HordeCommon.Rpc.Messages;
 using Horde.Agent.Services;
+using HordeCommon.Rpc.Messages;
+using Microsoft.Extensions.Logging;
 
 namespace Horde.Agent.Commands.Utilities
 {
diff --git a/Engine/Source/Programs/Horde/Horde.Agent/Commands/Workspace/CleanCommand.cs b/Engine/Source/Programs/Horde/Horde.Agent/Commands/Workspace/CleanCommand.cs
--- a/Engine/Source/Programs/Horde/Horde.Agent/Commands/Workspace/CleanCommand.cs
+++ b/Engine/Source/Programs/Horde/Horde.Agent/Commands/Workspace/CleanCommand.cs
@@ -3,10 +3,10 @@
 using System.ComponentModel;
 using System.Threading;
 using System.Threading.Tasks;
-using Microsoft.Extensions.Logging;
 using EpicGames.Core;
-using EpicGames.Perforce.Managed;
 using EpicGames.Perforce;
+using EpicGames.Perforce.Managed;
+using Microsoft.Extensions.Logging;
 
 namespace Horde.Agent.Commands.Workspace
 {
diff --git a/Engine/Source/Programs/Horde/Horde.Agent/Execution/ComputeTaskExecutor.cs b/Engine/Source/Programs/Horde/Horde.Agent/Execution/ComputeTaskExecutor.cs
--- a/Engine/Source/Programs/Horde/Horde.Agent/Execution/ComputeTaskExecutor.cs
+++ b/Engine/Source/Programs/Horde/Horde.Agent/Execution/ComputeTaskExecutor.cs
@@ -1,20 +1,20 @@
 // Copyright Epic Games, Inc. All Rights Reserved.
 
-using EpicGames.Core;
-using EpicGames.Horde.Storage;
-using Grpc.Core;
 using System;
 using System.Collections.Generic;
+using System.Diagnostics;
 using System.IO;
 using System.Linq;
+using System.Text;
 using System.Threading;
 using System.Threading.Tasks;
-using System.Diagnostics;
-using System.Text;
+using EpicGames.Core;
+using EpicGames.Horde.Compute;
+using EpicGames.Horde.Storage;
+using EpicGames.Serialization;
+using Grpc.Core;
 using HordeCommon.Rpc.Tasks;
 using Microsoft.Extensions.Logging;
-using EpicGames.Serialization;
-using EpicGames.Horde.Compute;
 
 namespace Horde.Agent.Execution
 {
diff --git a/Engine/Source/Programs/Horde/Horde.Agent/Logging.cs b/Engine/Source/Programs/Horde/Horde.Agent/Logging.cs
--- a/Engine/Source/Programs/Horde/Horde.Agent/Logging.cs
+++ b/Engine/Source/Programs/Horde/Horde.Agent/Logging.cs
@@ -1,17 +1,17 @@
 // Copyright Epic Games, Inc. All Rights Reserved.
 
+using System;
+using System.Runtime.InteropServices;
+using EpicGames.Core;
 using Microsoft.Extensions.Logging;
+using OpenTracing;
+using OpenTracing.Util;
 using Serilog;
 using Serilog.Core;
 using Serilog.Events;
 using Serilog.Extensions.Logging;
 using Serilog.Formatting.Json;
 using Serilog.Sinks.SystemConsole.Themes;
-using System;
-using System.Runtime.InteropServices;
-using EpicGames.Core;
-using OpenTracing.Util;
-using OpenTracing;
 
 namespace Horde.Agent
 {
diff --git a/Engine/Source/Programs/Horde/Horde.Agent/Services/WorkerService.cs b/Engine/Source/Programs/Horde/Horde.Agent/Services/WorkerService.cs
--- a/Engine/Source/Programs/Horde/Horde.Agent/Services/WorkerService.cs
+++ b/Engine/Source/Programs/Horde/Horde.Agent/Services/WorkerService.cs
@@ -1,9 +1,27 @@
 // Copyright Epic Games, Inc. All Rights Reserved.
 
+using System;
+using System.Collections.Generic;
+using System.Diagnostics;
+using System.Diagnostics.CodeAnalysis;
+using System.IO;
+using System.IO.Compression;
+using System.Linq;
+using System.Management;
+using System.Net;
+using System.Reflection;
+using System.Runtime.InteropServices;
+using System.Text;
+using System.Text.RegularExpressions;
+using System.Threading;
+using System.Threading.Tasks;
+using System.Xml;
 using Amazon.EC2;
 using Amazon.EC2.Model;
 using Amazon.Util;
 using EpicGames.Core;
+using EpicGames.Horde.Compute;
+using EpicGames.Horde.Storage;
 using Google.Protobuf;
 using Google.Protobuf.WellKnownTypes;
 using Grpc.Core;
@@ -21,24 +39,6 @@ using Microsoft.Extensions.Logging;
 using Microsoft.Extensions.Options;
 using OpenTracing;
 using OpenTracing.Util;
-using System;
-using System.Collections.Generic;
-using System.Diagnostics;
-using System.Diagnostics.CodeAnalysis;
-using System.IO;
-using System.IO.Compression;
-using System.Linq;
-using System.Management;
-using System.Reflection;
-using System.Runtime.InteropServices;
-using System.Text;
-using System.Threading;
-using System.Threading.Tasks;
-using System.Xml;
-using EpicGames.Horde.Storage;
-using EpicGames.Horde.Compute;
-using System.Net;
-using System.Text.RegularExpressions;
 
 namespace Horde.Agent.Services
 {
diff --git a/Engine/Source/Programs/Horde/Horde.Agent/Utility/WorkspaceInfo.cs b/Engine/Source/Programs/Horde/Horde.Agent/Utility/WorkspaceInfo.cs
--- a/Engine/Source/Programs/Horde/Horde.Agent/Utility/WorkspaceInfo.cs
+++ b/Engine/Source/Programs/Horde/Horde.Agent/Utility/WorkspaceInfo.cs
@@ -1,16 +1,16 @@
 // Copyright Epic Games, Inc. All Rights Reserved.
 
-using EpicGames.Perforce.Managed;
-using HordeCommon.Rpc.Messages;
-using Microsoft.Extensions.Logging;
 using System;
 using System.Collections.Generic;
+using System.Linq;
 using System.Text.RegularExpressions;
 using System.Threading;
 using System.Threading.Tasks;
 using EpicGames.Core;
 using EpicGames.Perforce;
-using System.Linq;
+using EpicGames.Perforce.Managed;
+using HordeCommon.Rpc.Messages;
+using Microsoft.Extensions.Logging;
 
 namespace Horde.Agent.Utility
 {
