Horde Storage - Updated assumption in tests when removing objects from the blob inded (deleting a object removes it from the current regions tracking but does not remove the actual tracking as it may still exist in other regions).

diff --git a/Engine/Source/Programs/Horde/HordeStorage/Horde.Storage/Tests/Functional/Blob/BlobIndexTests.cs b/Engine/Source/Programs/Horde/HordeStorage/Horde.Storage/Tests/Functional/Blob/BlobIndexTests.cs
--- a/Engine/Source/Programs/Horde/HordeStorage/Horde.Storage/Tests/Functional/Blob/BlobIndexTests.cs
+++ b/Engine/Source/Programs/Horde/HordeStorage/Horde.Storage/Tests/Functional/Blob/BlobIndexTests.cs
@@ -143,6 +143,7 @@ namespace Horde.Storage.FunctionalTests.Storage
             Assert.IsNotNull(index);
             IBlobIndex.BlobInfo? blobInfo = await index.GetBlobInfo(TestNamespaceName, contentHash);
             Assert.IsNotNull(blobInfo);
+            Assert.IsTrue(blobInfo.Regions.Any());
 
             // delete the blob
             {
@@ -151,7 +152,10 @@ namespace Horde.Storage.FunctionalTests.Storage
             }
 
             IBlobIndex.BlobInfo? deletedBlobInfo = await index.GetBlobInfo(TestNamespaceName, contentHash);
-            Assert.IsNull(deletedBlobInfo);
+            // the blob info will still exist after delete
+            Assert.IsNotNull(deletedBlobInfo);
+            // but the blob info will not contain the current region
+            Assert.IsTrue(!deletedBlobInfo.Regions.Any());
         }
 
 
@@ -205,7 +209,7 @@ namespace Horde.Storage.FunctionalTests.Storage
      
     }
 
-    [TestClass]
+    [TestClass()]
     public class MemoryBlobIndexTests : BlobIndexTests
     {
         protected override IEnumerable<KeyValuePair<string, string>> GetSettings()
