Horde Storage - In scylla a row that only has the partition key set but no other data will actually be considered a valid row, so updating the tests assumptions for this when the blob index is emptied.

diff --git a/Engine/Source/Programs/Horde/HordeStorage/Horde.Storage/Tests/Functional/Blob/BlobIndexTests.cs b/Engine/Source/Programs/Horde/HordeStorage/Horde.Storage/Tests/Functional/Blob/BlobIndexTests.cs
--- a/Engine/Source/Programs/Horde/HordeStorage/Horde.Storage/Tests/Functional/Blob/BlobIndexTests.cs
+++ b/Engine/Source/Programs/Horde/HordeStorage/Horde.Storage/Tests/Functional/Blob/BlobIndexTests.cs
@@ -152,10 +152,10 @@ namespace Horde.Storage.FunctionalTests.Storage
             }
 
             IBlobIndex.BlobInfo? deletedBlobInfo = await index.GetBlobInfo(TestNamespaceName, contentHash);
-            // the blob info will still exist after delete
-            Assert.IsNotNull(deletedBlobInfo);
+
+            bool hasRegions = deletedBlobInfo != null && deletedBlobInfo.Regions.Any();
             // but the blob info will not contain the current region
-            Assert.IsTrue(!deletedBlobInfo.Regions.Any());
+            Assert.IsTrue(!hasRegions);
         }
 
 
