Horde: Fix coding conventions within Horde.Build.

diff --git a/Engine/Source/Programs/Horde/Horde.Build.Tests/AgentServiceTest.cs b/Engine/Source/Programs/Horde/Horde.Build.Tests/AgentServiceTest.cs
--- a/Engine/Source/Programs/Horde/Horde.Build.Tests/AgentServiceTest.cs
+++ b/Engine/Source/Programs/Horde/Horde.Build.Tests/AgentServiceTest.cs
@@ -29,7 +29,7 @@ public class AgentServiceTest : TestSetup
 		Assert.AreEqual("hello2", (res.Value[0] as GetJobResponse)!.Name);
 		Assert.AreEqual("hello1", (res.Value[1] as GetJobResponse)!.Name);
 
-		res = await JobsController.FindJobsAsync(IncludePreflight: false);
+		res = await JobsController.FindJobsAsync(includePreflight: false);
 		Assert.AreEqual(1, res.Value!.Count);
 		Assert.AreEqual("hello2", (res.Value[0] as GetJobResponse)!.Name);
 	}
diff --git a/Engine/Source/Programs/Horde/Horde.Build.Tests/Fixture.cs b/Engine/Source/Programs/Horde/Horde.Build.Tests/Fixture.cs
--- a/Engine/Source/Programs/Horde/Horde.Build.Tests/Fixture.cs
+++ b/Engine/Source/Programs/Horde/Horde.Build.Tests/Fixture.cs
@@ -5,20 +5,20 @@ using System.IO;
 using System.Text;
 using System.Threading.Tasks;
 using EpicGames.Core;
-using HordeCommon;
+using Horde.Build.Api;
 using Horde.Build.Collections;
+using Horde.Build.Jobs;
 using Horde.Build.Models;
 using Horde.Build.Services;
-using StreamId = Horde.Build.Utilities.StringId<Horde.Build.Models.IStream>;
-using TemplateRefId = Horde.Build.Utilities.StringId<Horde.Build.Models.TemplateRef>;
-using ProjectId = Horde.Build.Utilities.StringId<Horde.Build.Models.IProject>;
-using Horde.Build.Api;
-using Horde.Build.Jobs;
 using Horde.Build.Utilities;
+using HordeCommon;
 
 namespace Horde.Build.Tests
 {
 	using JobId = ObjectId<IJob>;
+    using ProjectId = StringId<IProject>;
+    using StreamId = StringId<IStream>;
+    using TemplateRefId = StringId<TemplateRef>;
 
 	public class Fixture
 	{
@@ -84,18 +84,18 @@ namespace Horde.Build.Tests
 			);
 			
 			Job1 = await jobService.CreateJobAsync(
-				JobId: new JobId("5f283932841e7fdbcafb6ab5"),
-				Stream: Stream!,
-				TemplateRefId: TemplateRefId1,
-				TemplateHash: Template.Id,
-				Graph: Graph,
-				Name: "hello1",
-				Change: 1000001,
-				CodeChange: 1000002,
-				PreflightChange: 1001,
-				ClonedPreflightChange: null,
-				StartedByUserId: null,
-				Priority: Priority.Normal,
+				jobId: new JobId("5f283932841e7fdbcafb6ab5"),
+				stream: Stream!,
+				templateRefId: TemplateRefId1,
+				templateHash: Template.Id,
+				graph: Graph,
+				name: "hello1",
+				change: 1000001,
+				codeChange: 1000002,
+				preflightChange: 1001,
+				clonedPreflightChange: null,
+				startedByUserId: null,
+				priority: Priority.Normal,
 				null,
 				null,
 				null,
@@ -104,23 +104,23 @@ namespace Horde.Build.Tests
 				false,
 				null,
 				null,
-				Arguments: new List<string>()
+				arguments: new List<string>()
 			);
 			Job1 = (await jobService.GetJobAsync(Job1.Id))!;
 
 			Job2 = await jobService.CreateJobAsync(
-				JobId: new JobId("5f69ea1b68423e921b035106"),
-				Stream: Stream!,
-				TemplateRefId: new TemplateRefId("template-id-1"),
-				TemplateHash: ContentHash.MD5("made-up-template-hash"),
-				Graph: fg,
-				Name: "hello2",
-				Change: 2000001,
-				CodeChange: 2000002,
-				PreflightChange: null,
-				ClonedPreflightChange: null,
-				StartedByUserId: null,
-				Priority: Priority.Normal,
+				jobId: new JobId("5f69ea1b68423e921b035106"),
+				stream: Stream!,
+				templateRefId: new TemplateRefId("template-id-1"),
+				templateHash: ContentHash.MD5("made-up-template-hash"),
+				graph: fg,
+				name: "hello2",
+				change: 2000001,
+				codeChange: 2000002,
+				preflightChange: null,
+				clonedPreflightChange: null,
+				startedByUserId: null,
+				priority: Priority.Normal,
 				null,
 				null,
 				null,
@@ -129,7 +129,7 @@ namespace Horde.Build.Tests
 				false,
 				null,
 				null,
-				Arguments: new List<string>()
+				arguments: new List<string>()
 			);
 			Job2 = (await jobService.GetJobAsync(Job2.Id))!;
 
diff --git a/Engine/Source/Programs/Horde/Horde.Build.Tests/Fleet/AutoscaleServiceV2Test.cs b/Engine/Source/Programs/Horde/Horde.Build.Tests/Fleet/AutoscaleServiceV2Test.cs
--- a/Engine/Source/Programs/Horde/Horde.Build.Tests/Fleet/AutoscaleServiceV2Test.cs
+++ b/Engine/Source/Programs/Horde/Horde.Build.Tests/Fleet/AutoscaleServiceV2Test.cs
@@ -1,6 +1,5 @@
 // Copyright Epic Games, Inc. All Rights Reserved.
 
-using StatsdClient;
 using System;
 using System.Collections.Generic;
 using System.Threading;
@@ -13,6 +12,7 @@ using Microsoft.Extensions.DependencyInjection;
 using Microsoft.Extensions.Logging;
 using Microsoft.Extensions.Options;
 using Microsoft.VisualStudio.TestTools.UnitTesting;
+using StatsdClient;
 
 namespace Horde.Build.Tests.Fleet
 {
@@ -70,11 +70,11 @@ namespace Horde.Build.Tests.Fleet
 		{
 			using AutoscaleServiceV2 service = GetAutoscaleService(_fleetManagerSpy);
 			
-			IPool pool1 = await PoolService.CreatePoolAsync("bogusPoolLease1", null, true, 0, 0, SizeStrategy: PoolSizeStrategy.LeaseUtilization);
-			IPool pool2 = await PoolService.CreatePoolAsync("bogusPoolLease2", null, true, 0, 0, SizeStrategy: null);
-			IPool pool3 = await PoolService.CreatePoolAsync("bogusPoolJobQueue1", null, true, 0, 0, SizeStrategy: PoolSizeStrategy.JobQueue);
-			IPool pool4 = await PoolService.CreatePoolAsync("bogusPoolJobQueue2", null, true, 0, 0, SizeStrategy: PoolSizeStrategy.JobQueue);
-			IPool pool5 = await PoolService.CreatePoolAsync("bogusPoolNoOp", null, true, 0, 0, SizeStrategy: PoolSizeStrategy.NoOp);
+			IPool pool1 = await PoolService.CreatePoolAsync("bogusPoolLease1", null, true, 0, 0, sizeStrategy: PoolSizeStrategy.LeaseUtilization);
+			IPool pool2 = await PoolService.CreatePoolAsync("bogusPoolLease2", null, true, 0, 0, sizeStrategy: null);
+			IPool pool3 = await PoolService.CreatePoolAsync("bogusPoolJobQueue1", null, true, 0, 0, sizeStrategy: PoolSizeStrategy.JobQueue);
+			IPool pool4 = await PoolService.CreatePoolAsync("bogusPoolJobQueue2", null, true, 0, 0, sizeStrategy: PoolSizeStrategy.JobQueue);
+			IPool pool5 = await PoolService.CreatePoolAsync("bogusPoolNoOp", null, true, 0, 0, sizeStrategy: PoolSizeStrategy.NoOp);
 
 			PoolSizeStrategySpy leaseUtilizationSpy = new();
 			PoolSizeStrategySpy jobQueueSpy = new();
@@ -99,7 +99,7 @@ namespace Horde.Build.Tests.Fleet
 		public async Task ScaleOutCooldown()
 		{
 			using AutoscaleServiceV2 service = GetAutoscaleService(_fleetManagerSpy);
-			IPool pool = await PoolService.CreatePoolAsync("testPool", null, true, 0, 0, SizeStrategy: PoolSizeStrategy.NoOp);
+			IPool pool = await PoolService.CreatePoolAsync("testPool", null, true, 0, 0, sizeStrategy: PoolSizeStrategy.NoOp);
 
 			// First scale-out will succeed
 			await service.ResizePools(new() { new PoolSizeData(pool, new List<IAgent>(), 1, "Testing") });
@@ -119,7 +119,7 @@ namespace Horde.Build.Tests.Fleet
 		public async Task ScaleInCooldown()
 		{
 			using AutoscaleServiceV2 service = GetAutoscaleService(_fleetManagerSpy);
-			IPool pool = await PoolService.CreatePoolAsync("testPool", null, true, 0, 0, SizeStrategy: PoolSizeStrategy.NoOp);
+			IPool pool = await PoolService.CreatePoolAsync("testPool", null, true, 0, 0, sizeStrategy: PoolSizeStrategy.NoOp);
 			IAgent agent1 = await CreateAgentAsync(pool);
 			IAgent agent2 = await CreateAgentAsync(pool);
 
diff --git a/Engine/Source/Programs/Horde/Horde.Build.Tests/IssueServiceTests.cs b/Engine/Source/Programs/Horde/Horde.Build.Tests/IssueServiceTests.cs
--- a/Engine/Source/Programs/Horde/Horde.Build.Tests/IssueServiceTests.cs
+++ b/Engine/Source/Programs/Horde/Horde.Build.Tests/IssueServiceTests.cs
@@ -1,14 +1,20 @@
 // Copyright Epic Games, Inc. All Rights Reserved.
 
+using System;
+using System.Collections.Generic;
+using System.Linq;
+using System.Runtime.InteropServices;
+using System.Text.Json;
+using System.Threading.Tasks;
 using EpicGames.Core;
 using Horde.Agent.Parser;
 using Horde.Agent.Utility;
-using HordeCommon;
 using Horde.Build.IssueHandlers.Impl;
 using Horde.Build.Models;
 using Horde.Build.Services;
-using Horde.Build.Utilities;
 using Horde.Build.Tests.Stubs.Services;
+using Horde.Build.Utilities;
+using HordeCommon;
 using Microsoft.Extensions.DependencyInjection;
 using Microsoft.Extensions.Logging;
 using Microsoft.Extensions.Logging.Abstractions;
@@ -16,13 +22,6 @@ using Microsoft.VisualStudio.TestTools.UnitTesting;
 using MongoDB.Bson;
 using MongoDB.Driver;
 using Moq;
-using System;
-using System.Collections.Generic;
-using System.Linq;
-using System.Runtime.InteropServices;
-using System.Text.Json;
-using System.Threading.Tasks;
-
 using ProjectId = Horde.Build.Utilities.StringId<Horde.Build.Models.IProject>;
 using StreamId = Horde.Build.Utilities.StringId<Horde.Build.Models.IStream>;
 using TemplateRefId = Horde.Build.Utilities.StringId<Horde.Build.Models.TemplateRef>;
@@ -38,9 +37,9 @@ namespace Horde.Build.Tests
 	{
 		class TestJsonLogger : JsonLogger, IAsyncDisposable
 		{
-			ILogFileService _logFileService;
-			LogId _logId;
-			List<(LogLevel, byte[])> _events = new List<(LogLevel, byte[])>();
+			readonly ILogFileService _logFileService;
+			readonly LogId _logId;
+			readonly List<(LogLevel, byte[])> _events = new List<(LogLevel, byte[])>();
 
 			public TestJsonLogger(ILogFileService logFileService, LogId logId)
 				: base(null, NullLogger.Instance)
@@ -91,22 +90,19 @@ namespace Horde.Build.Tests
 		}
 
 		const string MainStreamName = "//UE4/Main";
-		StreamId _mainStreamId = StreamId.Sanitize(MainStreamName);
+		readonly StreamId _mainStreamId = StreamId.Sanitize(MainStreamName);
 
 		const string ReleaseStreamName = "//UE4/Release";
-		StreamId _releaseStreamId = StreamId.Sanitize(ReleaseStreamName);
+		readonly StreamId _releaseStreamId = StreamId.Sanitize(ReleaseStreamName);
 
 		const string DevStreamName = "//UE4/Dev";
-		StreamId _devStreamId = StreamId.Sanitize(DevStreamName);
-
-		IGraph _graph;
-		PerforceServiceStub _perforce;
-
-		UserId _timId;
-		UserId _jerryId;
-		UserId _bobId;
-
-		DirectoryReference _workspaceDir;
+		readonly StreamId _devStreamId = StreamId.Sanitize(DevStreamName);
+		readonly IGraph _graph;
+		readonly PerforceServiceStub _perforce;
+		readonly UserId _timId;
+		readonly UserId _jerryId;
+		readonly UserId _bobId;
+		readonly DirectoryReference _workspaceDir;
 
 		public IssueServiceTests()
 		{
@@ -328,7 +324,7 @@ namespace Horde.Build.Tests
 				List<IIssue> issues = await IssueService.FindIssuesAsync();
 
 				IIssue issue = issues[0];
-				await IssueService.UpdateIssueAsync(issue.Id, Description: "Hello world!");
+				await IssueService.UpdateIssueAsync(issue.Id, description: "Hello world!");
 				IIssue? newIssue = await IssueService.GetIssueAsync(issue.Id);
 				Assert.AreEqual(newIssue?.Description, "Hello world!");
 			}
@@ -381,7 +377,7 @@ namespace Horde.Build.Tests
 				List<IIssue> issues = await IssueService.FindIssuesAsync();
 				Assert.AreEqual(0, issues.Count);
 
-				List<IIssue> openIssues = await IssueService.FindIssuesAsync(Resolved: false);
+				List<IIssue> openIssues = await IssueService.FindIssuesAsync(resolved: false);
 				Assert.AreEqual(0, openIssues.Count);
 			}
 
@@ -429,7 +425,7 @@ namespace Horde.Build.Tests
 				List<IIssue> issues = await IssueService.FindIssuesAsync();
 				Assert.AreEqual(0, issues.Count);
 
-				List<IIssue> openIssues = await IssueService.FindIssuesAsync(Resolved: false);
+				List<IIssue> openIssues = await IssueService.FindIssuesAsync(resolved: false);
 				Assert.AreEqual(0, openIssues.Count);
 			}
 
@@ -476,7 +472,7 @@ namespace Horde.Build.Tests
 				Assert.AreEqual(_timId, suspects[0].AuthorId);
 				//				Assert.AreEqual(null, Suspects[0].OriginatingChange);
 
-				List<IIssue> openIssues = await IssueService.FindIssuesAsync(Resolved: false);
+				List<IIssue> openIssues = await IssueService.FindIssuesAsync(resolved: false);
 				Assert.AreEqual(1, openIssues.Count);
 			}
 
@@ -507,7 +503,7 @@ namespace Horde.Build.Tests
 				Assert.AreEqual(75, suspects[1].Change);
 				Assert.AreEqual(_jerryId, suspects[1].AuthorId);
 
-				List<IIssue> openIssues = await IssueService.FindIssuesAsync(Resolved: false);
+				List<IIssue> openIssues = await IssueService.FindIssuesAsync(resolved: false);
 				Assert.AreEqual(1, openIssues.Count);
 			}
 
@@ -518,7 +514,7 @@ namespace Horde.Build.Tests
 				IJob job = CreateJob(_mainStreamId, 125, "Test Build", _graph);
 				await UpdateCompleteStep(job, 0, 0, JobStepOutcome.Success);
 
-				List<IIssue> issues = await IssueService.FindIssuesAsync(Resolved: true);
+				List<IIssue> issues = await IssueService.FindIssuesAsync(resolved: true);
 				Assert.AreEqual(1, issues.Count);
 
 				IIssue issue = issues[0];
@@ -529,7 +525,7 @@ namespace Horde.Build.Tests
 				Assert.AreEqual(110, stream.LastSuccess?.Change);
 				Assert.AreEqual(125, stream.NextSuccess?.Change);
 
-				List<IIssue> openIssues = await IssueService.FindIssuesAsync(Resolved: false);
+				List<IIssue> openIssues = await IssueService.FindIssuesAsync(resolved: false);
 				Assert.AreEqual(0, openIssues.Count);
 			}
 
@@ -546,7 +542,7 @@ namespace Horde.Build.Tests
 				await ParseEventsAsync(job, 0, 0, lines);
 				await UpdateCompleteStep(job, 0, 0, JobStepOutcome.Failure);
 
-				List<IIssue> issues = await IssueService.FindIssuesAsync(Resolved: true);
+				List<IIssue> issues = await IssueService.FindIssuesAsync(resolved: true);
 				Assert.AreEqual(1, issues.Count);
 
 				IIssue issue = issues[0];
@@ -569,10 +565,10 @@ namespace Horde.Build.Tests
 				await AddEvent(job, 0, 1, new { });
 				await UpdateCompleteStep(job, 0, 1, JobStepOutcome.Failure);
 
-				List<IIssue> resolvedIssues = await IssueService.FindIssuesAsync(Resolved: true);
+				List<IIssue> resolvedIssues = await IssueService.FindIssuesAsync(resolved: true);
 				Assert.AreEqual(1, resolvedIssues.Count);
 
-				List<IIssue> unresolvedIssues = await IssueService.FindIssuesAsync(Resolved: false);
+				List<IIssue> unresolvedIssues = await IssueService.FindIssuesAsync(resolved: false);
 				Assert.AreEqual(1, unresolvedIssues.Count);
 			}
 		}
@@ -674,7 +670,7 @@ namespace Horde.Build.Tests
 				Assert.AreEqual(1, issues.Count);
 				Assert.IsFalse(issues[0].Promoted);
 
-				await IssueService.UpdateIssueAsync(issues[0].Id, Promoted: true);
+				await IssueService.UpdateIssueAsync(issues[0].Id, promoted: true);
 
 				issues = await IssueService.FindIssuesAsync();
 				Assert.AreEqual(1, issues.Count);
@@ -855,7 +851,7 @@ namespace Horde.Build.Tests
 			{
 				List<IIssue> issues = await IssueService.FindIssuesAsync();
 				Assert.AreEqual(1, issues.Count);
-				await IssueService.UpdateIssueAsync(issues[0].Id, DeclinedById: _timId);
+				await IssueService.UpdateIssueAsync(issues[0].Id, declinedById: _timId);
 
 				issues = await IssueService.FindIssuesAsync();
 				Assert.AreEqual(1, issues.Count);
@@ -1133,7 +1129,7 @@ namespace Horde.Build.Tests
 
 			// Add SpanB to IssueA
 			{
-				await IssueService.UpdateIssueAsync(issueA.Id, AddSpanIds: new List<ObjectId> { spanB.Id });
+				await IssueService.UpdateIssueAsync(issueA.Id, addSpanIds: new List<ObjectId> { spanB.Id });
 
 				IIssue newIssueA = (await IssueService.GetIssueAsync(issueA.Id))!;
 				Assert.IsNull(newIssueA.VerifiedAt);
@@ -1154,7 +1150,7 @@ namespace Horde.Build.Tests
 
 			// Add SpanA and SpanB to IssueB
 			{
-				await IssueService.UpdateIssueAsync(issueB.Id, AddSpanIds: new List<ObjectId> { spanA.Id, spanB.Id });
+				await IssueService.UpdateIssueAsync(issueB.Id, addSpanIds: new List<ObjectId> { spanA.Id, spanB.Id });
 
 				IIssue newIssueA = (await IssueService.GetIssueAsync(issueA.Id))!;
 				Assert.IsNotNull(newIssueA.VerifiedAt);
@@ -1273,7 +1269,7 @@ namespace Horde.Build.Tests
 			// Scenario: Issue is marked fixed
 			// Expected: Resolved time, owner is set
 			{
-				await IssueService.UpdateIssueAsync(issueId, ResolvedById: _bobId);
+				await IssueService.UpdateIssueAsync(issueId, resolvedById: _bobId);
 
 				List<IIssue> openIssues = await IssueService.FindIssuesAsync();
 				Assert.AreEqual(0, openIssues.Count);
@@ -1323,7 +1319,7 @@ namespace Horde.Build.Tests
 			// Scenario: Issue is marked fixed again, at a particular changelist
 			// Expected: Resolved time, owner is set
 			{
-				await IssueService.UpdateIssueAsync(issueId, ResolvedById: _bobId, FixChange: 115);
+				await IssueService.UpdateIssueAsync(issueId, resolvedById: _bobId, fixChange: 115);
 
 				List<IIssue> openIssues = await IssueService.FindIssuesAsync();
 				Assert.AreEqual(0, openIssues.Count);
@@ -1356,7 +1352,7 @@ namespace Horde.Build.Tests
 			// Scenario: Issue is marked fixed again, at a particular changelist
 			// Expected: Resolved time, owner is set
 			{
-				await IssueService.UpdateIssueAsync(issueId, ResolvedById: _bobId, FixChange: 125);
+				await IssueService.UpdateIssueAsync(issueId, resolvedById: _bobId, fixChange: 125);
 
 				List<IIssue> openIssues = await IssueService.FindIssuesAsync();
 				Assert.AreEqual(0, openIssues.Count);
diff --git a/Engine/Source/Programs/Horde/Horde.Build.Tests/JobCollectionTests.cs b/Engine/Source/Programs/Horde/Horde.Build.Tests/JobCollectionTests.cs
--- a/Engine/Source/Programs/Horde/Horde.Build.Tests/JobCollectionTests.cs
+++ b/Engine/Source/Programs/Horde/Horde.Build.Tests/JobCollectionTests.cs
@@ -1,27 +1,26 @@
 // Copyright Epic Games, Inc. All Rights Reserved.
 
-using MongoDB.Bson;
 using System;
 using System.Collections.Generic;
-using PoolId = Horde.Build.Utilities.StringId<Horde.Build.Models.IPool>;
-using StreamId = Horde.Build.Utilities.StringId<Horde.Build.Models.IStream>;
-using TemplateRefId = Horde.Build.Utilities.StringId<Horde.Build.Models.TemplateRef>;
-
-using Microsoft.VisualStudio.TestTools.UnitTesting;
-using EpicGames.Core;
-using Moq;
-using Horde.Build.Models;
+using System.Linq;
 using System.Threading.Tasks;
+using EpicGames.Core;
 using Horde.Build.Api;
-using System.Linq;
-using HordeCommon;
+using Horde.Build.Models;
 using Horde.Build.Utilities;
+using HordeCommon;
+using Microsoft.VisualStudio.TestTools.UnitTesting;
+using MongoDB.Bson;
+using Moq;
 
 namespace Horde.Build.Tests
 {
 	using JobId = ObjectId<IJob>;
 	using LeaseId = ObjectId<ILease>;
 	using LogId = ObjectId<ILogFile>;
+	using PoolId = StringId<IPool>;
+	using StreamId = StringId<IStream>;
+	using TemplateRefId = StringId<TemplateRef>;
 
 	[TestClass]
 	public class JobCollectionTests : TestSetup
diff --git a/Engine/Source/Programs/Horde/Horde.Build.Tests/JobServiceTests.cs b/Engine/Source/Programs/Horde/Horde.Build.Tests/JobServiceTests.cs
--- a/Engine/Source/Programs/Horde/Horde.Build.Tests/JobServiceTests.cs
+++ b/Engine/Source/Programs/Horde/Horde.Build.Tests/JobServiceTests.cs
@@ -1,23 +1,23 @@
 // Copyright Epic Games, Inc. All Rights Reserved.
 
 using System;
-using HordeCommon;
-using Horde.Build.Api;
-using Horde.Build.Models;
-using Microsoft.VisualStudio.TestTools.UnitTesting;
 using System.Collections.Generic;
 using System.Text;
 using System.Threading.Tasks;
 using EpicGames.Core;
-using ProjectId = Horde.Build.Utilities.StringId<Horde.Build.Models.IProject>;
-using StreamId = Horde.Build.Utilities.StringId<Horde.Build.Models.IStream>;
-using TemplateRefId = Horde.Build.Utilities.StringId<Horde.Build.Models.TemplateRef>;
+using Horde.Build.Api;
+using Horde.Build.Models;
 using Horde.Build.Utilities;
+using HordeCommon;
+using Microsoft.VisualStudio.TestTools.UnitTesting;
 
 namespace Horde.Build.Tests
 {
 	using JobId = ObjectId<IJob>;
 	using LogId = ObjectId<ILogFile>;
+	using ProjectId = StringId<IProject>;
+	using StreamId = StringId<IStream>;
+	using TemplateRefId = StringId<TemplateRef>;
 
 	[TestClass]
 	public class JobServiceTests : TestSetup
@@ -95,18 +95,18 @@ namespace Horde.Build.Tests
 		{
 			IUser user = await UserCollection.FindOrAddUserByLoginAsync(startedByUserName);
 			return await JobService.CreateJobAsync(
-				JobId: JobId.GenerateNewId(),
-				Stream: fixture!.Stream!,
-				TemplateRefId: new TemplateRefId(templateRefId),
-				TemplateHash: new ContentHash(Encoding.ASCII.GetBytes(templateHash)),
-				Graph: fixture!.Graph,
-				Name: "hello1",
-				Change: 1000001,
-				CodeChange: 1000002,
-				PreflightChange: preflightChange,
-				ClonedPreflightChange: null,
-				StartedByUserId: user.Id,
-				Priority: Priority.Normal,
+				jobId: JobId.GenerateNewId(),
+				stream: fixture!.Stream!,
+				templateRefId: new TemplateRefId(templateRefId),
+				templateHash: new ContentHash(Encoding.ASCII.GetBytes(templateHash)),
+				graph: fixture!.Graph,
+				name: "hello1",
+				change: 1000001,
+				codeChange: 1000002,
+				preflightChange: preflightChange,
+				clonedPreflightChange: null,
+				startedByUserId: user.Id,
+				priority: Priority.Normal,
 				null,
 				null,
 				null,
@@ -115,7 +115,7 @@ namespace Horde.Build.Tests
 				false,
 				null,
 				null,
-				Arguments: new List<string>(arguments)
+				arguments: new List<string>(arguments)
 			);
 		}
 		
@@ -163,9 +163,9 @@ namespace Horde.Build.Tests
 			groupA.Nodes.Add(new NewNode("Compile"));
 
 			NewGroup groupB = new NewGroup("win", new List<NewNode>());
-			groupB.Nodes.Add(new NewNode("Cook", RunEarly: true));
+			groupB.Nodes.Add(new NewNode("Cook", runEarly: true));
 			groupB.Nodes.Add(new NewNode("Middle"));
-			groupB.Nodes.Add(new NewNode("Pak", InputDependencies: new List<string> { "Compile", "Cook" }));
+			groupB.Nodes.Add(new NewNode("Pak", inputDependencies: new List<string> { "Compile", "Cook" }));
 
 			graph = await GraphCollection.AppendAsync(graph, new List<NewGroup> { groupA, groupB });
 
diff --git a/Engine/Source/Programs/Horde/Horde.Build.Tests/JobTaskSourceTests.cs b/Engine/Source/Programs/Horde/Horde.Build.Tests/JobTaskSourceTests.cs
--- a/Engine/Source/Programs/Horde/Horde.Build.Tests/JobTaskSourceTests.cs
+++ b/Engine/Source/Programs/Horde/Horde.Build.Tests/JobTaskSourceTests.cs
@@ -1,13 +1,13 @@
 // Copyright Epic Games, Inc. All Rights Reserved.
 
-using HordeCommon;
-using Horde.Build.Api;
-using Horde.Build.Models;
-using Microsoft.VisualStudio.TestTools.UnitTesting;
 using System.Collections.Generic;
 using System.Threading;
 using System.Threading.Tasks;
+using Horde.Build.Api;
+using Horde.Build.Models;
 using Horde.Build.Utilities;
+using HordeCommon;
+using Microsoft.VisualStudio.TestTools.UnitTesting;
 
 namespace Horde.Build.Tests
 {
diff --git a/Engine/Source/Programs/Horde/Horde.Build.Tests/JobsControllerDbTest.cs b/Engine/Source/Programs/Horde/Horde.Build.Tests/JobsControllerDbTest.cs
--- a/Engine/Source/Programs/Horde/Horde.Build.Tests/JobsControllerDbTest.cs
+++ b/Engine/Source/Programs/Horde/Horde.Build.Tests/JobsControllerDbTest.cs
@@ -3,10 +3,10 @@
 using System.Collections.Generic;
 using System.Threading.Tasks;
 using Horde.Build.Api;
-using Microsoft.VisualStudio.TestTools.UnitTesting;
 using Horde.Build.Controllers;
 using Horde.Build.Models;
 using Microsoft.AspNetCore.Mvc;
+using Microsoft.VisualStudio.TestTools.UnitTesting;
 
 namespace Horde.Build.Tests
 {
@@ -28,7 +28,7 @@ namespace Horde.Build.Tests
 	        Assert.AreEqual("hello2", (res.Value[0] as GetJobResponse)!.Name);
 	        Assert.AreEqual("hello1", (res.Value[1] as GetJobResponse)!.Name);
 	        
-	        res = await JobsController.FindJobsAsync(IncludePreflight: false);
+	        res = await JobsController.FindJobsAsync(includePreflight: false);
 	        Assert.AreEqual(1, res.Value!.Count);
 	        Assert.AreEqual("hello2", (res.Value[0] as GetJobResponse)!.Name);
         }
diff --git a/Engine/Source/Programs/Horde/Horde.Build.Tests/LogChunkTests.cs b/Engine/Source/Programs/Horde/Horde.Build.Tests/LogChunkTests.cs
--- a/Engine/Source/Programs/Horde/Horde.Build.Tests/LogChunkTests.cs
+++ b/Engine/Source/Programs/Horde/Horde.Build.Tests/LogChunkTests.cs
@@ -1,18 +1,18 @@
 // Copyright Epic Games, Inc. All Rights Reserved.
 
-using Horde.Build.Api;
-using Horde.Build.Logs;
-using Microsoft.VisualStudio.TestTools.UnitTesting;
 using System;
 using System.Collections.Generic;
 using System.Text;
+using Horde.Build.Api;
+using Horde.Build.Logs;
+using Microsoft.VisualStudio.TestTools.UnitTesting;
 
 namespace Horde.Build.Tests
 {
 	[TestClass]
 	public class LogChunkTests
 	{
-		string[] _text =
+		readonly string[] _text =
 		{
 			@"{ ""id"": 1, ""message"": ""foo"" }",
 			@"{ ""id"": 2, ""message"": ""bar"" }",
diff --git a/Engine/Source/Programs/Horde/Horde.Build.Tests/LogIndexingTests.cs b/Engine/Source/Programs/Horde/Horde.Build.Tests/LogIndexingTests.cs
--- a/Engine/Source/Programs/Horde/Horde.Build.Tests/LogIndexingTests.cs
+++ b/Engine/Source/Programs/Horde/Horde.Build.Tests/LogIndexingTests.cs
@@ -1,6 +1,14 @@
 // Copyright Epic Games, Inc. All Rights Reserved.
 
+using System;
+using System.Collections.Generic;
+using System.IO;
+using System.Linq;
+using System.Text;
+using System.Threading.Tasks;
 using Horde.Build.Api;
+using Horde.Build.Collections.Impl;
+using Horde.Build.Compute.Tests.Properties;
 using Horde.Build.Logs;
 using Horde.Build.Logs.Builder;
 using Horde.Build.Logs.Storage;
@@ -8,19 +16,11 @@ using Horde.Build.Logs.Storage.Impl;
 using Horde.Build.Models;
 using Horde.Build.Services;
 using Horde.Build.Utilities;
-using Horde.Build.Compute.Tests.Properties;
+using HordeCommon;
 using Microsoft.Extensions.Configuration;
 using Microsoft.Extensions.DependencyInjection;
 using Microsoft.Extensions.Logging;
 using Microsoft.VisualStudio.TestTools.UnitTesting;
-using System;
-using System.Collections.Generic;
-using System.IO;
-using System.Linq;
-using System.Text;
-using System.Threading.Tasks;
-using Horde.Build.Collections.Impl;
-using HordeCommon;
 
 namespace Horde.Build.Tests
 {
diff --git a/Engine/Source/Programs/Horde/Horde.Build.Tests/Models/Job/GetLabelStatesTests.cs b/Engine/Source/Programs/Horde/Horde.Build.Tests/Models/Job/GetLabelStatesTests.cs
--- a/Engine/Source/Programs/Horde/Horde.Build.Tests/Models/Job/GetLabelStatesTests.cs
+++ b/Engine/Source/Programs/Horde/Horde.Build.Tests/Models/Job/GetLabelStatesTests.cs
@@ -2,11 +2,11 @@
 
 using System.Collections.Generic;
 using Horde.Build.Api;
-using HordeCommon;
 using Horde.Build.Models;
+using Horde.Build.Utilities;
+using HordeCommon;
 using Microsoft.VisualStudio.TestTools.UnitTesting;
 using Moq;
-using Horde.Build.Utilities;
 
 namespace Horde.Build.Tests
 {
diff --git a/Engine/Source/Programs/Horde/Horde.Build.Tests/NotificationServiceTests.cs b/Engine/Source/Programs/Horde/Horde.Build.Tests/NotificationServiceTests.cs
--- a/Engine/Source/Programs/Horde/Horde.Build.Tests/NotificationServiceTests.cs
+++ b/Engine/Source/Programs/Horde/Horde.Build.Tests/NotificationServiceTests.cs
@@ -3,12 +3,12 @@
 using System;
 using System.Collections.Generic;
 using System.Threading.Tasks;
-using HordeCommon;
 using Horde.Build.Api;
 using Horde.Build.Models;
 using Horde.Build.Notifications;
 using Horde.Build.Notifications.Impl;
 using Horde.Build.Utilities;
+using HordeCommon;
 using Microsoft.Extensions.DependencyInjection;
 using Microsoft.VisualStudio.TestTools.UnitTesting;
 using Moq;
@@ -95,10 +95,10 @@ namespace Horde.Build.Tests
 			FakeNotificationSink fakeSink = ServiceProvider.GetRequiredService<FakeNotificationSink>();
 
 			NotificationService service = (NotificationService)ServiceProvider.GetRequiredService<INotificationService>();
-			await service.Ticker.StartAsync();
+			await service._ticker.StartAsync();
 
 			Fixture fixture = await CreateFixtureAsync();
-			IPool pool = await PoolService.CreatePoolAsync("BogusPool", Properties: new Dictionary<string, string>());
+			IPool pool = await PoolService.CreatePoolAsync("BogusPool", properties: new Dictionary<string, string>());
 
 			Assert.AreEqual(0, fakeSink.JobScheduledNotifications.Count);
 			service.NotifyJobScheduled(pool, false, fixture.Job1, fixture.Graph, SubResourceId.Random());
@@ -106,7 +106,7 @@ namespace Horde.Build.Tests
 			
 			// Currently no good way to wait for NotifyJobScheduled() to complete as the execution is completely async in background task (see ExecuteAsync)
 			await Task.Delay(1000);
-			await Clock.AdvanceAsync(service.NotificationBatchInterval + TimeSpan.FromMinutes(5));
+			await Clock.AdvanceAsync(service._notificationBatchInterval + TimeSpan.FromMinutes(5));
 			Assert.AreEqual(2, fakeSink.JobScheduledNotifications.Count);
 			Assert.AreEqual(1, fakeSink.JobScheduledCallCount);
 		}
diff --git a/Engine/Source/Programs/Horde/Horde.Build.Tests/PoolServiceTests.cs b/Engine/Source/Programs/Horde/Horde.Build.Tests/PoolServiceTests.cs
--- a/Engine/Source/Programs/Horde/Horde.Build.Tests/PoolServiceTests.cs
+++ b/Engine/Source/Programs/Horde/Horde.Build.Tests/PoolServiceTests.cs
@@ -25,7 +25,7 @@ namespace Horde.Build.Tests
 
         private async Task<IPool> CreatePoolFixture(string name)
         {
-            return await PoolService.CreatePoolAsync(name, Properties: _fixtureProps);
+            return await PoolService.CreatePoolAsync(name, properties: _fixtureProps);
         }
 
         [TestMethod]
@@ -79,7 +79,7 @@ namespace Horde.Build.Tests
                 {"cookies", "yumyum"},
             };
 
-            IPool? updatedPool = await PoolService.UpdatePoolAsync(pool, $"update-pool-new-name-{uniqueSuffix}", NewProperties: updatedProps);
+            IPool? updatedPool = await PoolService.UpdatePoolAsync(pool, $"update-pool-new-name-{uniqueSuffix}", newProperties: updatedProps);
 			Assert.IsNotNull(updatedPool);
             Assert.AreEqual(pool.Id, updatedPool!.Id);
             Assert.AreEqual($"update-pool-new-name-{uniqueSuffix}", updatedPool.Name);
@@ -92,7 +92,7 @@ namespace Horde.Build.Tests
         public async Task UpdatePoolCollectionTest()
         {
 	        IPool pool = await CreatePoolFixture("update-pool-2");
-	        await PoolCollection.TryUpdateAsync(pool, LastScaleUpTime: DateTime.UtcNow);
+	        await PoolCollection.TryUpdateAsync(pool, lastScaleUpTime: DateTime.UtcNow);
         }
     }
 }
\ No newline at end of file
diff --git a/Engine/Source/Programs/Horde/Horde.Build.Tests/PoolsControllerTest.cs b/Engine/Source/Programs/Horde/Horde.Build.Tests/PoolsControllerTest.cs
--- a/Engine/Source/Programs/Horde/Horde.Build.Tests/PoolsControllerTest.cs
+++ b/Engine/Source/Programs/Horde/Horde.Build.Tests/PoolsControllerTest.cs
@@ -2,8 +2,8 @@
 
 using System.Collections.Generic;
 using System.Threading.Tasks;
-using Horde.Build.Fleet.Autoscale;
 using Horde.Build.Api;
+using Horde.Build.Fleet.Autoscale;
 using Horde.Build.Models;
 using Horde.Build.Utilities;
 using Microsoft.AspNetCore.Mvc;
@@ -17,7 +17,7 @@ namespace Horde.Build.Tests
         [TestMethod]
         public async Task GetPoolsTest()
         {
-	        IPool pool1 = await PoolService.CreatePoolAsync("Pool1", Properties: new() { {"foo", "bar"}, {"lorem", "ipsum"} });
+	        IPool pool1 = await PoolService.CreatePoolAsync("Pool1", properties: new() { {"foo", "bar"}, {"lorem", "ipsum"} });
 	        ActionResult<List<object>> rawResult = await PoolsController.GetPoolsAsync();
 	        Assert.AreEqual(1, rawResult.Value!.Count);
 	        GetPoolResponse response = (rawResult.Value![0] as GetPoolResponse)!;
@@ -52,7 +52,7 @@ namespace Horde.Build.Tests
         [TestMethod]
         public async Task UpdatePoolTest()
         {
-	        IPool pool1 = await PoolService.CreatePoolAsync("Pool1", Properties: new() { {"foo", "bar"}, {"lorem", "ipsum"} });
+	        IPool pool1 = await PoolService.CreatePoolAsync("Pool1", properties: new() { {"foo", "bar"}, {"lorem", "ipsum"} });
 
 	        await PoolsController.UpdatePoolAsync(pool1.Id.ToString(), new UpdatePoolRequest
 	        {
diff --git a/Engine/Source/Programs/Horde/Horde.Build.Tests/RpcServiceTest.cs b/Engine/Source/Programs/Horde/Horde.Build.Tests/RpcServiceTest.cs
--- a/Engine/Source/Programs/Horde/Horde.Build.Tests/RpcServiceTest.cs
+++ b/Engine/Source/Programs/Horde/Horde.Build.Tests/RpcServiceTest.cs
@@ -10,10 +10,10 @@ using System.Threading.Channels;
 using System.Threading.Tasks;
 using Google.Protobuf;
 using Grpc.Core;
-using HordeCommon.Rpc;
 using Horde.Build.Models;
 using Horde.Build.Services;
 using Horde.Build.Utilities;
+using HordeCommon.Rpc;
 using Microsoft.AspNetCore.Http;
 using Microsoft.AspNetCore.Http.Features;
 using Microsoft.Extensions.Hosting;
@@ -329,7 +329,7 @@ namespace Horde.Build.Tests
 		[TestMethod]
 		public async Task QueryServerSessionTest()
 		{
-			RpcService.LongPollTimeout = TimeSpan.FromMilliseconds(200);
+			RpcService._longPollTimeout = TimeSpan.FromMilliseconds(200);
 
 			TestAsyncStreamReader<QueryServerStateRequest> requestStream =
 				new TestAsyncStreamReader<QueryServerStateRequest>(_adminContext);
diff --git a/Engine/Source/Programs/Horde/Horde.Build.Tests/SchedulerTests.cs b/Engine/Source/Programs/Horde/Horde.Build.Tests/SchedulerTests.cs
--- a/Engine/Source/Programs/Horde/Horde.Build.Tests/SchedulerTests.cs
+++ b/Engine/Source/Programs/Horde/Horde.Build.Tests/SchedulerTests.cs
@@ -2,24 +2,24 @@
 
 using System;
 using System.Collections.Generic;
+using System.Linq;
 using System.Threading.Tasks;
-using HordeCommon;
+using Horde.Build.Api;
+using Horde.Build.Collections;
 using Horde.Build.Models;
 using Horde.Build.Services;
+using Horde.Build.Utilities;
+using HordeCommon;
 using Microsoft.Extensions.DependencyInjection;
 using Microsoft.VisualStudio.TestTools.UnitTesting;
-using StreamId = Horde.Build.Utilities.StringId<Horde.Build.Models.IStream>;
-using ProjectId = Horde.Build.Utilities.StringId<Horde.Build.Models.IProject>;
-using TemplateRefId = Horde.Build.Utilities.StringId<Horde.Build.Models.TemplateRef>;
-using Horde.Build.Api;
-using System.Linq;
-using Horde.Build.Collections;
-using Horde.Build.Utilities;
 
 namespace Horde.Build.Tests
 {
 	using JobId = ObjectId<IJob>;
 	using LogId = ObjectId<ILogFile>;
+	using ProjectId = StringId<IProject>;
+	using StreamId = StringId<IStream>;
+	using TemplateRefId = StringId<TemplateRef>;
 
 	[TestClass]
     public class SchedulerTests : TestSetup
@@ -126,7 +126,7 @@ namespace Horde.Build.Tests
 			DateTime startTime = new DateTime(2021, 1, 1, 12, 0, 0, DateTimeKind.Utc); // Friday Jan 1, 2021 
 			Clock.UtcNow = startTime;
 
-			Schedule schedule = new Schedule(Clock.UtcNow, RequireSubmittedChange: false);
+			Schedule schedule = new Schedule(Clock.UtcNow, requireSubmittedChange: false);
 			schedule.Patterns.Add(new SchedulePattern(new List<DayOfWeek> { DayOfWeek.Friday, DayOfWeek.Sunday }, 13 * 60, null, null));
 
 			DateTime? nextTime = schedule.GetNextTriggerTimeUtc(startTime, TimeZoneInfo.Utc);
@@ -148,7 +148,7 @@ namespace Horde.Build.Tests
 			DateTime startTime = new DateTime(2021, 1, 1, 12, 0, 0, DateTimeKind.Utc); // Friday Jan 1, 2021 
 			Clock.UtcNow = startTime;
 
-			Schedule schedule = new Schedule(Clock.UtcNow, RequireSubmittedChange: false);
+			Schedule schedule = new Schedule(Clock.UtcNow, requireSubmittedChange: false);
 			schedule.Patterns.Add(new SchedulePattern(null, 13 * 60, 14 * 60, 15));
 
 			DateTime? nextTime = schedule.GetNextTriggerTimeUtc(startTime, TimeZoneInfo.Utc);
@@ -176,7 +176,7 @@ namespace Horde.Build.Tests
 			DateTime startTime = new DateTime(2021, 1, 1, 0, 0, 0, DateTimeKind.Utc); // Friday Jan 1, 2021 
 			Clock.UtcNow = startTime;
 
-			Schedule schedule = new Schedule(Clock.UtcNow, RequireSubmittedChange: false);
+			Schedule schedule = new Schedule(Clock.UtcNow, requireSubmittedChange: false);
 			schedule.Patterns.Add(new SchedulePattern(null, 11 * 60, 0, 0));
 			schedule.Patterns.Add(new SchedulePattern(null, 19 * 60, 0, 0));
 
@@ -358,7 +358,7 @@ namespace Horde.Build.Tests
 			Assert.AreEqual(0, jobs3.Count);
 
 			// Mark the original job as complete
-			await JobService.UpdateJobAsync(jobs2[0], AbortedByUserId: KnownUsers.System);
+			await JobService.UpdateJobAsync(jobs2[0], abortedByUserId: KnownUsers.System);
 
 			// Test that another job does not trigger
 			await Clock.AdvanceAsync(TimeSpan.FromHours(0.5));
@@ -530,7 +530,7 @@ namespace Horde.Build.Tests
 					for (int stepIdx = 0; stepIdx < job1.Batches[batchIdx].Steps.Count; stepIdx++)
 					{
 						SubResourceId stepId1 = job1.Batches[batchIdx].Steps[stepIdx].Id;
-						job1 = Deref(await JobService.UpdateStepAsync(job1, batchId1, stepId1, JobStepState.Completed, JobStepOutcome.Success, NewLogId: LogId.GenerateNewId()));
+						job1 = Deref(await JobService.UpdateStepAsync(job1, batchId1, stepId1, JobStepState.Completed, JobStepOutcome.Success, newLogId: LogId.GenerateNewId()));
 					}
 					job1 = Deref(await JobService.UpdateBatchAsync(job1, batchId1, LogId.GenerateNewId(), JobStepBatchState.Complete));
 				}
@@ -595,7 +595,7 @@ namespace Horde.Build.Tests
 			IStream stream = await SetScheduleAsync(schedule);
 
 			IStreamCollection streamCollection = ServiceProvider.GetRequiredService<IStreamCollection>();
-			await streamCollection.TryUpdatePauseStateAsync(stream, NewPausedUntil: startTime.AddHours(5), NewPauseComment: "testing");
+			await streamCollection.TryUpdatePauseStateAsync(stream, newPausedUntil: startTime.AddHours(5), newPauseComment: "testing");
 
 			// Try trigger a job. No job should be scheduled as the stream is paused
 			await Clock.AdvanceAsync(TimeSpan.FromHours(1.25));
diff --git a/Engine/Source/Programs/Horde/Horde.Build.Tests/StreamServiceTest.cs b/Engine/Source/Programs/Horde/Horde.Build.Tests/StreamServiceTest.cs
--- a/Engine/Source/Programs/Horde/Horde.Build.Tests/StreamServiceTest.cs
+++ b/Engine/Source/Programs/Horde/Horde.Build.Tests/StreamServiceTest.cs
@@ -22,14 +22,14 @@ namespace Horde.Build.Tests
 	        Assert.IsNull(stream.PauseComment);
 
 	        DateTime pausedUntil = DateTime.UtcNow.AddHours(1);
-	        await StreamService.UpdatePauseStateAsync(stream, NewPausedUntil: pausedUntil, NewPauseComment: "mycomment");
+	        await StreamService.UpdatePauseStateAsync(stream, newPausedUntil: pausedUntil, newPauseComment: "mycomment");
 	        stream = (await StreamService.GetStreamAsync(fixture!.Stream!.Id))!;
 	        // Comparing by string to avoid comparing exact milliseconds as those are not persisted in MongoDB fields
 	        Assert.IsTrue(stream.IsPaused(DateTime.UtcNow));
 	        Assert.AreEqual(pausedUntil.ToString(CultureInfo.InvariantCulture), stream.PausedUntil!.Value.ToString(CultureInfo.InvariantCulture));
 	        Assert.AreEqual("mycomment", stream.PauseComment);
 	        
-	        await StreamService.UpdatePauseStateAsync(stream, NewPausedUntil: null, NewPauseComment: null);
+	        await StreamService.UpdatePauseStateAsync(stream, newPausedUntil: null, newPauseComment: null);
 	        stream = (await StreamService.GetStreamAsync(fixture!.Stream!.Id))!;
 	        Assert.IsFalse(stream.IsPaused(DateTime.UtcNow));
 	        Assert.IsNull(stream.PausedUntil);
diff --git a/Engine/Source/Programs/Horde/Horde.Build.Tests/Stubs/Collections/JobStepRefCollectionStub.cs b/Engine/Source/Programs/Horde/Horde.Build.Tests/Stubs/Collections/JobStepRefCollectionStub.cs
--- a/Engine/Source/Programs/Horde/Horde.Build.Tests/Stubs/Collections/JobStepRefCollectionStub.cs
+++ b/Engine/Source/Programs/Horde/Horde.Build.Tests/Stubs/Collections/JobStepRefCollectionStub.cs
@@ -1,21 +1,20 @@
 // Copyright Epic Games, Inc. All Rights Reserved.
 
-using Horde.Build.Collections;
-using HordeCommon;
-using Horde.Build.Models;
-using Horde.Build.Utilities;
 using System;
 using System.Collections.Generic;
 using System.Threading.Tasks;
-
-using StreamId = Horde.Build.Utilities.StringId<Horde.Build.Models.IStream>;
-using TemplateRefId = Horde.Build.Utilities.StringId<Horde.Build.Models.TemplateRef>;
-using PoolId = Horde.Build.Utilities.StringId<Horde.Build.Models.IPool>;
+using Horde.Build.Collections;
+using Horde.Build.Models;
+using Horde.Build.Utilities;
+using HordeCommon;
 
 namespace Horde.Build.Tests.Stubs.Collections
 {
 	using JobId = ObjectId<IJob>;
 	using LogId = ObjectId<ILogFile>;
+	using PoolId = StringId<IPool>;
+	using StreamId = StringId<IStream>;
+	using TemplateRefId = StringId<TemplateRef>;
 
 	class JobStepRefStub : IJobStepRef
 	{
diff --git a/Engine/Source/Programs/Horde/Horde.Build.Tests/TestSetup.cs b/Engine/Source/Programs/Horde/Horde.Build.Tests/TestSetup.cs
--- a/Engine/Source/Programs/Horde/Horde.Build.Tests/TestSetup.cs
+++ b/Engine/Source/Programs/Horde/Horde.Build.Tests/TestSetup.cs
@@ -8,9 +8,7 @@ using System.Threading.Tasks;
 using Datadog.Trace;
 using EpicGames.Core;
 using EpicGames.Horde.Storage;
-using HordeCommon;
 using Horde.Build.Acls;
-using Horde.Build.Storage.Services;
 using Horde.Build.Collections;
 using Horde.Build.Collections.Impl;
 using Horde.Build.Controllers;
@@ -25,10 +23,12 @@ using Horde.Build.Notifications.Impl;
 using Horde.Build.Services;
 using Horde.Build.Services.Impl;
 using Horde.Build.Storage;
+using Horde.Build.Storage.Services;
 using Horde.Build.Tasks.Impl;
-using Horde.Build.Utilities;
 using Horde.Build.Tests.Stubs.Collections;
 using Horde.Build.Tests.Stubs.Services;
+using Horde.Build.Utilities;
+using HordeCommon;
 using Microsoft.AspNetCore.Http;
 using Microsoft.AspNetCore.Mvc;
 using Microsoft.Extensions.Configuration;
@@ -209,7 +209,7 @@ namespace Horde.Build.Tests
 		private JobsController GetJobsController()
         {
 			ILogger<JobsController> logger = ServiceProvider.GetRequiredService<ILogger<JobsController>>();
-			JobsController jobsCtrl = new JobsController(AclService, GraphCollection, PerforceService, StreamService, JobService,
+			JobsController jobsCtrl = new JobsController(GraphCollection, PerforceService, StreamService, JobService,
 		        TemplateCollection, ArtifactCollection, UserCollection, NotificationService, AgentService, logger);
 	        jobsCtrl.ControllerContext = GetControllerContext();
 	        return jobsCtrl;
@@ -217,7 +217,7 @@ namespace Horde.Build.Tests
 
 		private AgentsController GetAgentsController()
 		{
-			AgentsController agentCtrl = new AgentsController(AclService, AgentService, PoolService);
+			AgentsController agentCtrl = new AgentsController(AclService, AgentService);
 			agentCtrl.ControllerContext = GetControllerContext();
 			return agentCtrl;
 		}
@@ -232,7 +232,7 @@ namespace Horde.Build.Tests
 		private LeasesController GetLeasesController()
 		{
 			ILogger<LeasesController> logger = ServiceProvider.GetRequiredService<ILogger<LeasesController>>();
-			LeasesController controller = new LeasesController(AclService, AgentService, logger);
+			LeasesController controller = new LeasesController(AclService, AgentService);
 			controller.ControllerContext = GetControllerContext();
 			return controller;
 		}
diff --git a/Engine/Source/Programs/Horde/Horde.Build/Agents/AgentId.cs b/Engine/Source/Programs/Horde/Horde.Build/Agents/AgentId.cs
--- a/Engine/Source/Programs/Horde/Horde.Build/Agents/AgentId.cs
+++ b/Engine/Source/Programs/Horde/Horde.Build/Agents/AgentId.cs
@@ -1,21 +1,13 @@
 // Copyright Epic Games, Inc. All Rights Reserved.
 
+using System;
+using System.ComponentModel;
+using System.Globalization;
 using EpicGames.Redis;
-using Horde.Build.Models;
 using MongoDB.Bson.Serialization;
 using MongoDB.Bson.Serialization.Attributes;
 using MongoDB.Bson.Serialization.Serializers;
 using StackExchange.Redis;
-using System;
-using System.Collections.Generic;
-using System.ComponentModel;
-using System.Diagnostics.CodeAnalysis;
-using System.Globalization;
-using System.Linq;
-using System.Text;
-using System.Text.Json;
-using System.Text.Json.Serialization;
-using System.Threading.Tasks;
 
 namespace Horde.Build.Utilities
 {
@@ -30,102 +22,102 @@ namespace Horde.Build.Utilities
 		/// <summary>
 		/// The text representing this id
 		/// </summary>
-		readonly string Name;
+		readonly string _name;
 
 		/// <summary>
 		/// Constructor
 		/// </summary>
-		/// <param name="Name">Hostname of the agent</param>
-		public AgentId(string Name)
+		/// <param name="name">Hostname of the agent</param>
+		public AgentId(string name)
 		{
-			if (Name.Length == 0)
+			if (name.Length == 0)
 			{
 				throw new ArgumentException("Agent name may not be empty");
 			}
 
-			char[] Result = new char[Name.Length];
-			for (int Idx = 0; Idx < Name.Length; Idx++)
+			char[] result = new char[name.Length];
+			for (int idx = 0; idx < name.Length; idx++)
 			{
-				char Character = Name[Idx];
-				if ((Character >= 'A' && Character <= 'Z') || (Character >= '0' && Character <= '9') || Character == '_' || Character == '-')
+				char character = name[idx];
+				if ((character >= 'A' && character <= 'Z') || (character >= '0' && character <= '9') || character == '_' || character == '-')
 				{
-					Result[Idx] = Character;
+					result[idx] = character;
 				}
-				else if (Character >= 'a' && Character <= 'z')
+				else if (character >= 'a' && character <= 'z')
 				{
-					Result[Idx] = (char)('A' + (Character - 'a'));
+					result[idx] = (char)('A' + (character - 'a'));
 				}
-				else if (Character == '.')
+				else if (character == '.')
 				{
-					if (Idx == 0)
+					if (idx == 0)
 					{
 						throw new ArgumentException("Agent name may not begin with the '.' character");
 					}
-					else if (Idx == Name.Length - 1)
+					else if (idx == name.Length - 1)
 					{
 						throw new ArgumentException("Agent name may not end with the '.' character");
 					}
-					else if (Name[Idx - 1] == '.')
+					else if (name[idx - 1] == '.')
 					{
 						throw new ArgumentException("Agent name may not contain two consecutive '.' characters");
 					}
 					else
 					{
-						Result[Idx] = Character;
+						result[idx] = character;
 					}
 				}
 				else
 				{
-					throw new ArgumentException($"Character '{Character}' in agent '{Name}' is invalid");
+					throw new ArgumentException($"Character '{character}' in agent '{name}' is invalid");
 				}
 			}
-			this.Name = new string(Result);
+			_name = new string(result);
 		}
 
 		/// <inheritdoc/>
-		public override bool Equals(object? Obj)
+		public override bool Equals(object? obj)
 		{
-			return Obj is AgentId && Equals((AgentId)Obj);
+			return obj is AgentId id && Equals(id);
 		}
 
 		/// <inheritdoc/>
 		public override int GetHashCode()
 		{
-			return Name.GetHashCode(StringComparison.Ordinal);
+			return _name.GetHashCode(StringComparison.Ordinal);
 		}
 
 		/// <inheritdoc/>
-		public bool Equals(AgentId Other)
+		public bool Equals(AgentId other)
 		{
-			return Name.Equals(Other.Name, StringComparison.Ordinal);
+			return _name.Equals(other._name, StringComparison.Ordinal);
 		}
 
 		/// <inheritdoc/>
 		public override string ToString()
 		{
-			return Name;
+			return _name;
 		}
 
 		/// <summary>
 		/// Compares two string ids for equality
 		/// </summary>
-		/// <param name="Left">The first string id</param>
-		/// <param name="Right">Second string id</param>
+		/// <param name="left">The first string id</param>
+		/// <param name="right">Second string id</param>
 		/// <returns>True if the two string ids are equal</returns>
-		public static bool operator ==(AgentId Left, AgentId Right)
+		public static bool operator ==(AgentId left, AgentId right)
 		{
-			return Left.Equals(Right);
+			return left.Equals(right);
 		}
 
 		/// <summary>
 		/// Compares two string ids for inequality
 		/// </summary>
-		/// <param name="Left">The first string id</param>
-		/// <param name="Right">Second string id</param>
+		/// <param name="left">The first string id</param>
+		/// <param name="right">Second string id</param>
 		/// <returns>True if the two string ids are not equal</returns>
-		public static bool operator !=(AgentId Left, AgentId Right)
+		public static bool operator !=(AgentId left, AgentId right)
 		{
-			return !Left.Equals(Right);
+			return !left.Equals(right);
 		}
 	}
 
@@ -135,10 +127,10 @@ namespace Horde.Build.Utilities
 	public sealed class AgentIdRedisConverter : IRedisConverter<AgentId>
 	{
 		/// <inheritdoc/>
-		public AgentId FromRedisValue(RedisValue Value) => new AgentId((string)Value);
+		public AgentId FromRedisValue(RedisValue value) => new AgentId((string)value);
 
 		/// <inheritdoc/>
-		public RedisValue ToRedisValue(AgentId Value) => Value.ToString();
+		public RedisValue ToRedisValue(AgentId value) => value.ToString();
 	}
 
 	/// <summary>
@@ -147,57 +139,56 @@ namespace Horde.Build.Utilities
 	public sealed class AgentIdBsonSerializer : SerializerBase<AgentId>
 	{
 		/// <inheritdoc/>
-		public override AgentId Deserialize(BsonDeserializationContext Context, BsonDeserializationArgs Args)
+		public override AgentId Deserialize(BsonDeserializationContext context, BsonDeserializationArgs args)
 		{
-			string Argument;
-			if (Context.Reader.CurrentBsonType == MongoDB.Bson.BsonType.ObjectId)
+			string argument;
+			if (context.Reader.CurrentBsonType == MongoDB.Bson.BsonType.ObjectId)
 			{
-				Argument = Context.Reader.ReadObjectId().ToString();
+				argument = context.Reader.ReadObjectId().ToString();
 			}
 			else
 			{
-				Argument = Context.Reader.ReadString();
+				argument = context.Reader.ReadString();
 			}
-			return new AgentId(Argument);
+			return new AgentId(argument);
 		}
 
 		/// <inheritdoc/>
-		public override void Serialize(BsonSerializationContext Context, BsonSerializationArgs Args, AgentId Value)
+		public override void Serialize(BsonSerializationContext context, BsonSerializationArgs args, AgentId value)
 		{
-			Context.Writer.WriteString(Value.ToString());
+			context.Writer.WriteString(value.ToString());
 		}
 	}
 
-
 	/// <summary>
 	/// Type converter from strings to PropertyFilter objects
 	/// </summary>
 	sealed class AgentIdTypeConverter : TypeConverter
 	{
 		/// <inheritdoc/>
-		public override bool CanConvertFrom(ITypeDescriptorContext? Context, Type SourceType)
+		public override bool CanConvertFrom(ITypeDescriptorContext? context, Type sourceType)
 		{
-			return SourceType == typeof(string);
+			return sourceType == typeof(string);
 		}
 
 		/// <inheritdoc/>
-		public override object ConvertFrom(ITypeDescriptorContext? Context, CultureInfo? Culture, object Value)
+		public override object ConvertFrom(ITypeDescriptorContext? context, CultureInfo? culture, object value)
 		{
-			return new AgentId((string)Value);
+			return new AgentId((string)value);
 		}
 
 		/// <inheritdoc/>
-		public override bool CanConvertTo(ITypeDescriptorContext? Context, Type? DestinationType)
+		public override bool CanConvertTo(ITypeDescriptorContext? context, Type? destinationType)
 		{
-			return DestinationType == typeof(string);
+			return destinationType == typeof(string);
 		}
 
 		/// <inheritdoc/>
-		public override object? ConvertTo(ITypeDescriptorContext? Context, CultureInfo? Culture, object? Value, Type DestinationType)
+		public override object? ConvertTo(ITypeDescriptorContext? context, CultureInfo? culture, object? value, Type destinationType)
 		{
-			if (DestinationType == typeof(string))
+			if (destinationType == typeof(string))
 			{
-				return Value?.ToString();
+				return value?.ToString();
 			}
 			else
 			{
diff --git a/Engine/Source/Programs/Horde/Horde.Build/Agents/IAgent.cs b/Engine/Source/Programs/Horde/Horde.Build/Agents/IAgent.cs
--- a/Engine/Source/Programs/Horde/Horde.Build/Agents/IAgent.cs
+++ b/Engine/Source/Programs/Horde/Horde.Build/Agents/IAgent.cs
@@ -1,38 +1,30 @@
-﻿// Copyright Epic Games, Inc. All Rights Reserved.
+// Copyright Epic Games, Inc. All Rights Reserved.
 
-using Google.Protobuf.WellKnownTypes;
-using EpicGames.Core;
-using Horde.Build.Acls;
-using Horde.Build.Api;
-using HordeCommon;
-using HordeCommon.Rpc.Tasks;
-using Horde.Build.Services;
-using Horde.Build.Utilities;
-using MongoDB.Bson;
-using MongoDB.Bson.Serialization;
-using MongoDB.Bson.Serialization.Attributes;
 using System;
 using System.Collections.Generic;
-using System.Diagnostics.CodeAnalysis;
-using System.IO;
 using System.Linq;
 using System.Security.Cryptography;
-using Serilog;
-using System.Diagnostics;
 using System.Threading.Tasks;
-using Horde.Build.Collections;
-using EpicGames.Horde.Compute;
+using EpicGames.Core;
 using EpicGames.Horde.Common;
+using EpicGames.Horde.Compute;
+using Google.Protobuf.WellKnownTypes;
+using Horde.Build.Acls;
+using Horde.Build.Services;
+using Horde.Build.Utilities;
+using HordeCommon;
+using HordeCommon.Rpc.Tasks;
+using MongoDB.Bson;
+using MongoDB.Bson.Serialization.Attributes;
 
 namespace Horde.Build.Models
 {
+	using AgentSoftwareChannelName = StringId<AgentSoftwareChannels>;
 	using LeaseId = ObjectId<ILease>;
 	using LogId = ObjectId<ILogFile>;
 	using PoolId = StringId<IPool>;
 	using SessionId = ObjectId<ISession>;
 	using StreamId = StringId<IStream>;
-	using AgentSoftwareVersion = StringId<IAgentSoftwareCollection>;
-	using AgentSoftwareChannelName = StringId<AgentSoftwareChannels>;
 
 	/// <summary>
 	/// Information about a workspace synced to an agent
@@ -67,39 +59,39 @@ namespace Horde.Build.Models
 		/// <summary>
 		/// Whether to use an incremental workspace
 		/// </summary>
-		public bool bIncremental { get; set; }
+		public bool BIncremental { get; set; }
 
 		/// <summary>
 		/// Constructor
 		/// </summary>
-		/// <param name="Cluster">Name of the Perforce cluster</param>
-		/// <param name="UserName">User to log into Perforce with (eg. buildmachine)</param>
-		/// <param name="Identifier">Identifier to distinguish this workspace from other workspaces</param>
-		/// <param name="Stream">The stream to sync</param>
-		/// <param name="View">Custom view for the workspace</param>
+		/// <param name="cluster">Name of the Perforce cluster</param>
+		/// <param name="userName">User to log into Perforce with (eg. buildmachine)</param>
+		/// <param name="identifier">Identifier to distinguish this workspace from other workspaces</param>
+		/// <param name="stream">The stream to sync</param>
+		/// <param name="view">Custom view for the workspace</param>
 		/// <param name="bIncremental">Whether to use an incremental workspace</param>
-		public AgentWorkspace(string? Cluster, string? UserName, string Identifier, string Stream, List<string>? View, bool bIncremental)
+		public AgentWorkspace(string? cluster, string? userName, string identifier, string stream, List<string>? view, bool bIncremental)
 		{
-			if (!String.IsNullOrEmpty(Cluster))
+			if (!String.IsNullOrEmpty(cluster))
 			{
-				this.Cluster = Cluster;
+				Cluster = cluster;
 			}
-			if (!String.IsNullOrEmpty(UserName))
+			if (!String.IsNullOrEmpty(userName))
 			{
-				this.UserName = UserName;
+				UserName = userName;
 			}
-			this.Identifier = Identifier;
-			this.Stream = Stream;
-			this.View = View;
-			this.bIncremental = bIncremental;
+			Identifier = identifier;
+			Stream = stream;
+			View = view;
+			BIncremental = bIncremental;
 		}
 
 		/// <summary>
 		/// Constructor
 		/// </summary>
-		/// <param name="Workspace">RPC message to construct from</param>
-		public AgentWorkspace(HordeCommon.Rpc.Messages.AgentWorkspace Workspace)
-			: this(Workspace.ConfiguredCluster, Workspace.ConfiguredUserName, Workspace.Identifier, Workspace.Stream, (Workspace.View.Count > 0) ? Workspace.View.ToList() : null, Workspace.Incremental)
+		/// <param name="workspace">RPC message to construct from</param>
+		public AgentWorkspace(HordeCommon.Rpc.Messages.AgentWorkspace workspace)
+			: this(workspace.ConfiguredCluster, workspace.ConfiguredUserName, workspace.Identifier, workspace.Stream, (workspace.View.Count > 0) ? workspace.View.ToList() : null, workspace.Incremental)
 		{
 		}
 
@@ -110,27 +102,27 @@ namespace Horde.Build.Models
 		public string GetDigest()
 		{
 #pragma warning disable CA5351 // Do Not Use Broken Cryptographic Algorithms
-			using (MD5 Hasher = MD5.Create())
+			using (MD5 hasher = MD5.Create())
 			{
-				byte[] Data = BsonExtensionMethods.ToBson(this);
-				return BitConverter.ToString(Hasher.ComputeHash(Data)).Replace("-", "", StringComparison.Ordinal);
+				byte[] data = BsonExtensionMethods.ToBson(this);
+				return BitConverter.ToString(hasher.ComputeHash(data)).Replace("-", "", StringComparison.Ordinal);
 			}
 #pragma warning restore CA5351
 		}
 
 		/// <inheritdoc/>
-		public override bool Equals(object? Obj)
+		public override bool Equals(object? obj)
 		{
-			AgentWorkspace? Other = Obj as AgentWorkspace;
-			if (Other == null)
+			AgentWorkspace? other = obj as AgentWorkspace;
+			if (other == null)
 			{
 				return false;
 			}
-			if (Cluster != Other.Cluster || UserName != Other.UserName || Identifier != Other.Identifier || Stream != Other.Stream || bIncremental != Other.bIncremental)
+			if (Cluster != other.Cluster || UserName != other.UserName || Identifier != other.Identifier || Stream != other.Stream || BIncremental != other.BIncremental)
 			{
 				return false;
 			}
-			if (!Enumerable.SequenceEqual(View ?? new List<string>(), Other.View ?? new List<string>()))
+			if (!Enumerable.SequenceEqual(View ?? new List<string>(), other.View ?? new List<string>()))
 			{
 				return false;
 			}
@@ -140,44 +132,44 @@ namespace Horde.Build.Models
 		/// <inheritdoc/>
 		public override int GetHashCode()
 		{
-			return HashCode.Combine(Cluster, UserName, Identifier, Stream, bIncremental); // Ignore 'View' for now
+			return HashCode.Combine(Cluster, UserName, Identifier, Stream, BIncremental); // Ignore 'View' for now
 		}
 
 		/// <summary>
 		/// Checks if two workspace sets are equivalent, ignoring order
 		/// </summary>
-		/// <param name="WorkspacesA">First list of workspaces</param>
-		/// <param name="WorkspacesB">Second list of workspaces</param>
+		/// <param name="workspacesA">First list of workspaces</param>
+		/// <param name="workspacesB">Second list of workspaces</param>
 		/// <returns>True if the sets are equivalent</returns>
-		public static bool SetEquals(IReadOnlyList<AgentWorkspace> WorkspacesA, IReadOnlyList<AgentWorkspace> WorkspacesB)
+		public static bool SetEquals(IReadOnlyList<AgentWorkspace> workspacesA, IReadOnlyList<AgentWorkspace> workspacesB)
 		{
-			HashSet<AgentWorkspace> WorkspacesSetA = new HashSet<AgentWorkspace>(WorkspacesA);
-			return WorkspacesSetA.SetEquals(WorkspacesB);
+			HashSet<AgentWorkspace> workspacesSetA = new HashSet<AgentWorkspace>(workspacesA);
+			return workspacesSetA.SetEquals(workspacesB);
 		}
 
 		/// <summary>
 		/// Converts this workspace to an RPC message
 		/// </summary>
-		/// <param name="Server">The Perforce server</param>
-		/// <param name="Credentials">Credentials for the server</param>
+		/// <param name="server">The Perforce server</param>
+		/// <param name="credentials">Credentials for the server</param>
 		/// <returns>The RPC message</returns>
-		public HordeCommon.Rpc.Messages.AgentWorkspace ToRpcMessage(IPerforceServer Server, PerforceCredentials? Credentials)
+		public HordeCommon.Rpc.Messages.AgentWorkspace ToRpcMessage(IPerforceServer server, PerforceCredentials? credentials)
 		{
 			// Construct the message
-			HordeCommon.Rpc.Messages.AgentWorkspace Result = new HordeCommon.Rpc.Messages.AgentWorkspace();
-			Result.ConfiguredCluster = Cluster;
-			Result.ConfiguredUserName = UserName;
-			Result.ServerAndPort = Server.ServerAndPort;
-			Result.UserName = Credentials?.UserName ?? UserName;
-			Result.Password = Credentials?.Password;
-			Result.Identifier = Identifier;
-			Result.Stream = Stream;
+			HordeCommon.Rpc.Messages.AgentWorkspace result = new HordeCommon.Rpc.Messages.AgentWorkspace();
+			result.ConfiguredCluster = Cluster;
+			result.ConfiguredUserName = UserName;
+			result.ServerAndPort = server.ServerAndPort;
+			result.UserName = credentials?.UserName ?? UserName;
+			result.Password = credentials?.Password;
+			result.Identifier = Identifier;
+			result.Stream = Stream;
 			if (View != null)
 			{
-				Result.View.AddRange(View);
+				result.View.AddRange(View);
 			}
-			Result.Incremental = bIncremental;
-			return Result;
+			result.Incremental = BIncremental;
+			return result;
 		}
 	}
 
@@ -265,26 +257,26 @@ namespace Horde.Build.Models
 		/// <summary>
 		/// Constructor
 		/// </summary>
-		/// <param name="Id">Identifier for the lease</param>
-		/// <param name="Name">Name of this lease</param>
-		/// <param name="StreamId"></param>
-		/// <param name="PoolId"></param>
-		/// <param name="LogId">Unique id for the log</param>
-		/// <param name="State">State for the lease</param>
-		/// <param name="Resources">Resources required for this lease</param>
-		/// <param name="Exclusive">Whether to reserve the entire device</param>
-		/// <param name="Payload">Encoded "any" protobuf describing the contents of the payload</param>
-		public AgentLease(LeaseId Id, string Name, StreamId? StreamId, PoolId? PoolId, LogId? LogId, LeaseState State, IReadOnlyDictionary<string, int>? Resources, bool Exclusive, byte[]? Payload)
+		/// <param name="id">Identifier for the lease</param>
+		/// <param name="name">Name of this lease</param>
+		/// <param name="streamId"></param>
+		/// <param name="poolId"></param>
+		/// <param name="logId">Unique id for the log</param>
+		/// <param name="state">State for the lease</param>
+		/// <param name="resources">Resources required for this lease</param>
+		/// <param name="exclusive">Whether to reserve the entire device</param>
+		/// <param name="payload">Encoded "any" protobuf describing the contents of the payload</param>
+		public AgentLease(LeaseId id, string name, StreamId? streamId, PoolId? poolId, LogId? logId, LeaseState state, IReadOnlyDictionary<string, int>? resources, bool exclusive, byte[]? payload)
 		{
-			this.Id = Id;
-			this.Name = Name;
-			this.StreamId = StreamId;
-			this.PoolId = PoolId;
-			this.LogId = LogId;
-			this.State = State;
-			this.Resources = Resources;
-			this.Exclusive = Exclusive;
-			this.Payload = Payload;
+			Id = id;
+			Name = name;
+			StreamId = streamId;
+			PoolId = poolId;
+			LogId = logId;
+			State = state;
+			Resources = resources;
+			Exclusive = exclusive;
+			Payload = payload;
 			StartTime = DateTime.UtcNow;
 		}
 
@@ -296,8 +288,8 @@ namespace Horde.Build.Models
 		{
 			if (Payload != null)
 			{
-				Any BasePayload = Any.Parser.ParseFrom(Payload);
-				if (BasePayload.Is(ConformTask.Descriptor))
+				Any basePayload = Any.Parser.ParseFrom(Payload);
+				if (basePayload.Is(ConformTask.Descriptor))
 				{
 					return true;
 				}
@@ -311,11 +303,11 @@ namespace Horde.Build.Models
 		/// <returns>RPC message</returns>
 		public HordeCommon.Rpc.Messages.Lease ToRpcMessage()
 		{
-			HordeCommon.Rpc.Messages.Lease Lease = new HordeCommon.Rpc.Messages.Lease();
-			Lease.Id = Id.ToString();
-			Lease.Payload = Google.Protobuf.WellKnownTypes.Any.Parser.ParseFrom(Payload); 
-			Lease.State = State;
-			return Lease;
+			HordeCommon.Rpc.Messages.Lease lease = new HordeCommon.Rpc.Messages.Lease();
+			lease.Id = Id.ToString();
+			lease.Payload = Google.Protobuf.WellKnownTypes.Any.Parser.ParseFrom(Payload); 
+			lease.State = State;
+			return lease;
 		}
 	}
 
@@ -332,7 +324,7 @@ namespace Horde.Build.Models
 		/// <summary>
 		/// The operating system (Linux, MacOS, Windows)
 		/// </summary>
-		public const string OSFamily = "OSFamily";
+		public const string OsFamily = "OSFamily";
 
 		/// <summary>
 		/// Pools that this agent belongs to
@@ -347,7 +339,7 @@ namespace Horde.Build.Models
 		/// <summary>
 		/// Amount of RAM, in GB
 		/// </summary>
-		public const string RAM = "RAM";
+		public const string Ram = "RAM";
 	}
 
 	/// <summary>
@@ -500,83 +492,83 @@ namespace Horde.Build.Models
 		/// Determines whether this agent is online
 		/// </summary>
 		/// <returns></returns>
-		public static bool IsSessionValid(this IAgent Agent, DateTime UtcNow)
+		public static bool IsSessionValid(this IAgent agent, DateTime utcNow)
 		{
-			return Agent.SessionId.HasValue && Agent.SessionExpiresAt.HasValue && UtcNow < Agent.SessionExpiresAt.Value;
+			return agent.SessionId.HasValue && agent.SessionExpiresAt.HasValue && utcNow < agent.SessionExpiresAt.Value;
 		}
 
 		/// <summary>
 		/// Tests whether an agent is in the given pool
 		/// </summary>
-		/// <param name="Agent"></param>
-		/// <param name="PoolId"></param>
+		/// <param name="agent"></param>
+		/// <param name="poolId"></param>
 		/// <returns></returns>
-		public static bool IsInPool(this IAgent Agent, PoolId PoolId)
+		public static bool IsInPool(this IAgent agent, PoolId poolId)
 		{
-			return Agent.DynamicPools.Contains(PoolId) || Agent.ExplicitPools.Contains(PoolId);
+			return agent.DynamicPools.Contains(poolId) || agent.ExplicitPools.Contains(poolId);
 		}
 
 		/// <summary>
 		/// Get all the pools for each agent
 		/// </summary>
-		/// <param name="Agent">The agent to query</param>
+		/// <param name="agent">The agent to query</param>
 		/// <returns></returns>
-		public static IEnumerable<PoolId> GetPools(this IAgent Agent)
+		public static IEnumerable<PoolId> GetPools(this IAgent agent)
 		{
-			foreach (PoolId PoolId in Agent.DynamicPools)
+			foreach (PoolId poolId in agent.DynamicPools)
 			{
-				yield return PoolId;
+				yield return poolId;
 			}
-			foreach (PoolId PoolId in Agent.ExplicitPools)
+			foreach (PoolId poolId in agent.ExplicitPools)
 			{
-				yield return PoolId;
+				yield return poolId;
 			}
 		}
 
 		/// <summary>
 		/// Tests whether an agent has a particular property
 		/// </summary>
-		/// <param name="Agent"></param>
-		/// <param name="Property"></param>
+		/// <param name="agent"></param>
+		/// <param name="property"></param>
 		/// <returns></returns>
-		public static bool HasProperty(this IAgent Agent, string Property)
+		public static bool HasProperty(this IAgent agent, string property)
 		{
-			return Agent.Properties.BinarySearch(Property, StringComparer.OrdinalIgnoreCase) >= 0;
+			return agent.Properties.BinarySearch(property, StringComparer.OrdinalIgnoreCase) >= 0;
 		}
 
 		/// <summary>
 		/// Finds property values from a sorted list of Name=Value pairs
 		/// </summary>
-		/// <param name="Agent">The agent to query</param>
-		/// <param name="Name">Name of the property to find</param>
+		/// <param name="agent">The agent to query</param>
+		/// <param name="name">Name of the property to find</param>
 		/// <returns>Property values</returns>
-		public static IEnumerable<string> GetPropertyValues(this IAgent Agent, string Name)
+		public static IEnumerable<string> GetPropertyValues(this IAgent agent, string name)
 		{
-			if (Name.Equals(KnownPropertyNames.Id, StringComparison.OrdinalIgnoreCase))
+			if (name.Equals(KnownPropertyNames.Id, StringComparison.OrdinalIgnoreCase))
 			{
-				yield return Agent.Id.ToString();
+				yield return agent.Id.ToString();
 			}
-			else if (Name.Equals(KnownPropertyNames.Pool, StringComparison.OrdinalIgnoreCase))
+			else if (name.Equals(KnownPropertyNames.Pool, StringComparison.OrdinalIgnoreCase))
 			{
-				foreach (PoolId PoolId in Agent.GetPools())
+				foreach (PoolId poolId in agent.GetPools())
 				{
-					yield return PoolId.ToString();
+					yield return poolId.ToString();
 				}
 			}
 			else
 			{
-				int Index = Agent.Properties.BinarySearch(Name, StringComparer.OrdinalIgnoreCase);
-				if (Index < 0)
+				int index = agent.Properties.BinarySearch(name, StringComparer.OrdinalIgnoreCase);
+				if (index < 0)
 				{
-					Index = ~Index;
-					for (; Index < Agent.Properties.Count; Index++)
+					index = ~index;
+					for (; index < agent.Properties.Count; index++)
 					{
-						string Property = Agent.Properties[Index];
-						if (Property.Length <= Name.Length || !Property.StartsWith(Name, StringComparison.OrdinalIgnoreCase) || Property[Name.Length] != '=')
+						string property = agent.Properties[index];
+						if (property.Length <= name.Length || !property.StartsWith(name, StringComparison.OrdinalIgnoreCase) || property[name.Length] != '=')
 						{
 							break;
 						}
-						yield return Property.Substring(Name.Length + 1);
+						yield return property.Substring(name.Length + 1);
 					}
 				}
 			}
@@ -585,70 +577,70 @@ namespace Horde.Build.Models
 		/// <summary>
 		/// Evaluates a condition against an agent
 		/// </summary>
-		/// <param name="Agent">The agent to evaluate</param>
-		/// <param name="Condition">The condition to evaluate</param>
+		/// <param name="agent">The agent to evaluate</param>
+		/// <param name="condition">The condition to evaluate</param>
 		/// <returns>True if the agent satisfies the condition</returns>
-		public static bool SatisfiesCondition(this IAgent Agent, Condition Condition)
+		public static bool SatisfiesCondition(this IAgent agent, Condition condition)
 		{
-			return Condition.Evaluate(x => Agent.GetPropertyValues(x));
+			return condition.Evaluate(x => agent.GetPropertyValues(x));
 		}
 
 		/// <summary>
 		/// Determine whether it's possible to add a lease for the given resources
 		/// </summary>
-		/// <param name="Agent">The agent to create a lease for</param>
-		/// <param name="Requirements">Requirements for the lease</param>
+		/// <param name="agent">The agent to create a lease for</param>
+		/// <param name="requirements">Requirements for the lease</param>
 		/// <returns>True if the new lease can be granted</returns>
-		public static bool MeetsRequirements(this IAgent Agent, Requirements Requirements)
+		public static bool MeetsRequirements(this IAgent agent, Requirements requirements)
 		{
-			return MeetsRequirements(Agent, Requirements.Condition, Requirements.Resources, Requirements.Exclusive);
+			return MeetsRequirements(agent, requirements.Condition, requirements.Resources, requirements.Exclusive);
 		}
 
 		/// <summary>
 		/// Determine whether it's possible to add a lease for the given resources
 		/// </summary>
-		/// <param name="Agent">The agent to create a lease for</param>
-		/// <param name="Exclusive">Whether t</param>
-		/// <param name="Condition">Condition to satisfy</param>
-		/// <param name="Resources">Resources required to execute</param>
+		/// <param name="agent">The agent to create a lease for</param>
+		/// <param name="exclusive">Whether t</param>
+		/// <param name="condition">Condition to satisfy</param>
+		/// <param name="resources">Resources required to execute</param>
 		/// <returns>True if the new lease can be granted</returns>
-		public static bool MeetsRequirements(this IAgent Agent, Condition? Condition, Dictionary<string, int>? Resources, bool Exclusive)
+		public static bool MeetsRequirements(this IAgent agent, Condition? condition, Dictionary<string, int>? resources, bool exclusive)
 		{
-			if (!Agent.Enabled || Agent.Status != AgentStatus.Ok)
+			if (!agent.Enabled || agent.Status != AgentStatus.Ok)
 			{
 				return false;
 			}
-			if (Agent.Leases.Any(x => x.Exclusive))
+			if (agent.Leases.Any(x => x.Exclusive))
 			{
 				return false;
 			}
-			if (Exclusive && Agent.Leases.Any())
+			if (exclusive && agent.Leases.Any())
 			{
 				return false;
 			}
-			if (Condition != null && !Agent.SatisfiesCondition(Condition))
+			if (condition != null && !agent.SatisfiesCondition(condition))
 			{
 				return false;
 			}
-			if (Resources != null)
+			if (resources != null)
 			{
-				foreach ((string Name, int Count) in Resources)
+				foreach ((string name, int count) in resources)
 				{
-					int RemainingCount;
-					if (!Agent.Resources.TryGetValue(Name, out RemainingCount))
+					int remainingCount;
+					if (!agent.Resources.TryGetValue(name, out remainingCount))
 					{
 						return false;
 					}
-					foreach (AgentLease Lease in Agent.Leases)
+					foreach (AgentLease lease in agent.Leases)
 					{
-						if (Lease.Resources != null)
+						if (lease.Resources != null)
 						{
-							int LeaseCount;
-							Lease.Resources.TryGetValue(Name, out LeaseCount);
-							RemainingCount -= LeaseCount;
+							int leaseCount;
+							lease.Resources.TryGetValue(name, out leaseCount);
+							remainingCount -= leaseCount;
 						}
 					}
-					if (RemainingCount < Count)
+					if (remainingCount < count)
 					{
 						return false;
 					}
@@ -660,41 +652,41 @@ namespace Horde.Build.Models
 		/// <summary>
 		/// Gets all the autosdk workspaces required for an agent
 		/// </summary>
-		/// <param name="Agent"></param>
-		/// <param name="Globals"></param>
-		/// <param name="Workspaces"></param>
+		/// <param name="agent"></param>
+		/// <param name="globals"></param>
+		/// <param name="workspaces"></param>
 		/// <returns></returns>
-		public static HashSet<AgentWorkspace> GetAutoSdkWorkspaces(this IAgent Agent, Globals Globals, List<AgentWorkspace> Workspaces)
+		public static HashSet<AgentWorkspace> GetAutoSdkWorkspaces(this IAgent agent, Globals globals, List<AgentWorkspace> workspaces)
 		{
-			HashSet<AgentWorkspace> AutoSdkWorkspaces = new HashSet<AgentWorkspace>();
-			foreach (string? ClusterName in Workspaces.Select(x => x.Cluster).Distinct())
+			HashSet<AgentWorkspace> autoSdkWorkspaces = new HashSet<AgentWorkspace>();
+			foreach (string? clusterName in workspaces.Select(x => x.Cluster).Distinct())
 			{
-				PerforceCluster? Cluster = Globals.FindPerforceCluster(ClusterName);
-				if (Cluster != null)
+				PerforceCluster? cluster = globals.FindPerforceCluster(clusterName);
+				if (cluster != null)
 				{
-					AgentWorkspace? AutoSdkWorkspace = GetAutoSdkWorkspace(Agent, Cluster);
-					if (AutoSdkWorkspace != null)
+					AgentWorkspace? autoSdkWorkspace = GetAutoSdkWorkspace(agent, cluster);
+					if (autoSdkWorkspace != null)
 					{
-						AutoSdkWorkspaces.Add(AutoSdkWorkspace);
+						autoSdkWorkspaces.Add(autoSdkWorkspace);
 					}
 				}
 			}
-			return AutoSdkWorkspaces;
+			return autoSdkWorkspaces;
 		}
 
 		/// <summary>
 		/// Get the AutoSDK workspace required for an agent
 		/// </summary>
-		/// <param name="Agent"></param>
-		/// <param name="Cluster">The perforce cluster to get a workspace for</param>
+		/// <param name="agent"></param>
+		/// <param name="cluster">The perforce cluster to get a workspace for</param>
 		/// <returns></returns>
-		public static AgentWorkspace? GetAutoSdkWorkspace(this IAgent Agent, PerforceCluster Cluster)
+		public static AgentWorkspace? GetAutoSdkWorkspace(this IAgent agent, PerforceCluster cluster)
 		{
-			foreach (AutoSdkWorkspace AutoSdk in Cluster.AutoSdk)
+			foreach (AutoSdkWorkspace autoSdk in cluster.AutoSdk)
 			{
-				if (AutoSdk.Stream != null && AutoSdk.Properties.All(x => Agent.Properties.Contains(x)))
+				if (autoSdk.Stream != null && autoSdk.Properties.All(x => agent.Properties.Contains(x)))
 				{
-					return new AgentWorkspace(Cluster.Name, AutoSdk.UserName, AutoSdk.Name ?? "AutoSDK", AutoSdk.Stream!, null, true);
+					return new AgentWorkspace(cluster.Name, autoSdk.UserName, autoSdk.Name ?? "AutoSDK", autoSdk.Stream!, null, true);
 				}
 			}
 			return null;
@@ -703,72 +695,72 @@ namespace Horde.Build.Models
 		/// <summary>
 		/// Converts this workspace to an RPC message
 		/// </summary>
-		/// <param name="Agent">The agent to get a workspace for</param>
-		/// <param name="Workspace">The workspace definition</param>
-		/// <param name="Cluster">The global state</param>
-		/// <param name="LoadBalancer">The Perforce load balancer</param>
-		/// <param name="WorkspaceMessages">List of messages</param>
+		/// <param name="agent">The agent to get a workspace for</param>
+		/// <param name="workspace">The workspace definition</param>
+		/// <param name="cluster">The global state</param>
+		/// <param name="loadBalancer">The Perforce load balancer</param>
+		/// <param name="workspaceMessages">List of messages</param>
 		/// <returns>The RPC message</returns>
-		public static async Task<bool> TryAddWorkspaceMessage(this IAgent Agent, AgentWorkspace Workspace, PerforceCluster Cluster, PerforceLoadBalancer LoadBalancer, IList<HordeCommon.Rpc.Messages.AgentWorkspace> WorkspaceMessages)
+		public static async Task<bool> TryAddWorkspaceMessage(this IAgent agent, AgentWorkspace workspace, PerforceCluster cluster, PerforceLoadBalancer loadBalancer, IList<HordeCommon.Rpc.Messages.AgentWorkspace> workspaceMessages)
 		{
 			// Find a matching server, trying to use a previously selected one if possible
-			string? BaseServerAndPort;
-			string? ServerAndPort;
+			string? baseServerAndPort;
+			string? serverAndPort;
 
-			HordeCommon.Rpc.Messages.AgentWorkspace? ExistingWorkspace = WorkspaceMessages.FirstOrDefault(x => x.ConfiguredCluster == Workspace.Cluster);
-			if(ExistingWorkspace != null)
+			HordeCommon.Rpc.Messages.AgentWorkspace? existingWorkspace = workspaceMessages.FirstOrDefault(x => x.ConfiguredCluster == workspace.Cluster);
+			if(existingWorkspace != null)
 			{
-				BaseServerAndPort = ExistingWorkspace.BaseServerAndPort;
-				ServerAndPort = ExistingWorkspace.ServerAndPort;
+				baseServerAndPort = existingWorkspace.BaseServerAndPort;
+				serverAndPort = existingWorkspace.ServerAndPort;
 			}
 			else
 			{
-				if (Cluster == null)
+				if (cluster == null)
 				{
 					return false;
 				}
 
-				IPerforceServer? Server = await LoadBalancer.SelectServerAsync(Cluster, Agent);
-				if (Server == null)
+				IPerforceServer? server = await loadBalancer.SelectServerAsync(cluster, agent);
+				if (server == null)
 				{
 					return false;
 				}
 
-				BaseServerAndPort = Server.BaseServerAndPort;
-				ServerAndPort = Server.ServerAndPort;
+				baseServerAndPort = server.BaseServerAndPort;
+				serverAndPort = server.ServerAndPort;
 			}
 
 			// Find the matching credentials for the desired user
-			PerforceCredentials? Credentials = null;
-			if (Cluster != null)
+			PerforceCredentials? credentials = null;
+			if (cluster != null)
 			{
-				if (Workspace.UserName == null)
+				if (workspace.UserName == null)
 				{
-					Credentials = Cluster.Credentials.FirstOrDefault();
+					credentials = cluster.Credentials.FirstOrDefault();
 				}
 				else
 				{
-					Credentials = Cluster.Credentials.FirstOrDefault(x => String.Equals(x.UserName, Workspace.UserName, StringComparison.OrdinalIgnoreCase));
+					credentials = cluster.Credentials.FirstOrDefault(x => String.Equals(x.UserName, workspace.UserName, StringComparison.OrdinalIgnoreCase));
 				}
 			}
 
 			// Construct the message
-			HordeCommon.Rpc.Messages.AgentWorkspace Result = new HordeCommon.Rpc.Messages.AgentWorkspace();
-			Result.ConfiguredCluster = Workspace.Cluster;
-			Result.ConfiguredUserName = Workspace.UserName;
-			Result.Cluster = Cluster?.Name;
-			Result.BaseServerAndPort = BaseServerAndPort;
-			Result.ServerAndPort = ServerAndPort;
-			Result.UserName = Credentials?.UserName ?? Workspace.UserName;
-			Result.Password = Credentials?.Password;
-			Result.Identifier = Workspace.Identifier;
-			Result.Stream = Workspace.Stream;
-			if (Workspace.View != null)
+			HordeCommon.Rpc.Messages.AgentWorkspace result = new HordeCommon.Rpc.Messages.AgentWorkspace();
+			result.ConfiguredCluster = workspace.Cluster;
+			result.ConfiguredUserName = workspace.UserName;
+			result.Cluster = cluster?.Name;
+			result.BaseServerAndPort = baseServerAndPort;
+			result.ServerAndPort = serverAndPort;
+			result.UserName = credentials?.UserName ?? workspace.UserName;
+			result.Password = credentials?.Password;
+			result.Identifier = workspace.Identifier;
+			result.Stream = workspace.Stream;
+			if (workspace.View != null)
 			{
-				Result.View.AddRange(Workspace.View);
+				result.View.AddRange(workspace.View);
 			}
-			Result.Incremental = Workspace.bIncremental;
-			WorkspaceMessages.Add(Result);
+			result.Incremental = workspace.BIncremental;
+			workspaceMessages.Add(result);
 
 			return true;
 		}
diff --git a/Engine/Source/Programs/Horde/Horde.Build/Agents/IAgentCollection.cs b/Engine/Source/Programs/Horde/Horde.Build/Agents/IAgentCollection.cs
--- a/Engine/Source/Programs/Horde/Horde.Build/Agents/IAgentCollection.cs
+++ b/Engine/Source/Programs/Horde/Horde.Build/Agents/IAgentCollection.cs
@@ -1,26 +1,18 @@
 ﻿// Copyright Epic Games, Inc. All Rights Reserved.
 
-using Google.Protobuf.WellKnownTypes;
+using System;
+using System.Collections.Generic;
+using System.Threading.Tasks;
 using Horde.Build.Acls;
-using Horde.Build.Api;
-using HordeCommon;
 using Horde.Build.Models;
-using HordeCommon.Rpc.Tasks;
 using Horde.Build.Services;
 using Horde.Build.Utilities;
-using MongoDB.Bson;
-using MongoDB.Bson.Serialization.Attributes;
-using MongoDB.Driver;
-using System;
-using System.Collections.Generic;
-using System.Linq;
-using System.Linq.Expressions;
-using System.Threading.Tasks;
+using HordeCommon;
 
 namespace Horde.Build.Collections
 {
-	using PoolId = StringId<IPool>;
 	using AgentSoftwareChannelName = StringId<AgentSoftwareChannels>;
+	using PoolId = StringId<IPool>;
 	using SessionId = ObjectId<ISession>;
 
 	/// <summary>
@@ -31,140 +23,140 @@ namespace Horde.Build.Collections
 		/// <summary>
 		/// Adds a new agent with the given properties
 		/// </summary>
-		/// <param name="Id">Id for the new agent</param>
+		/// <param name="id">Id for the new agent</param>
 		/// <param name="bEnabled">Whether the agent is enabled or not</param>
-		/// <param name="Channel">Channel to use for software run by this agent</param>
-		/// <param name="Pools">Pools for the agent</param>
-		Task<IAgent> AddAsync(AgentId Id, bool bEnabled, AgentSoftwareChannelName? Channel = null, List<PoolId>? Pools = null);
+		/// <param name="channel">Channel to use for software run by this agent</param>
+		/// <param name="pools">Pools for the agent</param>
+		Task<IAgent> AddAsync(AgentId id, bool bEnabled, AgentSoftwareChannelName? channel = null, List<PoolId>? pools = null);
 
 		/// <summary>
 		/// Deletes an agent
 		/// </summary>
-		/// <param name="Agent">Deletes the agent</param>
+		/// <param name="agent">Deletes the agent</param>
 		/// <returns>Async task</returns>
-		Task<IAgent?> TryDeleteAsync(IAgent Agent);
+		Task<IAgent?> TryDeleteAsync(IAgent agent);
 
 		/// <summary>
 		/// Deletes an agent
 		/// </summary>
-		/// <param name="AgentId">Unique id of the agent</param>
+		/// <param name="agentId">Unique id of the agent</param>
 		/// <returns>Async task</returns>
-		Task ForceDeleteAsync(AgentId AgentId);
+		Task ForceDeleteAsync(AgentId agentId);
 
 		/// <summary>
 		/// Gets an agent by ID
 		/// </summary>
-		/// <param name="AgentId">Unique id of the agent</param>
+		/// <param name="agentId">Unique id of the agent</param>
 		/// <returns>The agent document</returns>
-		Task<IAgent?> GetAsync(AgentId AgentId);
+		Task<IAgent?> GetAsync(AgentId agentId);
 
 		/// <summary>
 		/// Finds all agents matching certain criteria
 		/// </summary>
-		/// <param name="PoolId">The pool ID in string form containing the agent</param>
-		/// <param name="ModifiedAfter">If set, only returns agents modified after this time</param>
-		/// <param name="Status">Status to look for</param>
-		/// <param name="Index">Index of the first result</param>
-		/// <param name="Count">Number of results to return</param>
+		/// <param name="poolId">The pool ID in string form containing the agent</param>
+		/// <param name="modifiedAfter">If set, only returns agents modified after this time</param>
+		/// <param name="status">Status to look for</param>
+		/// <param name="index">Index of the first result</param>
+		/// <param name="count">Number of results to return</param>
 		/// <returns>List of agents matching the given criteria</returns>
-		Task<List<IAgent>> FindAsync(PoolId? PoolId = null, DateTime? ModifiedAfter = null, AgentStatus? Status = null, int? Index = null, int? Count = null);
+		Task<List<IAgent>> FindAsync(PoolId? poolId = null, DateTime? modifiedAfter = null, AgentStatus? status = null, int? index = null, int? count = null);
 
 		/// <summary>
 		/// Finds all the expried agents
 		/// </summary>
-		/// <param name="UtcNow">The current time</param>
-		/// <param name="MaxAgents">Maximum number of agents to return</param>
+		/// <param name="utcNow">The current time</param>
+		/// <param name="maxAgents">Maximum number of agents to return</param>
 		/// <returns>List of agents</returns>
-		Task<List<IAgent>> FindExpiredAsync(DateTime UtcNow, int MaxAgents);
+		Task<List<IAgent>> FindExpiredAsync(DateTime utcNow, int maxAgents);
 
 		/// <summary>
 		/// Update an agent's settings
 		/// </summary>
-		/// <param name="Agent">Agent instance</param>
+		/// <param name="agent">Agent instance</param>
 		/// <param name="bEnabled">Whether the agent is enabled or not</param>
 		/// <param name="bRequestConform">Whether to request a conform job be run</param>
 		/// <param name="bRequestFullConform">Whether to request a full conform job be run</param>
 		/// <param name="bRequestRestart">Whether to request the machine be restarted</param>
 		/// <param name="bRequestShutdown">Whether to request the machine be shut down</param>
-		/// <param name="ShutdownReason">The reason for shutting down agent, ex. Autoscaler/Manual/Unexpected</param>
-		/// <param name="Channel">Override for the desired software channel</param>
-		/// <param name="Pools">List of pools for the agent</param>
-		/// <param name="Acl">New ACL for this agent</param>
-		/// <param name="Comment">New comment</param>
+		/// <param name="shutdownReason">The reason for shutting down agent, ex. Autoscaler/Manual/Unexpected</param>
+		/// <param name="channel">Override for the desired software channel</param>
+		/// <param name="pools">List of pools for the agent</param>
+		/// <param name="acl">New ACL for this agent</param>
+		/// <param name="comment">New comment</param>
 		/// <returns>Version of the software that needs to be installed on the agent. Null if the agent is running the correct version.</returns>
-		Task<IAgent?> TryUpdateSettingsAsync(IAgent Agent, bool? bEnabled = null, bool? bRequestConform = null, bool? bRequestFullConform = null, bool? bRequestRestart = null, bool? bRequestShutdown = null, string? ShutdownReason = null, AgentSoftwareChannelName? Channel = null, List<PoolId>? Pools = null, Acl? Acl = null, string? Comment = null);
+		Task<IAgent?> TryUpdateSettingsAsync(IAgent agent, bool? bEnabled = null, bool? bRequestConform = null, bool? bRequestFullConform = null, bool? bRequestRestart = null, bool? bRequestShutdown = null, string? shutdownReason = null, AgentSoftwareChannelName? channel = null, List<PoolId>? pools = null, Acl? acl = null, string? comment = null);
 
 		/// <summary>
 		/// Update the current workspaces for an agent.
 		/// </summary>
-		/// <param name="Agent">The agent to update</param>
-		/// <param name="Workspaces">Current list of workspaces</param>
-		/// <param name="RequestConform">Whether the agent still needs to run another conform</param>
+		/// <param name="agent">The agent to update</param>
+		/// <param name="workspaces">Current list of workspaces</param>
+		/// <param name="requestConform">Whether the agent still needs to run another conform</param>
 		/// <returns>New agent state</returns>
-		Task<IAgent?> TryUpdateWorkspacesAsync(IAgent Agent, List<AgentWorkspace> Workspaces, bool RequestConform);
+		Task<IAgent?> TryUpdateWorkspacesAsync(IAgent agent, List<AgentWorkspace> workspaces, bool requestConform);
 
 		/// <summary>
 		/// Sets the current session
 		/// </summary>
-		/// <param name="Agent">The agent to update</param>
-		/// <param name="SessionId">New session id</param>
-		/// <param name="SessionExpiresAt">Expiry time for the new session</param>
-		/// <param name="Status">Status of the agent</param>
-		/// <param name="Properties">Properties for the current session</param>
-		/// <param name="Resources">Resources for the agent</param>
-		/// <param name="DynamicPools">New list of dynamic pools for the agent</param>
-		/// <param name="Version">Current version of the agent software</param>
+		/// <param name="agent">The agent to update</param>
+		/// <param name="sessionId">New session id</param>
+		/// <param name="sessionExpiresAt">Expiry time for the new session</param>
+		/// <param name="status">Status of the agent</param>
+		/// <param name="properties">Properties for the current session</param>
+		/// <param name="resources">Resources for the agent</param>
+		/// <param name="dynamicPools">New list of dynamic pools for the agent</param>
+		/// <param name="version">Current version of the agent software</param>
 		/// <returns>New agent state</returns>
-		Task<IAgent?> TryStartSessionAsync(IAgent Agent, SessionId SessionId, DateTime SessionExpiresAt, AgentStatus Status, IReadOnlyList<string> Properties, IReadOnlyDictionary<string, int> Resources, IReadOnlyList<PoolId> DynamicPools, string? Version);
+		Task<IAgent?> TryStartSessionAsync(IAgent agent, SessionId sessionId, DateTime sessionExpiresAt, AgentStatus status, IReadOnlyList<string> properties, IReadOnlyDictionary<string, int> resources, IReadOnlyList<PoolId> dynamicPools, string? version);
 
 		/// <summary>
 		/// Attempt to update the agent state
 		/// </summary>
-		/// <param name="Agent">Agent instance</param>
-		/// <param name="Status">New status of the agent</param>
-		/// <param name="SessionExpiresAt">New expiry time for the current session</param>
-		/// <param name="Properties">Properties for the current session</param>
-		/// <param name="Resources">Resources for the agent</param>
-		/// <param name="DynamicPools">New list of dynamic pools for the agent</param>
-		/// <param name="Leases">New set of leases</param>
+		/// <param name="agent">Agent instance</param>
+		/// <param name="status">New status of the agent</param>
+		/// <param name="sessionExpiresAt">New expiry time for the current session</param>
+		/// <param name="properties">Properties for the current session</param>
+		/// <param name="resources">Resources for the agent</param>
+		/// <param name="dynamicPools">New list of dynamic pools for the agent</param>
+		/// <param name="leases">New set of leases</param>
 		/// <returns>True if the document was updated, false if another writer updated the document first</returns>
-		Task<IAgent?> TryUpdateSessionAsync(IAgent Agent, AgentStatus? Status, DateTime? SessionExpiresAt, IReadOnlyList<string>? Properties, IReadOnlyDictionary<string, int>? Resources, IReadOnlyList<PoolId>? DynamicPools, List<AgentLease>? Leases);
+		Task<IAgent?> TryUpdateSessionAsync(IAgent agent, AgentStatus? status, DateTime? sessionExpiresAt, IReadOnlyList<string>? properties, IReadOnlyDictionary<string, int>? resources, IReadOnlyList<PoolId>? dynamicPools, List<AgentLease>? leases);
 
 		/// <summary>
 		/// Terminates the current session
 		/// </summary>
-		/// <param name="Agent">The agent to terminate the session for</param>
+		/// <param name="agent">The agent to terminate the session for</param>
 		/// <returns>New agent state if it succeeded, otherwise null</returns>
-		Task<IAgent?> TryTerminateSessionAsync(IAgent Agent);
+		Task<IAgent?> TryTerminateSessionAsync(IAgent agent);
 
 		/// <summary>
 		/// Attempts to add a lease to an agent
 		/// </summary>
-		/// <param name="Agent">The agent to update</param>
-		/// <param name="NewLease">The new lease document</param>
+		/// <param name="agent">The agent to update</param>
+		/// <param name="newLease">The new lease document</param>
 		/// <returns>New agent state if it succeeded, otherwise null</returns>
-		Task<IAgent?> TryAddLeaseAsync(IAgent Agent, AgentLease NewLease);
+		Task<IAgent?> TryAddLeaseAsync(IAgent agent, AgentLease newLease);
 
 		/// <summary>
 		/// Attempts to cancel a lease
 		/// </summary>
-		/// <param name="Agent">The agent to update</param>
-		/// <param name="LeaseIdx">Index of the lease</param>
+		/// <param name="agent">The agent to update</param>
+		/// <param name="leaseIdx">Index of the lease</param>
 		/// <returns>New agent state if it succeeded, otherwise null</returns>
-		Task<IAgent?> TryCancelLeaseAsync(IAgent Agent, int LeaseIdx);
+		Task<IAgent?> TryCancelLeaseAsync(IAgent agent, int leaseIdx);
 
 		/// <summary>
 		/// Gets the log channel for an agent
 		/// </summary>
-		/// <param name="AgentId"></param>
+		/// <param name="agentId"></param>
 		/// <returns></returns>
-		IAuditLogChannel<AgentId> GetLogger(AgentId AgentId);
+		IAuditLogChannel<AgentId> GetLogger(AgentId agentId);
 
 		/// <summary>
 		/// Subscribe to notifications on agent states being updated
 		/// </summary>
-		/// <param name="OnUpdate">Callback for updates</param>
+		/// <param name="onUpdate">Callback for updates</param>
 		/// <returns>Disposable subscription object</returns>
-		Task<IDisposable> SubscribeToUpdateEventsAsync(Action<AgentId> OnUpdate);
+		Task<IDisposable> SubscribeToUpdateEventsAsync(Action<AgentId> onUpdate);
 	}
 }
diff --git a/Engine/Source/Programs/Horde/Horde.Build/Agents/IAgentSoftwareChannel.cs b/Engine/Source/Programs/Horde/Horde.Build/Agents/IAgentSoftwareChannel.cs
--- a/Engine/Source/Programs/Horde/Horde.Build/Agents/IAgentSoftwareChannel.cs
+++ b/Engine/Source/Programs/Horde/Horde.Build/Agents/IAgentSoftwareChannel.cs
@@ -1,22 +1,11 @@
 // Copyright Epic Games, Inc. All Rights Reserved.
 
-using Horde.Build.Api;
-using Horde.Build.Collections;
+using System;
 using Horde.Build.Services;
 using Horde.Build.Utilities;
-using MongoDB.Bson;
-using MongoDB.Bson.Serialization.Attributes;
-using System;
-using System.Collections.Generic;
-using System.IO;
-using System.IO.Compression;
-using System.Linq;
-using System.Text.Json;
-using System.Threading.Tasks;
 
 namespace Horde.Build.Models
 {
-	using AgentSoftwareVersion = StringId<IAgentSoftwareCollection>;
 	using AgentSoftwareChannelName = StringId<AgentSoftwareChannels>;
 
 	/// <summary>
diff --git a/Engine/Source/Programs/Horde/Horde.Build/Agents/IAgentSoftwareCollection.cs b/Engine/Source/Programs/Horde/Horde.Build/Agents/IAgentSoftwareCollection.cs
--- a/Engine/Source/Programs/Horde/Horde.Build/Agents/IAgentSoftwareCollection.cs
+++ b/Engine/Source/Programs/Horde/Horde.Build/Agents/IAgentSoftwareCollection.cs
@@ -1,19 +1,5 @@
 // Copyright Epic Games, Inc. All Rights Reserved.
 
-using EpicGames.Core;
-using HordeCommon;
-using Horde.Build.Models;
-using Horde.Build.Services;
-using MongoDB.Bson;
-using MongoDB.Bson.Serialization.Attributes;
-using MongoDB.Driver;
-using System;
-using System.Collections.Generic;
-using System.Globalization;
-using System.IO;
-using System.IO.Compression;
-using System.Linq;
-using System.Text.Json;
 using System.Threading.Tasks;
 
 namespace Horde.Build.Collections
@@ -26,29 +12,29 @@ namespace Horde.Build.Collections
 		/// <summary>
 		/// Adds a new software revision
 		/// </summary>
-		/// <param name="Version">The version number</param>
-		/// <param name="Data">Zip file containing the new software</param>
+		/// <param name="version">The version number</param>
+		/// <param name="data">Zip file containing the new software</param>
 		/// <returns>New software instance</returns>
-		Task<bool> AddAsync(string Version, byte[] Data);
+		Task<bool> AddAsync(string version, byte[] data);
 
 		/// <summary>
 		/// Tests whether a given version exists
 		/// </summary>
-		/// <param name="Version">Version of the software</param>
+		/// <param name="version">Version of the software</param>
 		/// <returns>True if it exists, false otherwise</returns>
-		Task<bool> ExistsAsync(string Version);
+		Task<bool> ExistsAsync(string version);
 
 		/// <summary>
 		/// Removes a software archive 
 		/// </summary>
-		/// <param name="Version">Version of the software to delete</param>
-		Task<bool> RemoveAsync(string Version);
+		/// <param name="version">Version of the software to delete</param>
+		Task<bool> RemoveAsync(string version);
 
 		/// <summary>
 		/// Downloads software of a given revision
 		/// </summary>
-		/// <param name="Version">Version of the software</param>
+		/// <param name="version">Version of the software</param>
 		/// <returns>Data for the given software</returns>
-		Task<byte[]?> GetAsync(string Version);
+		Task<byte[]?> GetAsync(string version);
 	}
 }
diff --git a/Engine/Source/Programs/Horde/Horde.Build/Agents/ILease.cs b/Engine/Source/Programs/Horde/Horde.Build/Agents/ILease.cs
--- a/Engine/Source/Programs/Horde/Horde.Build/Agents/ILease.cs
+++ b/Engine/Source/Programs/Horde/Horde.Build/Agents/ILease.cs
@@ -1,17 +1,10 @@
 // Copyright Epic Games, Inc. All Rights Reserved.
 
-using Horde.Build.Api;
-using HordeCommon;
-using Horde.Build.Utilities;
-using MongoDB.Bson;
-using MongoDB.Bson.Serialization.Attributes;
-using Serilog;
 using System;
-using System.Collections.Generic;
-using System.Linq;
-using System.Threading.Tasks;
-using Google.Protobuf.WellKnownTypes;
 using Google.Protobuf;
+using Google.Protobuf.WellKnownTypes;
+using Horde.Build.Utilities;
+using HordeCommon;
 
 namespace Horde.Build.Models
 {
@@ -95,22 +88,22 @@ namespace Horde.Build.Models
 		/// <summary>
 		/// Gets the task from a lease, encoded as an Any protobuf object
 		/// </summary>
-		/// <param name="Lease">The lease to query</param>
+		/// <param name="lease">The lease to query</param>
 		/// <returns>The task definition encoded as a protobuf Any object</returns>
-		public static Any GetTask(this ILease Lease)
+		public static Any GetTask(this ILease lease)
 		{
-			return Any.Parser.ParseFrom(Lease.Payload.ToArray());
+			return Any.Parser.ParseFrom(lease.Payload.ToArray());
 		}
 
 		/// <summary>
 		/// Gets a typed task object from a lease
 		/// </summary>
 		/// <typeparam name="T">Type of the protobuf message to return</typeparam>
-		/// <param name="Lease">The lease to query</param>
+		/// <param name="lease">The lease to query</param>
 		/// <returns>The task definition</returns>
-		public static T GetTask<T>(this ILease Lease) where T : IMessage<T>, new()
+		public static T GetTask<T>(this ILease lease) where T : IMessage<T>, new()
 		{
-			return GetTask(Lease).Unpack<T>();
+			return GetTask(lease).Unpack<T>();
 		}
 	}
 }
