Horde: Fix coding conventions within Horde.Build.

diff --git a/Engine/Source/Programs/Horde/Horde.Build/Agents/ILeaseCollection.cs b/Engine/Source/Programs/Horde/Horde.Build/Agents/ILeaseCollection.cs
--- a/Engine/Source/Programs/Horde/Horde.Build/Agents/ILeaseCollection.cs
+++ b/Engine/Source/Programs/Horde/Horde.Build/Agents/ILeaseCollection.cs
@@ -1,16 +1,11 @@
 // Copyright Epic Games, Inc. All Rights Reserved.
 
-using HordeCommon;
-using Horde.Build.Models;
-using Horde.Build.Services;
-using Horde.Build.Utilities;
-using MongoDB.Bson;
-using MongoDB.Bson.Serialization.Attributes;
-using MongoDB.Driver;
 using System;
 using System.Collections.Generic;
-using System.Linq;
 using System.Threading.Tasks;
+using Horde.Build.Models;
+using Horde.Build.Utilities;
+using HordeCommon;
 
 namespace Horde.Build.Collections
 {
@@ -28,83 +23,83 @@ namespace Horde.Build.Collections
 		/// <summary>
 		/// Adds a lease to the collection
 		/// </summary>
-		/// <param name="Id">The lease id</param>
-		/// <param name="Name">Name of the lease</param>
-		/// <param name="AgentId">The agent id</param>
-		/// <param name="SessionId">The agent session handling the lease</param>
-		/// <param name="StreamId">Stream for the payload</param>
-		/// <param name="PoolId">The pool for the lease</param>
-		/// <param name="LogId">Log id for the lease</param>
-		/// <param name="StartTime">Start time of the lease</param>
-		/// <param name="Payload">Payload for the lease</param>
+		/// <param name="id">The lease id</param>
+		/// <param name="name">Name of the lease</param>
+		/// <param name="agentId">The agent id</param>
+		/// <param name="sessionId">The agent session handling the lease</param>
+		/// <param name="streamId">Stream for the payload</param>
+		/// <param name="poolId">The pool for the lease</param>
+		/// <param name="logId">Log id for the lease</param>
+		/// <param name="startTime">Start time of the lease</param>
+		/// <param name="payload">Payload for the lease</param>
 		/// <returns>Async task</returns>
-		Task<ILease> AddAsync(LeaseId Id, string Name, AgentId AgentId, SessionId SessionId, StreamId? StreamId, PoolId? PoolId, LogId? LogId, DateTime StartTime, byte[] Payload);
+		Task<ILease> AddAsync(LeaseId id, string name, AgentId agentId, SessionId sessionId, StreamId? streamId, PoolId? poolId, LogId? logId, DateTime startTime, byte[] payload);
 
 		/// <summary>
 		/// Deletes a lease from the collection
 		/// </summary>
-		/// <param name="LeaseId">Unique id of the lease</param>
+		/// <param name="leaseId">Unique id of the lease</param>
 		/// <returns>Async task</returns>
-		Task DeleteAsync(LeaseId LeaseId);
+		Task DeleteAsync(LeaseId leaseId);
 
 		/// <summary>
 		/// Gets a specific lease
 		/// </summary>
-		/// <param name="LeaseId">Unique id of the lease</param>
+		/// <param name="leaseId">Unique id of the lease</param>
 		/// <returns>The lease that was found, or null if it does not exist</returns>
-		Task<ILease?> GetAsync(LeaseId LeaseId);
+		Task<ILease?> GetAsync(LeaseId leaseId);
 
 		/// <summary>
 		/// Finds all leases matching a set of criteria
 		/// </summary>
-		/// <param name="AgentId">Unqiue id of the agent executing this lease</param>
-		/// <param name="SessionId">Unique id of the agent session</param>
-		/// <param name="MinTime">Start of the window to include leases</param>
-		/// <param name="MaxTime">End of the window to include leases</param>
-		/// <param name="Index">Index of the first result to return</param>
-		/// <param name="Count">Number of results to return</param>
-		/// <param name="IndexHint">Name of index to be specified as a hint to the database query planner</param>
-		/// <param name="ConsistentRead">If the database read should be made to the replica server</param>
+		/// <param name="agentId">Unqiue id of the agent executing this lease</param>
+		/// <param name="sessionId">Unique id of the agent session</param>
+		/// <param name="minTime">Start of the window to include leases</param>
+		/// <param name="maxTime">End of the window to include leases</param>
+		/// <param name="index">Index of the first result to return</param>
+		/// <param name="count">Number of results to return</param>
+		/// <param name="indexHint">Name of index to be specified as a hint to the database query planner</param>
+		/// <param name="consistentRead">If the database read should be made to the replica server</param>
 		/// <returns>List of leases matching the given criteria</returns>
-		Task<List<ILease>> FindLeasesAsync(AgentId? AgentId = null, SessionId? SessionId = null, DateTime? MinTime = null, DateTime? MaxTime = null, int? Index = null, int? Count = null, string? IndexHint = null, bool ConsistentRead = true);
+		Task<List<ILease>> FindLeasesAsync(AgentId? agentId = null, SessionId? sessionId = null, DateTime? minTime = null, DateTime? maxTime = null, int? index = null, int? count = null, string? indexHint = null, bool consistentRead = true);
 
 		/// <summary>
 		/// Finds all leases by finish time
 		/// </summary>
-		/// <param name="MinFinishTime">Start of finish time window</param>
-		/// <param name="MaxFinishTime">End of finish time window</param>
-		/// <param name="Index">Index of the first result to return</param>
-		/// <param name="Count">Number of results to return</param>
-		/// <param name="IndexHint">Name of index to be specified as a hint to the database query planner</param>
-		/// <param name="ConsistentRead">If the database read should be made to the replica server</param>
+		/// <param name="minFinishTime">Start of finish time window</param>
+		/// <param name="maxFinishTime">End of finish time window</param>
+		/// <param name="index">Index of the first result to return</param>
+		/// <param name="count">Number of results to return</param>
+		/// <param name="indexHint">Name of index to be specified as a hint to the database query planner</param>
+		/// <param name="consistentRead">If the database read should be made to the replica server</param>
 		/// <returns>List of leases matching the given criteria</returns>
-		Task<List<ILease>> FindLeasesByFinishTimeAsync(DateTime? MinFinishTime, DateTime? MaxFinishTime, int? Index, int? Count, string? IndexHint = null, bool ConsistentRead = true);
+		Task<List<ILease>> FindLeasesByFinishTimeAsync(DateTime? minFinishTime, DateTime? maxFinishTime, int? index, int? count, string? indexHint = null, bool consistentRead = true);
 
 		/// <summary>
 		/// Finds all leases between min and/or max time
 		/// Preferred over the more generic FindLeasesAsync as this one use a better index and relaxed read consistency
 		/// </summary>
-		/// <param name="MinTime">Start of the window to include leases</param>
-		/// <param name="MaxTime">End of the window to include leases</param>
+		/// <param name="minTime">Start of the window to include leases</param>
+		/// <param name="maxTime">End of the window to include leases</param>
 		/// <returns>List of leases matching the given criteria</returns>
-		Task<List<ILease>> FindLeasesAsync(DateTime? MinTime, DateTime? MaxTime);
+		Task<List<ILease>> FindLeasesAsync(DateTime? minTime, DateTime? maxTime);
 
 		/// <summary>
 		/// Finds all active leases
 		/// </summary>
-		/// <param name="Index">Index of the first result to return</param>
-		/// <param name="Count">Number of results to return</param>
+		/// <param name="index">Index of the first result to return</param>
+		/// <param name="count">Number of results to return</param>
 		/// <returns>List of leases</returns>
-		Task<List<ILease>> FindActiveLeasesAsync(int? Index = null, int? Count = null);
+		Task<List<ILease>> FindActiveLeasesAsync(int? index = null, int? count = null);
 
 		/// <summary>
 		/// Sets the outcome of a lease
 		/// </summary>
-		/// <param name="LeaseId">The lease to update</param>
-		/// <param name="FinishTime">Time at which the lease finished</param>
-		/// <param name="Outcome">Outcome of the lease</param>
-		/// <param name="Output">Output data from the task</param>
+		/// <param name="leaseId">The lease to update</param>
+		/// <param name="finishTime">Time at which the lease finished</param>
+		/// <param name="outcome">Outcome of the lease</param>
+		/// <param name="output">Output data from the task</param>
 		/// <returns>True if the lease was updated, false otherwise</returns>
-		Task<bool> TrySetOutcomeAsync(LeaseId LeaseId, DateTime FinishTime, LeaseOutcome Outcome, byte[]? Output);
+		Task<bool> TrySetOutcomeAsync(LeaseId leaseId, DateTime finishTime, LeaseOutcome outcome, byte[]? output);
 	}
 }
diff --git a/Engine/Source/Programs/Horde/Horde.Build/Agents/ISession.cs b/Engine/Source/Programs/Horde/Horde.Build/Agents/ISession.cs
--- a/Engine/Source/Programs/Horde/Horde.Build/Agents/ISession.cs
+++ b/Engine/Source/Programs/Horde/Horde.Build/Agents/ISession.cs
@@ -1,13 +1,8 @@
 // Copyright Epic Games, Inc. All Rights Reserved.
 
-using Horde.Build.Api;
-using Horde.Build.Utilities;
-using MongoDB.Bson;
-using MongoDB.Bson.Serialization.Attributes;
 using System;
 using System.Collections.Generic;
-using System.Linq;
-using System.Threading.Tasks;
+using Horde.Build.Utilities;
 
 namespace Horde.Build.Models
 {
diff --git a/Engine/Source/Programs/Horde/Horde.Build/Agents/ISessionCollection.cs b/Engine/Source/Programs/Horde/Horde.Build/Agents/ISessionCollection.cs
--- a/Engine/Source/Programs/Horde/Horde.Build/Agents/ISessionCollection.cs
+++ b/Engine/Source/Programs/Horde/Horde.Build/Agents/ISessionCollection.cs
@@ -1,16 +1,10 @@
 // Copyright Epic Games, Inc. All Rights Reserved.
 
-using Horde.Build.Models;
-using Horde.Build.Services;
-using Horde.Build.Utilities;
-using Microsoft.Extensions.DependencyInjection;
-using MongoDB.Bson;
-using MongoDB.Bson.Serialization.Attributes;
-using MongoDB.Driver;
 using System;
 using System.Collections.Generic;
-using System.Linq;
 using System.Threading.Tasks;
+using Horde.Build.Models;
+using Horde.Build.Utilities;
 
 namespace Horde.Build.Collections
 {
@@ -24,55 +18,55 @@ namespace Horde.Build.Collections
 		/// <summary>
 		/// Adds a new session to the collection
 		/// </summary>
-		/// <param name="Id">The session id</param>
-		/// <param name="AgentId">The agent this session is for</param>
-		/// <param name="StartTime">Start time of this session</param>
-		/// <param name="Properties">Properties of this agent at the time the session started</param>
-		/// <param name="Resources">Resources which the agent has</param>
-		/// <param name="Version">Version of the agent software</param>
-		Task<ISession> AddAsync(SessionId Id, AgentId AgentId, DateTime StartTime, IReadOnlyList<string>? Properties, IReadOnlyDictionary<string, int>? Resources, string? Version);
+		/// <param name="id">The session id</param>
+		/// <param name="agentId">The agent this session is for</param>
+		/// <param name="startTime">Start time of this session</param>
+		/// <param name="properties">Properties of this agent at the time the session started</param>
+		/// <param name="resources">Resources which the agent has</param>
+		/// <param name="version">Version of the agent software</param>
+		Task<ISession> AddAsync(SessionId id, AgentId agentId, DateTime startTime, IReadOnlyList<string>? properties, IReadOnlyDictionary<string, int>? resources, string? version);
 
 		/// <summary>
 		/// Gets information about a particular session
 		/// </summary>
-		/// <param name="SessionId">The unique session id</param>
+		/// <param name="sessionId">The unique session id</param>
 		/// <returns>The session information</returns>
-		Task<ISession?> GetAsync(SessionId SessionId);
+		Task<ISession?> GetAsync(SessionId sessionId);
 
 		/// <summary>
 		/// Find sessions for the given agent
 		/// </summary>
-		/// <param name="AgentId">The unique agent id</param>
-		/// <param name="StartTime">Start time to include in the search</param>
-		/// <param name="FinishTime">Finish time to include in the search</param>
-		/// <param name="Index">Index of the first result to return</param>
-		/// <param name="Count">Number of results to return</param>
+		/// <param name="agentId">The unique agent id</param>
+		/// <param name="startTime">Start time to include in the search</param>
+		/// <param name="finishTime">Finish time to include in the search</param>
+		/// <param name="index">Index of the first result to return</param>
+		/// <param name="count">Number of results to return</param>
 		/// <returns>List of sessions matching the given criteria</returns>
-		Task<List<ISession>> FindAsync(AgentId AgentId, DateTime? StartTime, DateTime? FinishTime, int Index, int Count);
+		Task<List<ISession>> FindAsync(AgentId agentId, DateTime? startTime, DateTime? finishTime, int index, int count);
 
 		/// <summary>
 		/// Finds any active sessions
 		/// </summary>
-		/// <param name="Index">Index of the first result to return</param>
-		/// <param name="Count">Number of results to return</param>
+		/// <param name="index">Index of the first result to return</param>
+		/// <param name="count">Number of results to return</param>
 		/// <returns>List of sessions</returns>
-		Task<List<ISession>> FindActiveSessionsAsync(int? Index = null, int? Count = null);
+		Task<List<ISession>> FindActiveSessionsAsync(int? index = null, int? count = null);
 
 		/// <summary>
 		/// Update a session from the collection
 		/// </summary>
-		/// <param name="SessionId">The session to update</param>
-		/// <param name="FinishTime">Time at which the session finished</param>
-		/// <param name="Properties">The agent properties</param>
-		/// <param name="Resources">Resources which the agent has</param>
+		/// <param name="sessionId">The session to update</param>
+		/// <param name="finishTime">Time at which the session finished</param>
+		/// <param name="properties">The agent properties</param>
+		/// <param name="resources">Resources which the agent has</param>
 		/// <returns>Async task</returns>
-		Task UpdateAsync(SessionId SessionId, DateTime FinishTime, IReadOnlyList<string> Properties, IReadOnlyDictionary<string, int> Resources);
+		Task UpdateAsync(SessionId sessionId, DateTime finishTime, IReadOnlyList<string> properties, IReadOnlyDictionary<string, int> resources);
 
 		/// <summary>
 		/// Delete a session from the collection
 		/// </summary>
-		/// <param name="SessionId">The session id</param>
+		/// <param name="sessionId">The session id</param>
 		/// <returns>Async task</returns>
-		Task DeleteAsync(SessionId SessionId);
+		Task DeleteAsync(SessionId sessionId);
 	}
 }
diff --git a/Engine/Source/Programs/Horde/Horde.Build/Agents/ITaskSource.cs b/Engine/Source/Programs/Horde/Horde.Build/Agents/ITaskSource.cs
--- a/Engine/Source/Programs/Horde/Horde.Build/Agents/ITaskSource.cs
+++ b/Engine/Source/Programs/Horde/Horde.Build/Agents/ITaskSource.cs
@@ -1,23 +1,20 @@
 // Copyright Epic Games, Inc. All Rights Reserved.
 
-using Google.Protobuf;
-using Google.Protobuf.Reflection;
-using Google.Protobuf.WellKnownTypes;
-using HordeCommon;
-using Horde.Build.Models;
-using Horde.Build.Services;
-using Horde.Build.Utilities;
-using Microsoft.Extensions.Logging;
-using MongoDB.Bson;
 using System;
 using System.Collections.Generic;
 using System.Globalization;
-using System.Linq;
 using System.Linq.Expressions;
 using System.Reflection;
 using System.Text;
 using System.Threading;
 using System.Threading.Tasks;
+using Google.Protobuf;
+using Google.Protobuf.Reflection;
+using Google.Protobuf.WellKnownTypes;
+using Horde.Build.Models;
+using Horde.Build.Utilities;
+using HordeCommon;
+using Microsoft.Extensions.Logging;
 
 namespace Horde.Build.Tasks
 {
@@ -68,46 +65,46 @@ namespace Horde.Build.Tasks
 		/// <summary>
 		/// Assigns a lease or waits for one to be available
 		/// </summary>
-		/// <param name="Agent">The agent to assign a lease to</param>
-		/// <param name="CancellationToken">Cancellation token for the wait</param>
+		/// <param name="agent">The agent to assign a lease to</param>
+		/// <param name="cancellationToken">Cancellation token for the wait</param>
 		/// <returns>New lease object</returns>
-		Task<AgentLease?> AssignLeaseAsync(IAgent Agent, CancellationToken CancellationToken);
+		Task<AgentLease?> AssignLeaseAsync(IAgent agent, CancellationToken cancellationToken);
 
 		/// <summary>
 		/// Cancel a lease that was previously assigned to an agent, allowing it to be assigned out again
 		/// </summary>
-		/// <param name="Agent">The agent that was assigned the lease</param>
-		/// <param name="LeaseId">The lease id</param>
-		/// <param name="Payload">Payload for the lease</param>
+		/// <param name="agent">The agent that was assigned the lease</param>
+		/// <param name="leaseId">The lease id</param>
+		/// <param name="payload">Payload for the lease</param>
 		/// <returns></returns>
-		Task CancelLeaseAsync(IAgent Agent, LeaseId LeaseId, Any Payload);
+		Task CancelLeaseAsync(IAgent agent, LeaseId leaseId, Any payload);
 
 		/// <summary>
 		/// Notification that a lease has been started
 		/// </summary>
-		/// <param name="Agent">The agent executing the lease</param>
-		/// <param name="LeaseId">The lease id</param>
-		/// <param name="Payload">Payload for the lease</param>
-		/// <param name="Logger">Logger for the agent</param>
-		Task OnLeaseStartedAsync(IAgent Agent, LeaseId LeaseId, Any Payload, ILogger Logger);
+		/// <param name="agent">The agent executing the lease</param>
+		/// <param name="leaseId">The lease id</param>
+		/// <param name="payload">Payload for the lease</param>
+		/// <param name="logger">Logger for the agent</param>
+		Task OnLeaseStartedAsync(IAgent agent, LeaseId leaseId, Any payload, ILogger logger);
 
 		/// <summary>
 		/// Notification that a task has completed
 		/// </summary>
-		/// <param name="Agent">The agent that was allocated to the lease</param>
-		/// <param name="LeaseId">The lease id</param>
-		/// <param name="Payload">The lease payload</param>
-		/// <param name="Outcome">Outcome of the lease</param>
-		/// <param name="Output">Output from the task</param>
-		/// <param name="Logger">Logger for the agent</param>
-		Task OnLeaseFinishedAsync(IAgent Agent, LeaseId LeaseId, Any Payload, LeaseOutcome Outcome, ReadOnlyMemory<byte> Output, ILogger Logger);
+		/// <param name="agent">The agent that was allocated to the lease</param>
+		/// <param name="leaseId">The lease id</param>
+		/// <param name="payload">The lease payload</param>
+		/// <param name="outcome">Outcome of the lease</param>
+		/// <param name="output">Output from the task</param>
+		/// <param name="logger">Logger for the agent</param>
+		Task OnLeaseFinishedAsync(IAgent agent, LeaseId leaseId, Any payload, LeaseOutcome outcome, ReadOnlyMemory<byte> output, ILogger logger);
 
 		/// <summary>
 		/// Gets information to include for a lease in a lease info response
 		/// </summary>
-		/// <param name="Payload">The lease payload</param>
-		/// <param name="Details">Properties for the lease</param>
-		void GetLeaseDetails(Any Payload, Dictionary<string, string> Details);
+		/// <param name="payload">The lease payload</param>
+		/// <param name="details">Properties for the lease</param>
+		void GetLeaseDetails(Any payload, Dictionary<string, string> details);
 	}
 	
 	/// <summary>
@@ -121,36 +118,36 @@ namespace Horde.Build.Tasks
 		/// </summary>
 		protected class PropertyList
 		{
-			internal StringBuilder FormatString = new StringBuilder();
-			internal List<Func<TMessage, object>> Accessors = new List<Func<TMessage, object>>();
-			internal List<(string, Func<TMessage, object>)> JsonAccessors = new List<(string, Func<TMessage, object>)>();
+			internal StringBuilder _formatString = new StringBuilder();
+			internal List<Func<TMessage, object>> _accessors = new List<Func<TMessage, object>>();
+			internal List<(string, Func<TMessage, object>)> _jsonAccessors = new List<(string, Func<TMessage, object>)>();
 
 			/// <summary>
 			/// Adds a new property to the list
 			/// </summary>
-			/// <param name="Expr">Accessor for the property</param>
-			public PropertyList Add(Expression<Func<TMessage, object>> Expr)
+			/// <param name="expr">Accessor for the property</param>
+			public PropertyList Add(Expression<Func<TMessage, object>> expr)
 			{
-				MemberInfo Member = ((MemberExpression)Expr.Body).Member;
-				return Add(Member.Name, Expr.Compile());
+				MemberInfo member = ((MemberExpression)expr.Body).Member;
+				return Add(member.Name, expr.Compile());
 			}
 
 			/// <summary>
 			/// Adds a new property to the list
 			/// </summary>
-			/// <param name="Name">Name of the property</param>
-			/// <param name="Accessor">Accessor for the property</param>
+			/// <param name="name">Name of the property</param>
+			/// <param name="accessor">Accessor for the property</param>
 			[System.Diagnostics.CodeAnalysis.SuppressMessage("Globalization", "CA1308:Normalize strings to uppercase", Justification = "<Pending>")]
-			public PropertyList Add(string Name, Func<TMessage, object> Accessor)
+			public PropertyList Add(string name, Func<TMessage, object> accessor)
 			{
-				FormatString.Append(CultureInfo.InvariantCulture, $", {Name}={{{Name}}}");
-				Accessors.Add(Accessor);
-				JsonAccessors.Add((Name[0..1].ToLower(CultureInfo.InvariantCulture) + Name[1..], Accessor));
+				_formatString.Append(CultureInfo.InvariantCulture, $", {name}={{{name}}}");
+				_accessors.Add(accessor);
+				_jsonAccessors.Add((name[0..1].ToLower(CultureInfo.InvariantCulture) + name[1..], accessor));
 				return this;
 			}
 		}
 
-		static TMessage Message = new TMessage();
+		static readonly TMessage s_message = new TMessage();
 
 		/// <inheritdoc/>
 		public abstract string Type { get; }
@@ -162,57 +159,57 @@ namespace Horde.Build.Tasks
 		protected virtual PropertyList OnLeaseStartedProperties { get; } = new PropertyList();
 
 		/// <inheritdoc/>
-		public MessageDescriptor Descriptor => Message.Descriptor;
+		public MessageDescriptor Descriptor => s_message.Descriptor;
 
 		/// <inheritdoc/>
-		public abstract Task<AgentLease?> AssignLeaseAsync(IAgent Agent, CancellationToken CancellationToken);
+		public abstract Task<AgentLease?> AssignLeaseAsync(IAgent agent, CancellationToken cancellationToken);
 
 		/// <inheritdoc/>
-		public Task CancelLeaseAsync(IAgent Agent, LeaseId LeaseId, Any Payload) => CancelLeaseAsync(Agent, LeaseId, Payload.Unpack<TMessage>());
+		public Task CancelLeaseAsync(IAgent agent, LeaseId leaseId, Any payload) => CancelLeaseAsync(agent, leaseId, payload.Unpack<TMessage>());
 
 		/// <inheritdoc/>
-		public Task OnLeaseStartedAsync(IAgent Agent, LeaseId LeaseId, Any Payload, ILogger Logger) => OnLeaseStartedAsync(Agent, LeaseId, Payload.Unpack<TMessage>(), Logger);
+		public Task OnLeaseStartedAsync(IAgent agent, LeaseId leaseId, Any payload, ILogger logger) => OnLeaseStartedAsync(agent, leaseId, payload.Unpack<TMessage>(), logger);
 
 		/// <inheritdoc/>
-		public Task OnLeaseFinishedAsync(IAgent Agent, LeaseId LeaseId, Any Payload, LeaseOutcome Outcome, ReadOnlyMemory<byte> Output, ILogger Logger) => OnLeaseFinishedAsync(Agent, LeaseId, Payload.Unpack<TMessage>(), Outcome, Output, Logger);
+		public Task OnLeaseFinishedAsync(IAgent agent, LeaseId leaseId, Any payload, LeaseOutcome outcome, ReadOnlyMemory<byte> output, ILogger logger) => OnLeaseFinishedAsync(agent, leaseId, payload.Unpack<TMessage>(), outcome, output, logger);
 
 		/// <inheritdoc cref="ITaskSource.CancelLeaseAsync(IAgent, LeaseId, Any)"/>
-		public virtual Task CancelLeaseAsync(IAgent Agent, LeaseId LeaseId, TMessage Payload) => Task.CompletedTask;
+		public virtual Task CancelLeaseAsync(IAgent agent, LeaseId leaseId, TMessage payload) => Task.CompletedTask;
 
 		/// <inheritdoc cref="ITaskSource.OnLeaseStartedAsync(IAgent, LeaseId, Any, ILogger)"/>
-		public virtual Task OnLeaseStartedAsync(IAgent Agent, LeaseId LeaseId, TMessage Payload, ILogger Logger)
+		public virtual Task OnLeaseStartedAsync(IAgent agent, LeaseId leaseId, TMessage payload, ILogger logger)
 		{
-			object[] Arguments = new object[2 + OnLeaseStartedProperties.Accessors.Count];
-			Arguments[0] = LeaseId;
-			Arguments[1] = Type;
-			for (int Idx = 0; Idx < OnLeaseStartedProperties.Accessors.Count; Idx++)
+			object[] arguments = new object[2 + OnLeaseStartedProperties._accessors.Count];
+			arguments[0] = leaseId;
+			arguments[1] = Type;
+			for (int idx = 0; idx < OnLeaseStartedProperties._accessors.Count; idx++)
 			{
-				Arguments[Idx + 2] = OnLeaseStartedProperties.Accessors[Idx](Payload);
+				arguments[idx + 2] = OnLeaseStartedProperties._accessors[idx](payload);
 			}
 #pragma warning disable CA2254 // Template should be a static expression
-			Logger.LogInformation($"Lease {{LeaseId}} started (Type={{Type}}{OnLeaseStartedProperties.FormatString})", Arguments);
+			logger.LogInformation($"Lease {{LeaseId}} started (Type={{Type}}{OnLeaseStartedProperties._formatString})", arguments);
 #pragma warning restore CA2254 // Template should be a static expression
 			return Task.CompletedTask;
 		}
 
 		/// <inheritdoc cref="ITaskSource.OnLeaseFinishedAsync(IAgent, LeaseId, Any, LeaseOutcome, ReadOnlyMemory{byte}, ILogger)"/>
-		public virtual Task OnLeaseFinishedAsync(IAgent Agent, LeaseId LeaseId, TMessage Payload, LeaseOutcome Outcome, ReadOnlyMemory<byte> Output, ILogger Logger)
+		public virtual Task OnLeaseFinishedAsync(IAgent agent, LeaseId leaseId, TMessage payload, LeaseOutcome outcome, ReadOnlyMemory<byte> output, ILogger logger)
 		{
-			Logger.LogInformation("Lease {LeaseId} complete", LeaseId);
+			logger.LogInformation("Lease {LeaseId} complete", leaseId);
 			return Task.CompletedTask;
 		}
 
 		/// <inheritdoc/>
-		public virtual void GetLeaseDetails(Any Payload, Dictionary<string, string> Details)
+		public virtual void GetLeaseDetails(Any payload, Dictionary<string, string> details)
 		{
-			Details["type"] = Type;
+			details["type"] = Type;
 
-			if (OnLeaseStartedProperties.Accessors.Count > 0)
+			if (OnLeaseStartedProperties._accessors.Count > 0)
 			{
-				TMessage Message = Payload.Unpack<TMessage>();
-				foreach ((string Name, Func<TMessage, object> GetMethod) in OnLeaseStartedProperties.JsonAccessors)
+				TMessage message = payload.Unpack<TMessage>();
+				foreach ((string name, Func<TMessage, object> getMethod) in OnLeaseStartedProperties._jsonAccessors)
 				{
-					Details[Name] = GetMethod(Message)?.ToString() ?? String.Empty;
+					details[name] = getMethod(message)?.ToString() ?? String.Empty;
 				}
 			}
 		}
diff --git a/Engine/Source/Programs/Horde/Horde.Build/Agents/Impl/AgentCollection.cs b/Engine/Source/Programs/Horde/Horde.Build/Agents/Impl/AgentCollection.cs
--- a/Engine/Source/Programs/Horde/Horde.Build/Agents/Impl/AgentCollection.cs
+++ b/Engine/Source/Programs/Horde/Horde.Build/Agents/Impl/AgentCollection.cs
@@ -1,22 +1,21 @@
-﻿// Copyright Epic Games, Inc. All Rights Reserved.
+// Copyright Epic Games, Inc. All Rights Reserved.
 
+using System;
+using System.Collections.Generic;
+using System.Linq;
+using System.Linq.Expressions;
+using System.Threading.Tasks;
 using EpicGames.Redis;
 using Google.Protobuf.WellKnownTypes;
-using HordeCommon;
-using HordeCommon.Rpc.Tasks;
 using Horde.Build.Acls;
 using Horde.Build.Models;
 using Horde.Build.Services;
 using Horde.Build.Utilities;
-using Microsoft.Extensions.Logging;
+using HordeCommon;
+using HordeCommon.Rpc.Tasks;
 using MongoDB.Bson;
 using MongoDB.Bson.Serialization.Attributes;
 using MongoDB.Driver;
-using System;
-using System.Collections.Generic;
-using System.Linq;
-using System.Linq.Expressions;
-using System.Threading.Tasks;
 
 namespace Horde.Build.Collections.Impl
 {
@@ -57,7 +56,7 @@ namespace Horde.Build.Collections.Impl
 		/// </summary>
 		class AgentDocument : IAgent
 		{
-			static IReadOnlyList<AgentLease> EmptyLeases = new List<AgentLease>();
+			static readonly IReadOnlyList<AgentLease> s_emptyLeases = new List<AgentLease>();
 
 			[BsonRequired, BsonId]
 			public AgentId Id { get; set; }
@@ -124,7 +123,7 @@ namespace Horde.Build.Collections.Impl
 			IReadOnlyList<PoolId> IAgent.DynamicPools => DynamicPools;
 			IReadOnlyList<PoolId> IAgent.ExplicitPools => Pools;
 			IReadOnlyList<AgentWorkspace> IAgent.Workspaces => Workspaces;
-			IReadOnlyList<AgentLease> IAgent.Leases => Leases ?? EmptyLeases;
+			IReadOnlyList<AgentLease> IAgent.Leases => Leases ?? s_emptyLeases;
 			IReadOnlyList<string> IAgent.Properties => Properties ?? Capabilities.Devices.FirstOrDefault()?.Properties?.ToList() ?? new List<string>();
 			IReadOnlyDictionary<string, int> IAgent.Resources => Resources ?? Capabilities.Devices.FirstOrDefault()?.Resources ?? new Dictionary<string, int>();
 
@@ -133,291 +132,291 @@ namespace Horde.Build.Collections.Impl
 			{
 			}
 
-			public AgentDocument(AgentId Id, bool bEnabled, AgentSoftwareChannelName? Channel, List<PoolId> Pools)
+			public AgentDocument(AgentId id, bool bEnabled, AgentSoftwareChannelName? channel, List<PoolId> pools)
 			{
-				this.Id = Id;
-				this.Acl = new Acl();
-				this.Enabled = bEnabled;
-				this.Channel = Channel;
-				this.Pools = Pools;
+				Id = id;
+				Acl = new Acl();
+				Enabled = bEnabled;
+				Channel = channel;
+				Pools = pools;
 			}
 		}
 
-		readonly IMongoCollection<AgentDocument> Agents;
-		readonly IAuditLog<AgentId> AuditLog;
-		readonly RedisService RedisService;
-		readonly RedisChannel<AgentId> UpdateEventChannel;
+		readonly IMongoCollection<AgentDocument> _agents;
+		readonly IAuditLog<AgentId> _auditLog;
+		readonly RedisService _redisService;
+		readonly RedisChannel<AgentId> _updateEventChannel;
 
 		/// <summary>
 		/// Constructor
 		/// </summary>
-		public AgentCollection(DatabaseService DatabaseService, RedisService RedisService, IAuditLog<AgentId> AuditLog)
+		public AgentCollection(DatabaseService databaseService, RedisService redisService, IAuditLog<AgentId> auditLog)
 		{
-			this.Agents = DatabaseService.GetCollection<AgentDocument>("Agents");
-			this.RedisService = RedisService;
-			this.UpdateEventChannel = new RedisChannel<AgentId>("agents/notify");
-			this.AuditLog = AuditLog;
+			_agents = databaseService.GetCollection<AgentDocument>("Agents");
+			_redisService = redisService;
+			_updateEventChannel = new RedisChannel<AgentId>("agents/notify");
+			_auditLog = auditLog;
 
-			if (!DatabaseService.ReadOnlyMode)
+			if (!databaseService.ReadOnlyMode)
 			{
-				Agents.Indexes.CreateOne(new CreateIndexModel<AgentDocument>(Builders<AgentDocument>.IndexKeys.Ascending(x => x.Deleted).Ascending(x => x.Id).Ascending(x => x.Pools)));
+				_agents.Indexes.CreateOne(new CreateIndexModel<AgentDocument>(Builders<AgentDocument>.IndexKeys.Ascending(x => x.Deleted).Ascending(x => x.Id).Ascending(x => x.Pools)));
 			}
 		}
 
 		/// <inheritdoc/>
-		public async Task<IAgent> AddAsync(AgentId Id, bool bEnabled, AgentSoftwareChannelName? Channel, List<PoolId>? Pools)
+		public async Task<IAgent> AddAsync(AgentId id, bool bEnabled, AgentSoftwareChannelName? channel, List<PoolId>? pools)
 		{
-			AgentDocument Agent = new AgentDocument(Id, bEnabled, Channel, Pools ?? new List<PoolId>());
-			await Agents.InsertOneAsync(Agent);
-			return Agent;
+			AgentDocument agent = new AgentDocument(id, bEnabled, channel, pools ?? new List<PoolId>());
+			await _agents.InsertOneAsync(agent);
+			return agent;
 		}
 
 		/// <inheritdoc/>
-		public async Task<IAgent?> TryDeleteAsync(IAgent AgentInterface)
+		public async Task<IAgent?> TryDeleteAsync(IAgent agentInterface)
 		{
-			AgentDocument Agent = (AgentDocument)AgentInterface;
+			AgentDocument agent = (AgentDocument)agentInterface;
 
-			UpdateDefinition<AgentDocument> Update = Builders<AgentDocument>.Update.Set(x => x.Deleted, true);
-			return await TryUpdateAsync(Agent, Update);
+			UpdateDefinition<AgentDocument> update = Builders<AgentDocument>.Update.Set(x => x.Deleted, true);
+			return await TryUpdateAsync(agent, update);
 		}
 
 		/// <inheritdoc/>
-		public async Task ForceDeleteAsync(AgentId AgentId)
+		public async Task ForceDeleteAsync(AgentId agentId)
 		{
-			await Agents.DeleteOneAsync(x => x.Id == AgentId);
+			await _agents.DeleteOneAsync(x => x.Id == agentId);
 		}
 
 		/// <inheritdoc/>
-		public async Task<IAgent?> GetAsync(AgentId AgentId)
+		public async Task<IAgent?> GetAsync(AgentId agentId)
 		{
-			return await Agents.Find<AgentDocument>(x => x.Id == AgentId).FirstOrDefaultAsync();
+			return await _agents.Find<AgentDocument>(x => x.Id == agentId).FirstOrDefaultAsync();
 		}
 
 		/// <inheritdoc/>
-		public async Task<List<IAgent>> FindAsync(PoolId? PoolId, DateTime? ModifiedAfter, AgentStatus? Status, int? Index, int? Count)
+		public async Task<List<IAgent>> FindAsync(PoolId? poolId, DateTime? modifiedAfter, AgentStatus? status, int? index, int? count)
 		{
-			FilterDefinitionBuilder<AgentDocument> FilterBuilder = new FilterDefinitionBuilder<AgentDocument>();
+			FilterDefinitionBuilder<AgentDocument> filterBuilder = new FilterDefinitionBuilder<AgentDocument>();
 
-			FilterDefinition<AgentDocument> Filter = FilterBuilder.Ne(x => x.Deleted, true);
+			FilterDefinition<AgentDocument> filter = filterBuilder.Ne(x => x.Deleted, true);
 			
-			if (PoolId != null)
+			if (poolId != null)
 			{
-				Filter &= FilterBuilder.Eq(nameof(AgentDocument.Pools), PoolId);
+				filter &= filterBuilder.Eq(nameof(AgentDocument.Pools), poolId);
 			}
 			
-			if (ModifiedAfter != null)
+			if (modifiedAfter != null)
 			{
-				Filter &= FilterBuilder.Gt(x => x.UpdateTime, ModifiedAfter.Value);
+				filter &= filterBuilder.Gt(x => x.UpdateTime, modifiedAfter.Value);
 			}
 			
-			if (Status != null)
+			if (status != null)
 			{
-				Filter &= FilterBuilder.Eq(x => x.Status, Status.Value);
+				filter &= filterBuilder.Eq(x => x.Status, status.Value);
 			}
 
-			IFindFluent<AgentDocument, AgentDocument> Search = Agents.Find(Filter);
-			if (Index != null)
+			IFindFluent<AgentDocument, AgentDocument> search = _agents.Find(filter);
+			if (index != null)
 			{
-				Search = Search.Skip(Index.Value);
+				search = search.Skip(index.Value);
 			}
-			if (Count != null)
+			if (count != null)
 			{
-				Search = Search.Limit(Count.Value);
+				search = search.Limit(count.Value);
 			}
 
-			List<AgentDocument> Results = await Search.ToListAsync();
-			return Results.ConvertAll<IAgent>(x => x);
+			List<AgentDocument> results = await search.ToListAsync();
+			return results.ConvertAll<IAgent>(x => x);
 		}
 
 		/// <inheritdoc/>
-		public async Task<List<IAgent>> FindExpiredAsync(DateTime UtcNow, int MaxAgents)
+		public async Task<List<IAgent>> FindExpiredAsync(DateTime utcNow, int maxAgents)
 		{
-			List<AgentDocument> Results = await Agents.Find(x => x.SessionId.HasValue && !(x.SessionExpiresAt > UtcNow)).Limit(MaxAgents).ToListAsync();
-			return Results.ConvertAll<IAgent>(x => x);
+			List<AgentDocument> results = await _agents.Find(x => x.SessionId.HasValue && !(x.SessionExpiresAt > utcNow)).Limit(maxAgents).ToListAsync();
+			return results.ConvertAll<IAgent>(x => x);
 		}
 
 		/// <summary>
 		/// Update a single document
 		/// </summary>
-		/// <param name="Current">The document to update</param>
-		/// <param name="Update">The update definition</param>
+		/// <param name="current">The document to update</param>
+		/// <param name="update">The update definition</param>
 		/// <returns>True if the agent was updated</returns>
-		private async Task<AgentDocument?> TryUpdateAsync(AgentDocument Current, UpdateDefinition<AgentDocument> Update)
+		private async Task<AgentDocument?> TryUpdateAsync(AgentDocument current, UpdateDefinition<AgentDocument> update)
 		{
-			uint PrevUpdateIndex = Current.UpdateIndex++;
-			Current.UpdateTime = DateTime.UtcNow;
+			uint prevUpdateIndex = current.UpdateIndex++;
+			current.UpdateTime = DateTime.UtcNow;
 
-			Expression<Func<AgentDocument, bool>> Filter = x => x.Id == Current.Id && x.UpdateIndex == PrevUpdateIndex;
-			UpdateDefinition<AgentDocument> UpdateWithIndex = Update.Set(x => x.UpdateIndex, Current.UpdateIndex).Set(x => x.UpdateTime, Current.UpdateTime);
+			Expression<Func<AgentDocument, bool>> filter = x => x.Id == current.Id && x.UpdateIndex == prevUpdateIndex;
+			UpdateDefinition<AgentDocument> updateWithIndex = update.Set(x => x.UpdateIndex, current.UpdateIndex).Set(x => x.UpdateTime, current.UpdateTime);
 
-			return await Agents.FindOneAndUpdateAsync<AgentDocument>(Filter, UpdateWithIndex, new FindOneAndUpdateOptions<AgentDocument, AgentDocument> { ReturnDocument = ReturnDocument.After });
+			return await _agents.FindOneAndUpdateAsync<AgentDocument>(filter, updateWithIndex, new FindOneAndUpdateOptions<AgentDocument, AgentDocument> { ReturnDocument = ReturnDocument.After });
 		}
 
 		/// <inheritdoc/>
-		public async Task<IAgent?> TryUpdateSettingsAsync(IAgent AgentInterface, bool? Enabled = null, bool? RequestConform = null, bool? RequestFullConform = null, bool? RequestRestart = null, bool? RequestShutdown = null, string? ShutdownReason = null, AgentSoftwareChannelName? Channel = null, List<PoolId>? Pools = null, Acl? Acl = null, string? Comment = null)
+		public async Task<IAgent?> TryUpdateSettingsAsync(IAgent agentInterface, bool? enabled = null, bool? requestConform = null, bool? requestFullConform = null, bool? requestRestart = null, bool? requestShutdown = null, string? shutdownReason = null, AgentSoftwareChannelName? channel = null, List<PoolId>? pools = null, Acl? acl = null, string? comment = null)
 		{
-			AgentDocument Agent = (AgentDocument)AgentInterface;
+			AgentDocument agent = (AgentDocument)agentInterface;
 
 			// Update the database
-			UpdateDefinitionBuilder<AgentDocument> UpdateBuilder = new UpdateDefinitionBuilder<AgentDocument>();
+			UpdateDefinitionBuilder<AgentDocument> updateBuilder = new UpdateDefinitionBuilder<AgentDocument>();
 
-			List<UpdateDefinition<AgentDocument>> Updates = new List<UpdateDefinition<AgentDocument>>();
-			if (Pools != null)
+			List<UpdateDefinition<AgentDocument>> updates = new List<UpdateDefinition<AgentDocument>>();
+			if (pools != null)
 			{
-				Updates.Add(UpdateBuilder.Set(x => x.Pools, Pools));
+				updates.Add(updateBuilder.Set(x => x.Pools, pools));
 			}
-			if (Enabled != null)
+			if (enabled != null)
 			{
-				Updates.Add(UpdateBuilder.Set(x => x.Enabled, Enabled.Value));
+				updates.Add(updateBuilder.Set(x => x.Enabled, enabled.Value));
 			}
-			if (RequestConform != null)
+			if (requestConform != null)
 			{
-				Updates.Add(UpdateBuilder.Set(x => x.RequestConform, RequestConform.Value));
-				Updates.Add(UpdateBuilder.Unset(x => x.ConformAttemptCount));
+				updates.Add(updateBuilder.Set(x => x.RequestConform, requestConform.Value));
+				updates.Add(updateBuilder.Unset(x => x.ConformAttemptCount));
 			}
-			if (RequestFullConform != null)
+			if (requestFullConform != null)
 			{
-				Updates.Add(UpdateBuilder.Set(x => x.RequestFullConform, RequestFullConform.Value));
-				Updates.Add(UpdateBuilder.Unset(x => x.ConformAttemptCount));
+				updates.Add(updateBuilder.Set(x => x.RequestFullConform, requestFullConform.Value));
+				updates.Add(updateBuilder.Unset(x => x.ConformAttemptCount));
 			}
-			if (RequestRestart != null)
+			if (requestRestart != null)
 			{
-				if (RequestRestart.Value)
+				if (requestRestart.Value)
 				{
-					Updates.Add(UpdateBuilder.Set(x => x.RequestRestart, true));
+					updates.Add(updateBuilder.Set(x => x.RequestRestart, true));
 				}
 				else
 				{
-					Updates.Add(UpdateBuilder.Unset(x => x.RequestRestart));
+					updates.Add(updateBuilder.Unset(x => x.RequestRestart));
 				}
 			}
-			if (RequestShutdown != null)
+			if (requestShutdown != null)
 			{
-				if (RequestShutdown.Value)
+				if (requestShutdown.Value)
 				{
-					Updates.Add(UpdateBuilder.Set(x => x.RequestShutdown, true));
+					updates.Add(updateBuilder.Set(x => x.RequestShutdown, true));
 				}
 				else
 				{
-					Updates.Add(UpdateBuilder.Unset(x => x.RequestShutdown));
+					updates.Add(updateBuilder.Unset(x => x.RequestShutdown));
 				}
 			}
 
-			if (ShutdownReason != null)
+			if (shutdownReason != null)
 			{
-				Updates.Add(UpdateBuilder.Set(x => x.LastShutdownReason, ShutdownReason));
+				updates.Add(updateBuilder.Set(x => x.LastShutdownReason, shutdownReason));
 			}
 
-			if (Channel != null)
+			if (channel != null)
 			{
-				if (Channel.Value == AgentSoftwareService.DefaultChannelName)
+				if (channel.Value == AgentSoftwareService.DefaultChannelName)
 				{
-					Updates.Add(UpdateBuilder.Unset(x => x.Channel));
+					updates.Add(updateBuilder.Unset(x => x.Channel));
 				}
 				else
 				{
-					Updates.Add(UpdateBuilder.Set(x => x.Channel, Channel));
+					updates.Add(updateBuilder.Set(x => x.Channel, channel));
 				}
 			}
-			if (Acl != null)
+			if (acl != null)
 			{
-				Updates.Add(Acl.CreateUpdate<AgentDocument>(x => x.Acl!, Acl));
+				updates.Add(Acl.CreateUpdate<AgentDocument>(x => x.Acl!, acl));
 			}
-			if (Comment != null)
+			if (comment != null)
 			{
-				Updates.Add(UpdateBuilder.Set(x => x.Comment, Comment));
+				updates.Add(updateBuilder.Set(x => x.Comment, comment));
 			}
 
 			// Apply the update
-			IAgent? NewAgent = await TryUpdateAsync(Agent, UpdateBuilder.Combine(Updates));
-			if (NewAgent != null)
+			IAgent? newAgent = await TryUpdateAsync(agent, updateBuilder.Combine(updates));
+			if (newAgent != null)
 			{
-				if (NewAgent.RequestRestart != Agent.RequestRestart || NewAgent.RequestConform != Agent.RequestConform || NewAgent.RequestShutdown != Agent.RequestShutdown || NewAgent.Channel != Agent.Channel)
+				if (newAgent.RequestRestart != agent.RequestRestart || newAgent.RequestConform != agent.RequestConform || newAgent.RequestShutdown != agent.RequestShutdown || newAgent.Channel != agent.Channel)
 				{
-					await PublishUpdateEventAsync(Agent.Id);
+					await PublishUpdateEventAsync(agent.Id);
 				}
 			}
-			return NewAgent;
+			return newAgent;
 		}
 
 		/// <inheritdoc/>
-		public async Task<IAgent?> TryUpdateSessionAsync(IAgent AgentInterface, AgentStatus? Status, DateTime? SessionExpiresAt, IReadOnlyList<string>? Properties, IReadOnlyDictionary<string, int>? Resources, IReadOnlyList<PoolId>? DynamicPools, List<AgentLease>? Leases)
+		public async Task<IAgent?> TryUpdateSessionAsync(IAgent agentInterface, AgentStatus? status, DateTime? sessionExpiresAt, IReadOnlyList<string>? properties, IReadOnlyDictionary<string, int>? resources, IReadOnlyList<PoolId>? dynamicPools, List<AgentLease>? leases)
 		{
-			AgentDocument Agent = (AgentDocument)AgentInterface;
+			AgentDocument agent = (AgentDocument)agentInterface;
 
 			// Create an update definition for the agent
-			UpdateDefinitionBuilder<AgentDocument> UpdateBuilder = Builders<AgentDocument>.Update;
-			List<UpdateDefinition<AgentDocument>> Updates = new List<UpdateDefinition<AgentDocument>>();
+			UpdateDefinitionBuilder<AgentDocument> updateBuilder = Builders<AgentDocument>.Update;
+			List<UpdateDefinition<AgentDocument>> updates = new List<UpdateDefinition<AgentDocument>>();
 
-			if (Status != null && Agent.Status != Status.Value)
+			if (status != null && agent.Status != status.Value)
 			{
-				Updates.Add(UpdateBuilder.Set(x => x.Status, Status.Value));
+				updates.Add(updateBuilder.Set(x => x.Status, status.Value));
 			}
-			if (SessionExpiresAt != null)
+			if (sessionExpiresAt != null)
 			{
-				Updates.Add(UpdateBuilder.Set(x => x.SessionExpiresAt, SessionExpiresAt.Value));
+				updates.Add(updateBuilder.Set(x => x.SessionExpiresAt, sessionExpiresAt.Value));
 			}
-			if (Properties != null)
+			if (properties != null)
 			{
-				List<string> NewProperties = Properties.OrderBy(x => x, StringComparer.OrdinalIgnoreCase).ToList();
-				if (!AgentInterface.Properties.SequenceEqual(NewProperties, StringComparer.Ordinal))
+				List<string> newProperties = properties.OrderBy(x => x, StringComparer.OrdinalIgnoreCase).ToList();
+				if (!agentInterface.Properties.SequenceEqual(newProperties, StringComparer.Ordinal))
 				{
-					Updates.Add(UpdateBuilder.Set(x => x.Properties, NewProperties));
+					updates.Add(updateBuilder.Set(x => x.Properties, newProperties));
 				}
 			}
-			if (Resources != null && !ResourcesEqual(Resources, AgentInterface.Resources))
+			if (resources != null && !ResourcesEqual(resources, agentInterface.Resources))
 			{
-				Updates.Add(UpdateBuilder.Set(x => x.Resources, new Dictionary<string, int>(Resources)));
+				updates.Add(updateBuilder.Set(x => x.Resources, new Dictionary<string, int>(resources)));
 			}
-			if (DynamicPools != null && !DynamicPools.SequenceEqual(Agent.DynamicPools))
+			if (dynamicPools != null && !dynamicPools.SequenceEqual(agent.DynamicPools))
 			{
-				Updates.Add(UpdateBuilder.Set(x => x.DynamicPools, new List<PoolId>(DynamicPools)));
+				updates.Add(updateBuilder.Set(x => x.DynamicPools, new List<PoolId>(dynamicPools)));
 			}
-			if (Leases != null)
+			if (leases != null)
 			{
-				foreach (AgentLease Lease in Leases)
+				foreach (AgentLease lease in leases)
 				{
-					if (Lease.Payload != null && (Agent.Leases == null || !Agent.Leases.Any(x => x.Id == Lease.Id)))
+					if (lease.Payload != null && (agent.Leases == null || !agent.Leases.Any(x => x.Id == lease.Id)))
 					{
-						Any Payload = Any.Parser.ParseFrom(Lease.Payload.ToArray());
-						if (Payload.TryUnpack(out ConformTask ConformTask))
+						Any payload = Any.Parser.ParseFrom(lease.Payload.ToArray());
+						if (payload.TryUnpack(out ConformTask conformTask))
 						{
-							int NewConformAttemptCount = (Agent.ConformAttemptCount ?? 0) + 1;
-							Updates.Add(UpdateBuilder.Set(x => x.ConformAttemptCount, NewConformAttemptCount));
-							Updates.Add(UpdateBuilder.Set(x => x.LastConformTime, DateTime.UtcNow));
+							int newConformAttemptCount = (agent.ConformAttemptCount ?? 0) + 1;
+							updates.Add(updateBuilder.Set(x => x.ConformAttemptCount, newConformAttemptCount));
+							updates.Add(updateBuilder.Set(x => x.LastConformTime, DateTime.UtcNow));
 						}
-						else if (Payload.TryUnpack(out UpgradeTask UpgradeTask))
+						else if (payload.TryUnpack(out UpgradeTask upgradeTask))
 						{
-							Updates.Add(UpdateBuilder.Set(x => x.LastUpgradeVersion, UpgradeTask.SoftwareId));
-							Updates.Add(UpdateBuilder.Set(x => x.LastUpgradeTime, DateTime.UtcNow));
+							updates.Add(updateBuilder.Set(x => x.LastUpgradeVersion, upgradeTask.SoftwareId));
+							updates.Add(updateBuilder.Set(x => x.LastUpgradeTime, DateTime.UtcNow));
 						}
 					}
 				}
 
-				Updates.Add(UpdateBuilder.Set(x => x.Leases, Leases));
+				updates.Add(updateBuilder.Set(x => x.Leases, leases));
 			}
 
 			// If there are no new updates, return immediately. This is important for preventing UpdateSession calls from returning immediately.
-			if (Updates.Count == 0)
+			if (updates.Count == 0)
 			{
-				return Agent;
+				return agent;
 			}
 
 			// Update the agent, and try to create new lease documents if we succeed
-			return await TryUpdateAsync(Agent, UpdateBuilder.Combine(Updates));
+			return await TryUpdateAsync(agent, updateBuilder.Combine(updates));
 		}
 
-		static bool ResourcesEqual(IReadOnlyDictionary<string, int> DictA, IReadOnlyDictionary<string, int> DictB)
+		static bool ResourcesEqual(IReadOnlyDictionary<string, int> dictA, IReadOnlyDictionary<string, int> dictB)
 		{
-			if (DictA.Count != DictB.Count)
+			if (dictA.Count != dictB.Count)
 			{
 				return false;
 			}
 
-			foreach (KeyValuePair<string, int> Pair in DictA)
+			foreach (KeyValuePair<string, int> pair in dictA)
 			{
-				int Value;
-				if (!DictB.TryGetValue(Pair.Key, out Value) || Value != Pair.Value)
+				int value;
+				if (!dictB.TryGetValue(pair.Key, out value) || value != pair.Value)
 				{
 					return false;
 				}
@@ -427,120 +426,120 @@ namespace Horde.Build.Collections.Impl
 		}
 
 		/// <inheritdoc/>
-		public async Task<IAgent?> TryUpdateWorkspacesAsync(IAgent AgentInterface, List<AgentWorkspace> Workspaces, bool RequestConform)
+		public async Task<IAgent?> TryUpdateWorkspacesAsync(IAgent agentInterface, List<AgentWorkspace> workspaces, bool requestConform)
 		{
-			AgentDocument Agent = (AgentDocument)AgentInterface;
-			DateTime LastConformTime = DateTime.UtcNow;
+			AgentDocument agent = (AgentDocument)agentInterface;
+			DateTime lastConformTime = DateTime.UtcNow;
 
 			// Set the new workspaces
-			UpdateDefinition<AgentDocument> Update = Builders<AgentDocument>.Update.Set(x => x.Workspaces, Workspaces);
-			Update = Update.Set(x => x.LastConformTime, LastConformTime);
-			Update = Update.Unset(x => x.ConformAttemptCount);
-			if (!RequestConform)
+			UpdateDefinition<AgentDocument> update = Builders<AgentDocument>.Update.Set(x => x.Workspaces, workspaces);
+			update = update.Set(x => x.LastConformTime, lastConformTime);
+			update = update.Unset(x => x.ConformAttemptCount);
+			if (!requestConform)
 			{
-				Update = Update.Unset(x => x.RequestConform);
-				Update = Update.Unset(x => x.RequestFullConform);
+				update = update.Unset(x => x.RequestConform);
+				update = update.Unset(x => x.RequestFullConform);
 			}
 
 			// Update the agent
-			return await TryUpdateAsync(Agent, Update);
+			return await TryUpdateAsync(agent, update);
 		}
 
 		/// <inheritdoc/>
-		public async Task<IAgent?> TryStartSessionAsync(IAgent AgentInterface, SessionId SessionId, DateTime SessionExpiresAt, AgentStatus Status, IReadOnlyList<string> Properties, IReadOnlyDictionary<string, int> Resources, IReadOnlyList<PoolId> DynamicPools, string? Version)
+		public async Task<IAgent?> TryStartSessionAsync(IAgent agentInterface, SessionId sessionId, DateTime sessionExpiresAt, AgentStatus status, IReadOnlyList<string> properties, IReadOnlyDictionary<string, int> resources, IReadOnlyList<PoolId> dynamicPools, string? version)
 		{
-			AgentDocument Agent = (AgentDocument)AgentInterface;
-			List<string> NewProperties = Properties.OrderBy(x => x, StringComparer.OrdinalIgnoreCase).ToList();
-			Dictionary<string, int> NewResources = new Dictionary<string, int>(Resources);
-			List<PoolId> NewDynamicPools = new List<PoolId>(DynamicPools);
+			AgentDocument agent = (AgentDocument)agentInterface;
+			List<string> newProperties = properties.OrderBy(x => x, StringComparer.OrdinalIgnoreCase).ToList();
+			Dictionary<string, int> newResources = new Dictionary<string, int>(resources);
+			List<PoolId> newDynamicPools = new List<PoolId>(dynamicPools);
 
 			// Reset the agent to use the new session
-			UpdateDefinitionBuilder<AgentDocument> UpdateBuilder = Builders<AgentDocument>.Update;
-
-			List<UpdateDefinition<AgentDocument>> Updates = new List<UpdateDefinition<AgentDocument>>();
-			Updates.Add(UpdateBuilder.Set(x => x.SessionId, SessionId));
-			Updates.Add(UpdateBuilder.Set(x => x.SessionExpiresAt, SessionExpiresAt));
-			Updates.Add(UpdateBuilder.Set(x => x.Status, Status));
-			Updates.Add(UpdateBuilder.Unset(x => x.Leases));
-			Updates.Add(UpdateBuilder.Unset(x => x.Deleted));
-			Updates.Add(UpdateBuilder.Set(x => x.Properties, NewProperties));
-			Updates.Add(UpdateBuilder.Set(x => x.Resources, NewResources));
-			Updates.Add(UpdateBuilder.Set(x => x.DynamicPools, NewDynamicPools));
-			Updates.Add(UpdateBuilder.Set(x => x.Version, Version));
-			Updates.Add(UpdateBuilder.Unset(x => x.RequestRestart));
-			Updates.Add(UpdateBuilder.Unset(x => x.RequestShutdown));
-			Updates.Add(UpdateBuilder.Set(x => x.LastShutdownReason, "Unexpected"));
+			UpdateDefinitionBuilder<AgentDocument> updateBuilder = Builders<AgentDocument>.Update;
+
+			List<UpdateDefinition<AgentDocument>> updates = new List<UpdateDefinition<AgentDocument>>();
+			updates.Add(updateBuilder.Set(x => x.SessionId, sessionId));
+			updates.Add(updateBuilder.Set(x => x.SessionExpiresAt, sessionExpiresAt));
+			updates.Add(updateBuilder.Set(x => x.Status, status));
+			updates.Add(updateBuilder.Unset(x => x.Leases));
+			updates.Add(updateBuilder.Unset(x => x.Deleted));
+			updates.Add(updateBuilder.Set(x => x.Properties, newProperties));
+			updates.Add(updateBuilder.Set(x => x.Resources, newResources));
+			updates.Add(updateBuilder.Set(x => x.DynamicPools, newDynamicPools));
+			updates.Add(updateBuilder.Set(x => x.Version, version));
+			updates.Add(updateBuilder.Unset(x => x.RequestRestart));
+			updates.Add(updateBuilder.Unset(x => x.RequestShutdown));
+			updates.Add(updateBuilder.Set(x => x.LastShutdownReason, "Unexpected"));
 
 			// Apply the update
-			return await TryUpdateAsync(Agent, UpdateBuilder.Combine(Updates));
+			return await TryUpdateAsync(agent, updateBuilder.Combine(updates));
 		}
 
 		/// <inheritdoc/>
-		public async Task<IAgent?> TryTerminateSessionAsync(IAgent AgentInterface)
+		public async Task<IAgent?> TryTerminateSessionAsync(IAgent agentInterface)
 		{
-			AgentDocument Agent = (AgentDocument)AgentInterface;
-			UpdateDefinition<AgentDocument> Update = new BsonDocument();
+			AgentDocument agent = (AgentDocument)agentInterface;
+			UpdateDefinition<AgentDocument> update = new BsonDocument();
 
-			Update = Update.Unset(x => x.SessionId);
-			Update = Update.Unset(x => x.SessionExpiresAt);
-			Update = Update.Unset(x => x.Leases);
-			Update = Update.Set(x => x.Status, AgentStatus.Stopped);
+			update = update.Unset(x => x.SessionId);
+			update = update.Unset(x => x.SessionExpiresAt);
+			update = update.Unset(x => x.Leases);
+			update = update.Set(x => x.Status, AgentStatus.Stopped);
 
-			bool bDeleted = Agent.Deleted || Agent.Ephemeral;
-			if (bDeleted != Agent.Deleted)
+			bool bDeleted = agent.Deleted || agent.Ephemeral;
+			if (bDeleted != agent.Deleted)
 			{
-				Update = Update.Set(x => x.Deleted, Agent.Deleted);
+				update = update.Set(x => x.Deleted, agent.Deleted);
 			}
 
-			return await TryUpdateAsync(Agent, Update);
+			return await TryUpdateAsync(agent, update);
 		}
 
 		/// <inheritdoc/>
-		public async Task<IAgent?> TryAddLeaseAsync(IAgent AgentInterface, AgentLease NewLease)
+		public async Task<IAgent?> TryAddLeaseAsync(IAgent agentInterface, AgentLease newLease)
 		{
-			AgentDocument Agent = (AgentDocument)AgentInterface;
+			AgentDocument agent = (AgentDocument)agentInterface;
 
-			List<AgentLease> Leases = new List<AgentLease>();
-			if (Agent.Leases != null)
+			List<AgentLease> leases = new List<AgentLease>();
+			if (agent.Leases != null)
 			{
-				Leases.AddRange(Agent.Leases);
+				leases.AddRange(agent.Leases);
 			}
-			Leases.Add(NewLease);
+			leases.Add(newLease);
 
-			UpdateDefinition<AgentDocument> Update = Builders<AgentDocument>.Update.Set(x => x.Leases, Leases);
-			return await TryUpdateAsync(Agent, Update);
+			UpdateDefinition<AgentDocument> update = Builders<AgentDocument>.Update.Set(x => x.Leases, leases);
+			return await TryUpdateAsync(agent, update);
 		}
 
 		/// <inheritdoc/>
-		public async Task<IAgent?> TryCancelLeaseAsync(IAgent AgentInterface, int LeaseIdx)
+		public async Task<IAgent?> TryCancelLeaseAsync(IAgent agentInterface, int leaseIdx)
 		{
-			AgentDocument Agent = (AgentDocument)AgentInterface;
+			AgentDocument agent = (AgentDocument)agentInterface;
 
-			UpdateDefinition<AgentDocument> Update = Builders<AgentDocument>.Update.Set(x => x.Leases![LeaseIdx].State, LeaseState.Cancelled);
-			IAgent? NewAgent = await TryUpdateAsync(Agent, Update);
-			if (NewAgent != null)
+			UpdateDefinition<AgentDocument> update = Builders<AgentDocument>.Update.Set(x => x.Leases![leaseIdx].State, LeaseState.Cancelled);
+			IAgent? newAgent = await TryUpdateAsync(agent, update);
+			if (newAgent != null)
 			{
-				await PublishUpdateEventAsync(Agent.Id);
+				await PublishUpdateEventAsync(agent.Id);
 			}
-			return NewAgent;
+			return newAgent;
 		}
 
 		/// <inheritdoc/>
-		public IAuditLogChannel<AgentId> GetLogger(AgentId AgentId)
+		public IAuditLogChannel<AgentId> GetLogger(AgentId agentId)
 		{
-			return AuditLog[AgentId];
+			return _auditLog[agentId];
 		}
 
 		/// <inheritdoc/>
-		Task PublishUpdateEventAsync(AgentId AgentId)
+		Task PublishUpdateEventAsync(AgentId agentId)
 		{
-			return RedisService.Database.PublishAsync(UpdateEventChannel, AgentId);
+			return _redisService.Database.PublishAsync(_updateEventChannel, agentId);
 		}
 
 		/// <inheritdoc/>
-		public async Task<IDisposable> SubscribeToUpdateEventsAsync(Action<AgentId> OnUpdate)
+		public async Task<IDisposable> SubscribeToUpdateEventsAsync(Action<AgentId> onUpdate)
 		{
-			return await RedisService.Multiplexer.GetSubscriber().SubscribeAsync(UpdateEventChannel, (Channel, AgentId) => OnUpdate(AgentId));
+			return await _redisService.Multiplexer.GetSubscriber().SubscribeAsync(_updateEventChannel, (channel, agentId) => onUpdate(agentId));
 		}
 	}
 }
