Horde: Fix coding conventions within Horde.Build.

diff --git a/Engine/Source/Programs/Horde/Horde.Build/Agents/Impl/AgentService.cs b/Engine/Source/Programs/Horde/Horde.Build/Agents/Impl/AgentService.cs
--- a/Engine/Source/Programs/Horde/Horde.Build/Agents/Impl/AgentService.cs
+++ b/Engine/Source/Programs/Horde/Horde.Build/Agents/Impl/AgentService.cs
@@ -1,30 +1,29 @@
-﻿// Copyright Epic Games, Inc. All Rights Reserved.
+// Copyright Epic Games, Inc. All Rights Reserved.
 
+using System;
+using System.Collections.Generic;
+using System.Diagnostics;
+using System.Linq;
+using System.Security.Claims;
+using System.Threading;
+using System.Threading.Tasks;
 using EpicGames.Core;
 using EpicGames.Redis;
 using EpicGames.Serialization;
 using Google.Protobuf.WellKnownTypes;
-using HordeCommon;
 using Horde.Build.Acls;
 using Horde.Build.Collections;
 using Horde.Build.Models;
 using Horde.Build.Tasks;
 using Horde.Build.Utilities;
+using HordeCommon;
 using Microsoft.Extensions.Hosting;
 using Microsoft.Extensions.Logging;
-using MongoDB.Bson;
 using MongoDB.Driver;
-using StackExchange.Redis;
-using StatsdClient;
-using System;
-using System.Collections.Generic;
-using System.Diagnostics;
-using System.Linq;
-using System.Security.Claims;
-using System.Threading;
-using System.Threading.Tasks;
 using OpenTracing;
 using OpenTracing.Util;
+using StackExchange.Redis;
+using StatsdClient;
 
 namespace Horde.Build.Services
 {
@@ -65,209 +64,208 @@ namespace Horde.Build.Services
 		/// Time after which a session will be renewed
 		/// </summary>
 		public static readonly TimeSpan SessionRenewTime = TimeSpan.FromSeconds(50);
-
-		AclService AclService;
-		IDowntimeService DowntimeService;
+		readonly AclService _aclService;
+		readonly IDowntimeService _downtimeService;
 
 		/// <summary>
 		/// Collection of agent documents
 		/// </summary>
 		public IAgentCollection Agents { get; }
 
-		ILeaseCollection Leases;
-		ISessionCollection Sessions;
-		IDogStatsd DogStatsd;
-		ITaskSource[] TaskSources;
-		IHostApplicationLifetime ApplicationLifetime;
-		IClock Clock;
-		ILogger Logger;
-		ITicker Ticker;
+		readonly ILeaseCollection _leases;
+		readonly ISessionCollection _sessions;
+		readonly IDogStatsd _dogStatsd;
+		readonly ITaskSource[] _taskSources;
+		readonly IHostApplicationLifetime _applicationLifetime;
+		readonly IClock _clock;
+		readonly ILogger _logger;
+		readonly ITicker _ticker;
 		
-		RedisString<AgentRateTable> AgentRateTableData;
-		RedisService RedisService;
+		readonly RedisString<AgentRateTable> _agentRateTableData;
+		readonly RedisService _redisService;
 
 		// Lazily updated costs for different agent types
-		AsyncCachedValue<AgentRateTable?> AgentRateTable;
-		
+		readonly AsyncCachedValue<AgentRateTable?> _agentRateTable;
+
 		// Lazily updated list of current pools
-		AsyncCachedValue<List<IPool>> PoolsList;
+		readonly AsyncCachedValue<List<IPool>> _poolsList;
 
 		// All the agents currently performing a long poll for work on this server
-		Dictionary<AgentId, CancellationTokenSource> WaitingAgents = new Dictionary<AgentId, CancellationTokenSource>();
-		
+		readonly Dictionary<AgentId, CancellationTokenSource> _waitingAgents = new Dictionary<AgentId, CancellationTokenSource>();
+
 		// Subscription for update events
-		IDisposable Subscription;
+		readonly IDisposable _subscription;
 
 		/// <summary>
 		/// Constructor
 		/// </summary>
-		public AgentService(IAgentCollection Agents, ILeaseCollection Leases, ISessionCollection Sessions, AclService AclService, IDowntimeService DowntimeService, IPoolCollection PoolCollection, IDogStatsd DogStatsd, IEnumerable<ITaskSource> TaskSources, RedisService RedisService, IHostApplicationLifetime ApplicationLifetime, IClock Clock, ILogger<AgentService> Logger)
+		public AgentService(IAgentCollection agents, ILeaseCollection leases, ISessionCollection sessions, AclService aclService, IDowntimeService downtimeService, IPoolCollection poolCollection, IDogStatsd dogStatsd, IEnumerable<ITaskSource> taskSources, RedisService redisService, IHostApplicationLifetime applicationLifetime, IClock clock, ILogger<AgentService> logger)
 		{
-			this.Agents = Agents;
-			this.Leases = Leases;
-			this.Sessions = Sessions;
-			this.AclService = AclService;
-			this.DowntimeService = DowntimeService;
-			this.AgentRateTableData = new RedisString<AgentRateTable>(RedisService.Database, "agent-rates");
-			this.AgentRateTable = new AsyncCachedValue<AgentRateTable?>(() => AgentRateTableData.GetAsync(), TimeSpan.FromSeconds(2.0));//.FromMinutes(5.0));
-			this.PoolsList = new AsyncCachedValue<List<IPool>>(() => PoolCollection.GetAsync(), TimeSpan.FromSeconds(30.0));
-			this.DogStatsd = DogStatsd;
-			this.TaskSources = TaskSources.ToArray();
-			this.ApplicationLifetime = ApplicationLifetime;
-			this.RedisService = RedisService;
-			this.Clock = Clock;
-			this.Ticker = Clock.AddTicker(TimeSpan.FromSeconds(30.0), TickAsync, Logger);
-			this.Logger = Logger;
-
-			Subscription = Agents.SubscribeToUpdateEventsAsync(OnAgentUpdate).Result;
+			Agents = agents;
+			_leases = leases;
+			_sessions = sessions;
+			_aclService = aclService;
+			_downtimeService = downtimeService;
+			_agentRateTableData = new RedisString<AgentRateTable>(redisService.Database, "agent-rates");
+			_agentRateTable = new AsyncCachedValue<AgentRateTable?>(() => _agentRateTableData.GetAsync(), TimeSpan.FromSeconds(2.0));//.FromMinutes(5.0));
+			_poolsList = new AsyncCachedValue<List<IPool>>(() => poolCollection.GetAsync(), TimeSpan.FromSeconds(30.0));
+			_dogStatsd = dogStatsd;
+			_taskSources = taskSources.ToArray();
+			_applicationLifetime = applicationLifetime;
+			_redisService = redisService;
+			_clock = clock;
+			_ticker = clock.AddTicker<AgentService>(TimeSpan.FromSeconds(30.0), TickAsync, logger);
+			_logger = logger;
+
+			_subscription = agents.SubscribeToUpdateEventsAsync(OnAgentUpdate).Result;
 		}
 
 		/// <inheritdoc/>
-		public Task StartAsync(CancellationToken cancellationToken) => Ticker.StartAsync();
+		public Task StartAsync(CancellationToken cancellationToken) => _ticker.StartAsync();
 
 		/// <inheritdoc/>
-		public Task StopAsync(CancellationToken cancellationToken) => Ticker.StopAsync();
+		public Task StopAsync(CancellationToken cancellationToken) => _ticker.StopAsync();
 
 		/// <inheritdoc/>
 		public void Dispose()
 		{
-			Subscription.Dispose();
-			Ticker.Dispose();
+			_subscription.Dispose();
+			_ticker.Dispose();
 		}
 
 		/// <summary>
 		/// Gets user-readable payload information
 		/// </summary>
-		/// <param name="Payload">The payload data</param>
+		/// <param name="payload">The payload data</param>
 		/// <returns>Dictionary of key/value pairs for the payload</returns>
-		public Dictionary<string, string>? GetPayloadDetails(ReadOnlyMemory<byte>? Payload)
+		public Dictionary<string, string>? GetPayloadDetails(ReadOnlyMemory<byte>? payload)
 		{
-			Dictionary<string, string>? Details = null;
-			if (Payload != null)
+			Dictionary<string, string>? details = null;
+			if (payload != null)
 			{
-				Any BasePayload = Any.Parser.ParseFrom(Payload.Value.ToArray());
-				foreach (ITaskSource TaskSource in TaskSources)
+				Any basePayload = Any.Parser.ParseFrom(payload.Value.ToArray());
+				foreach (ITaskSource taskSource in _taskSources)
 				{
-					if (BasePayload.Is(TaskSource.Descriptor))
+					if (basePayload.Is(taskSource.Descriptor))
 					{
-						Details = new Dictionary<string, string>();
-						TaskSource.GetLeaseDetails(BasePayload, Details);
+						details = new Dictionary<string, string>();
+						taskSource.GetLeaseDetails(basePayload, details);
 						break;
 					}
 				}
 			}
-			return Details;
+			return details;
 		}
 
 		/// <summary>
 		/// Issues a bearer token for the given session id
 		/// </summary>
-		/// <param name="AgentId">The agent id</param>
-		/// <param name="SessionId">The session id</param>
+		/// <param name="agentId">The agent id</param>
+		/// <param name="sessionId">The session id</param>
 		/// <returns>Bearer token for the agent</returns>
-		public string IssueSessionToken(AgentId AgentId, SessionId SessionId)
+		public string IssueSessionToken(AgentId agentId, SessionId sessionId)
 		{
-			List<AclClaim> Claims = new List<AclClaim>();
-			Claims.Add(AclService.AgentRoleClaim);
-			Claims.Add(AclService.GetAgentClaim(AgentId));
-			Claims.Add(AclService.GetSessionClaim(SessionId));
-			return AclService.IssueBearerToken(Claims, null);
+			List<AclClaim> claims = new List<AclClaim>();
+			claims.Add(AclService.AgentRoleClaim);
+			claims.Add(AclService.GetAgentClaim(agentId));
+			claims.Add(AclService.GetSessionClaim(sessionId));
+			return _aclService.IssueBearerToken(claims, null);
 		}
 
 		/// <summary>
 		/// Register a new agent
 		/// </summary>
-		/// <param name="Name">Name of the agent</param>
+		/// <param name="name">Name of the agent</param>
 		/// <param name="bEnabled">Whether the agent is currently enabled</param>
-		/// <param name="Channel">Override for the desired software version</param>
-		/// <param name="Pools">Pools for this agent</param>
+		/// <param name="channel">Override for the desired software version</param>
+		/// <param name="pools">Pools for this agent</param>
 		/// <returns>Unique id for the agent</returns>
-		public Task<IAgent> CreateAgentAsync(string Name, bool bEnabled, AgentSoftwareChannelName? Channel, List<PoolId>? Pools)
+		public Task<IAgent> CreateAgentAsync(string name, bool bEnabled, AgentSoftwareChannelName? channel, List<PoolId>? pools)
 		{
-			return Agents.AddAsync(new AgentId(Name), bEnabled, Channel, Pools);
+			return Agents.AddAsync(new AgentId(name), bEnabled, channel, pools);
 		}
 
 		/// <summary>
 		/// Gets an agent by ID
 		/// </summary>
-		/// <param name="AgentId">Unique id of the agent</param>
+		/// <param name="agentId">Unique id of the agent</param>
 		/// <returns>The agent document</returns>
-		public Task<IAgent?> GetAgentAsync(AgentId AgentId)
+		public Task<IAgent?> GetAgentAsync(AgentId agentId)
 		{
-			return Agents.GetAsync(AgentId);
+			return Agents.GetAsync(agentId);
 		}
 
 		/// <summary>
 		/// Finds all agents matching certain criteria
 		/// </summary>
-		/// <param name="PoolId">The pool containing the agent</param>
-		/// <param name="ModifiedAfter">If set, only returns agents modified after this time</param>
-		/// <param name="Index">Index within the list of results</param>
-		/// <param name="Count">Number of results to return</param>
+		/// <param name="poolId">The pool containing the agent</param>
+		/// <param name="modifiedAfter">If set, only returns agents modified after this time</param>
+		/// <param name="index">Index within the list of results</param>
+		/// <param name="count">Number of results to return</param>
 		/// <returns>List of agents matching the given criteria</returns>
-		public Task<List<IAgent>> FindAgentsAsync(PoolId? PoolId, DateTime? ModifiedAfter, int? Index, int? Count)
+		public Task<List<IAgent>> FindAgentsAsync(PoolId? poolId, DateTime? modifiedAfter, int? index, int? count)
 		{
-			return Agents.FindAsync(PoolId, ModifiedAfter, null, Index, Count);
+			return Agents.FindAsync(poolId, modifiedAfter, null, index, count);
 		}
 
 		/// <summary>
 		/// Update the current workspaces for an agent.
 		/// </summary>
-		/// <param name="Agent">The agent to update</param>
-		/// <param name="Workspaces">Current list of workspaces</param>
+		/// <param name="agent">The agent to update</param>
+		/// <param name="workspaces">Current list of workspaces</param>
 		/// <param name="bPendingConform">Whether the agent still needs to run another conform</param>
 		/// <returns>New agent state</returns>
-		public async Task<bool> TryUpdateWorkspacesAsync(IAgent Agent, List<AgentWorkspace> Workspaces, bool bPendingConform)
+		public async Task<bool> TryUpdateWorkspacesAsync(IAgent agent, List<AgentWorkspace> workspaces, bool bPendingConform)
 		{
-			IAgent? NewAgent = await Agents.TryUpdateWorkspacesAsync(Agent, Workspaces, bPendingConform);
-			return NewAgent != null;
+			IAgent? newAgent = await Agents.TryUpdateWorkspacesAsync(agent, workspaces, bPendingConform);
+			return newAgent != null;
 		}
 
 		/// <summary>
 		/// Marks the agent as deleted
 		/// </summary>
-		/// <param name="Agent">The agent to delete</param>
+		/// <param name="agent">The agent to delete</param>
 		/// <returns>Async task</returns>
-		public async Task DeleteAgentAsync(IAgent? Agent)
+		public async Task DeleteAgentAsync(IAgent? agent)
 		{
-			while (Agent != null && !Agent.Deleted)
+			while (agent != null && !agent.Deleted)
 			{
-				IAgent? NewAgent = await Agents.TryDeleteAsync(Agent);
-				if(NewAgent != null)
+				IAgent? newAgent = await Agents.TryDeleteAsync(agent);
+				if(newAgent != null)
 				{
 					break;
 				}
-				Agent = await GetAgentAsync(Agent.Id);
+				agent = await GetAgentAsync(agent.Id);
 			}
 		}
 
-		async ValueTask<List<PoolId>> GetDynamicPoolsAsync(IAgent Agent)
+		async ValueTask<List<PoolId>> GetDynamicPoolsAsync(IAgent agent)
 		{
-			List<PoolId> NewDynamicPools = new List<PoolId>();
+			List<PoolId> newDynamicPools = new List<PoolId>();
 
-			List<IPool> Pools = await PoolsList.GetAsync();
-			foreach (IPool Pool in Pools)
+			List<IPool> pools = await _poolsList.GetAsync();
+			foreach (IPool pool in pools)
 			{
-				if (Pool.Condition != null && Agent.SatisfiesCondition(Pool.Condition))
+				if (pool.Condition != null && agent.SatisfiesCondition(pool.Condition))
 				{
-					NewDynamicPools.Add(Pool.Id);
+					newDynamicPools.Add(pool.Id);
 				}
 			}
 
-			return NewDynamicPools;
+			return newDynamicPools;
 		}
 
 		/// <summary>
 		/// Callback for an agents 
 		/// </summary>
-		/// <param name="AgentId"></param>
-		void OnAgentUpdate(AgentId AgentId)
+		/// <param name="agentId"></param>
+		void OnAgentUpdate(AgentId agentId)
 		{
-			lock (WaitingAgents)
+			lock (_waitingAgents)
 			{
-				if (WaitingAgents.TryGetValue(AgentId, out CancellationTokenSource? CancellationSource))
+				if (_waitingAgents.TryGetValue(agentId, out CancellationTokenSource? cancellationSource))
 				{
-					CancellationSource.Cancel();
+					cancellationSource.Cancel();
 				}
 			}
 		}
@@ -275,81 +273,81 @@ namespace Horde.Build.Services
 		/// <summary>
 		/// Creates a new agent session
 		/// </summary>
-		/// <param name="Agent">The agent to create a session for</param>
-		/// <param name="Status">Current status of the agent</param>
-		/// <param name="Properties">Properties for the agent</param>
-		/// <param name="Resources">Resources which the agent has</param>
-		/// <param name="Version">Version of the software that's running</param>
+		/// <param name="agent">The agent to create a session for</param>
+		/// <param name="status">Current status of the agent</param>
+		/// <param name="properties">Properties for the agent</param>
+		/// <param name="resources">Resources which the agent has</param>
+		/// <param name="version">Version of the software that's running</param>
 		/// <returns>New agent state</returns>
-		public async Task<IAgent> CreateSessionAsync(IAgent Agent, AgentStatus Status, IReadOnlyList<string> Properties, IReadOnlyDictionary<string, int> Resources, string? Version)
+		public async Task<IAgent> CreateSessionAsync(IAgent agent, AgentStatus status, IReadOnlyList<string> properties, IReadOnlyDictionary<string, int> resources, string? version)
 		{
 			for (; ; )
 			{
-				IAuditLogChannel<AgentId> AgentLogger = Agents.GetLogger(Agent.Id);
+				IAuditLogChannel<AgentId> agentLogger = Agents.GetLogger(agent.Id);
 
 				// Check if there's already a session running for this agent.
-				IAgent? NewAgent;
-				if (Agent.SessionId != null)
+				IAgent? newAgent;
+				if (agent.SessionId != null)
 				{
 					// Try to terminate the current session
-					await TryTerminateSessionAsync(Agent);
+					await TryTerminateSessionAsync(agent);
 				}
 				else
 				{
-					DateTime UtcNow = Clock.UtcNow;
+					DateTime utcNow = _clock.UtcNow;
 
 					// Remove any outstanding leases
-					foreach (AgentLease Lease in Agent.Leases)
+					foreach (AgentLease lease in agent.Leases)
 					{
-						AgentLogger.LogInformation("Removing outstanding lease {LeaseId}", Lease.Id);
-						await RemoveLeaseAsync(Agent, Lease, UtcNow, LeaseOutcome.Failed, null);
+						agentLogger.LogInformation("Removing outstanding lease {LeaseId}", lease.Id);
+						await RemoveLeaseAsync(agent, lease, utcNow, LeaseOutcome.Failed, null);
 					}
 
 					// Create a new session document
-					ISession NewSession = await Sessions.AddAsync(SessionId.GenerateNewId(), Agent.Id, Clock.UtcNow, Properties, Resources, Version);
-					DateTime SessionExpiresAt = UtcNow + SessionExpiryTime;
+					ISession newSession = await _sessions.AddAsync(SessionId.GenerateNewId(), agent.Id, _clock.UtcNow, properties, resources, version);
+					DateTime sessionExpiresAt = utcNow + SessionExpiryTime;
 
 					// Get the new dynamic pools for the agent
-					List<PoolId> DynamicPools = await GetDynamicPoolsAsync(Agent);
+					List<PoolId> dynamicPools = await GetDynamicPoolsAsync(agent);
 
 					// Reset the agent to use the new session
-					NewAgent = await Agents.TryStartSessionAsync(Agent, NewSession.Id, SessionExpiresAt, Status, Properties, Resources, DynamicPools, Version);
-					if(NewAgent != null)
+					newAgent = await Agents.TryStartSessionAsync(agent, newSession.Id, sessionExpiresAt, status, properties, resources, dynamicPools, version);
+					if(newAgent != null)
 					{
-						Agent = NewAgent;
-						AgentLogger.LogInformation("Session {SessionId} started", NewSession.Id);
+						agent = newAgent;
+						agentLogger.LogInformation("Session {SessionId} started", newSession.Id);
 						break;
 					}
 
 					// Remove the session we didn't use
-					await Sessions.DeleteAsync(NewSession.Id);
+					await _sessions.DeleteAsync(newSession.Id);
 				}
 
 				// Get the current agent state
-				NewAgent = await GetAgentAsync(Agent.Id);
-				if (NewAgent == null)
+				newAgent = await GetAgentAsync(agent.Id);
+				if (newAgent == null)
 				{
-					throw new InvalidOperationException($"Invalid agent id '{Agent.Id}'");
+					throw new InvalidOperationException($"Invalid agent id '{agent.Id}'");
 				}
-				Agent = NewAgent;
+				agent = newAgent;
 			}
-			return Agent;
+			return agent;
 		}
 
-		async Task<(ITaskSource, AgentLease)?> GuardedWaitForLeaseAsync(ITaskSource Source, IAgent Agent, CancellationToken CancellationToken)
+		async Task<(ITaskSource, AgentLease)?> GuardedWaitForLeaseAsync(ITaskSource source, IAgent agent, CancellationToken cancellationToken)
 		{
 			try
 			{
-				AgentLease? Lease = await Source.AssignLeaseAsync(Agent, CancellationToken);
-				return (Lease != null) ? (Source, Lease) : null;
+				AgentLease? lease = await source.AssignLeaseAsync(agent, cancellationToken);
+				return (lease != null) ? (source, lease) : null;
 			}
 			catch (TaskCanceledException)
 			{
 				return null;
 			}
-			catch (Exception Ex)
+			catch (Exception ex)
 			{
-				Logger.LogError(Ex, "Exception while trying to assign lease"); 
+				_logger.LogError(ex, "Exception while trying to assign lease"); 
 				return null;
 			}
 		}
@@ -357,14 +355,14 @@ namespace Horde.Build.Services
 		/// <summary>
 		/// Determines whether a task source can currently issue tasks
 		/// </summary>
-		bool CanUseTaskSource(IAgent Agent, ITaskSource TaskSource)
+		bool CanUseTaskSource(IAgent agent, ITaskSource taskSource)
 		{
-			TaskSourceFlags Flags = TaskSource.Flags;
-			if ((Flags & TaskSourceFlags.AllowWhenDisabled) == 0 && !Agent.Enabled)
+			TaskSourceFlags flags = taskSource.Flags;
+			if ((flags & TaskSourceFlags.AllowWhenDisabled) == 0 && !agent.Enabled)
 			{
 				return false;
 			}
-			if ((Flags & TaskSourceFlags.AllowDuringDowntime) == 0 && DowntimeService.IsDowntimeActive)
+			if ((flags & TaskSourceFlags.AllowDuringDowntime) == 0 && _downtimeService.IsDowntimeActive)
 			{
 				return false;
 			}
@@ -374,212 +372,212 @@ namespace Horde.Build.Services
 		/// <summary>
 		/// Waits for a lease to be assigned to an agent
 		/// </summary>
-		/// <param name="Agent">The agent to assign a lease to</param>
-		/// <param name="CancellationToken"></param>
+		/// <param name="agent">The agent to assign a lease to</param>
+		/// <param name="cancellationToken"></param>
 		/// <returns>True if a lease was assigned, false otherwise</returns>
-		public async Task<IAgent?> WaitForLeaseAsync(IAgent? Agent, CancellationToken CancellationToken)
+		public async Task<IAgent?> WaitForLeaseAsync(IAgent? agent, CancellationToken cancellationToken)
 		{
-			while (Agent != null)
+			while (agent != null)
 			{
-				if (!Agent.SessionExpiresAt.HasValue)
+				if (!agent.SessionExpiresAt.HasValue)
 				{
 					break;
 				}
 
 				// Check we have some time to wait
-				DateTime UtcNow = Clock.UtcNow;
-				TimeSpan MaxWaitTime = (Agent.SessionExpiresAt.Value - SessionExpiryTime + SessionLongPollTime) - UtcNow;
-				if (MaxWaitTime <= TimeSpan.Zero)
+				DateTime utcNow = _clock.UtcNow;
+				TimeSpan maxWaitTime = (agent.SessionExpiresAt.Value - SessionExpiryTime + SessionLongPollTime) - utcNow;
+				if (maxWaitTime <= TimeSpan.Zero)
 				{
 					break;
 				}
 
 				// Create a cancellation token that will expire with the session
-				using CancellationTokenSource CancellationSource = CancellationTokenSource.CreateLinkedTokenSource(ApplicationLifetime.ApplicationStopping, CancellationToken);
-				CancellationSource.CancelAfter(MaxWaitTime);
+				using CancellationTokenSource cancellationSource = CancellationTokenSource.CreateLinkedTokenSource(_applicationLifetime.ApplicationStopping, cancellationToken);
+				cancellationSource.CancelAfter(maxWaitTime);
 
 				// Assign a new lease
-				(ITaskSource, AgentLease)? Result = null;
+				(ITaskSource, AgentLease)? result = null;
 				try
 				{
 					// Add the cancellation source to the set of waiting agents
-					lock (WaitingAgents)
+					lock (_waitingAgents)
 					{
-						WaitingAgents[Agent.Id] = CancellationSource;
+						_waitingAgents[agent.Id] = cancellationSource;
 					}
 
 					// Create all the tasks to wait for
-					List<Task<(ITaskSource, AgentLease)?>> Tasks = new List<Task<(ITaskSource, AgentLease)?>>();
-					foreach (ITaskSource TaskSource in TaskSources)
+					List<Task<(ITaskSource, AgentLease)?>> tasks = new List<Task<(ITaskSource, AgentLease)?>>();
+					foreach (ITaskSource taskSource in _taskSources)
 					{
-						if (CanUseTaskSource(Agent, TaskSource))
+						if (CanUseTaskSource(agent, taskSource))
 						{
-							Tasks.Add(GuardedWaitForLeaseAsync(TaskSource, Agent, CancellationSource.Token));
+							tasks.Add(GuardedWaitForLeaseAsync(taskSource, agent, cancellationSource.Token));
 						}
 					}
 
 					// If no task source is valid, just add a delay
-					if (Tasks.Count == 0)
+					if (tasks.Count == 0)
 					{
-						await AsyncUtils.DelayNoThrow(MaxWaitTime, CancellationToken);
+						await AsyncUtils.DelayNoThrow(maxWaitTime, cancellationToken);
 						break;
 					}
 
 					// Wait for a lease to be available
-					while (Tasks.Count > 0)
+					while (tasks.Count > 0)
 					{
-						await Task.WhenAny(Tasks);
+						await Task.WhenAny(tasks);
 
-						for (int Idx = 0; Idx < Tasks.Count; Idx++)
+						for (int idx = 0; idx < tasks.Count; idx++)
 						{
-							(ITaskSource, AgentLease)? TaskResult;
-							if (!Tasks[Idx].TryGetResult(out TaskResult))
+							(ITaskSource, AgentLease)? taskResult;
+							if (!tasks[idx].TryGetResult(out taskResult))
 							{
 								continue;
 							}
 
-							Tasks.RemoveAt(Idx--);
+							tasks.RemoveAt(idx--);
 
-							if (!TaskResult.HasValue)
+							if (!taskResult.HasValue)
 							{
 								continue;
 							}
 
-							if (Result == null)
+							if (result == null)
 							{
-								Result = TaskResult;
-								CancellationSource.Cancel();
+								result = taskResult;
+								cancellationSource.Cancel();
 							}
-							else if (TaskResult.Value.Item2 != AgentLease.Drain)
+							else if (taskResult.Value.Item2 != AgentLease.Drain)
 							{
-								(ITaskSource TaskSource, AgentLease TaskLease) = TaskResult.Value;
-								await TaskSource.CancelLeaseAsync(Agent, TaskLease.Id, Any.Parser.ParseFrom(TaskLease.Payload));
+								(ITaskSource taskSource, AgentLease taskLease) = taskResult.Value;
+								await taskSource.CancelLeaseAsync(agent, taskLease.Id, Any.Parser.ParseFrom(taskLease.Payload));
 							}
 						}
 					}
 				}
 				finally
 				{
-					lock (WaitingAgents)
+					lock (_waitingAgents)
 					{
-						WaitingAgents.Remove(Agent.Id);
+						_waitingAgents.Remove(agent.Id);
 					}
 				}
 
 				// Exit if we didn't find any work to do. It may be that all the task sources returned null, in which case wait for the time period to expire.
-				if (Result == null)
+				if (result == null)
 				{
-					if (!CancellationSource.IsCancellationRequested)
+					if (!cancellationSource.IsCancellationRequested)
 					{
-						await CancellationSource.Token.AsTask();
+						await cancellationSource.Token.AsTask();
 					}
 					break;
 				}
 
 				// Get the resulting lease
-				(ITaskSource Source, AgentLease Lease) = Result.Value;
-				if (Lease == AgentLease.Drain)
+				(ITaskSource source, AgentLease lease) = result.Value;
+				if (lease == AgentLease.Drain)
 				{
-					await AsyncUtils.DelayNoThrow(MaxWaitTime, CancellationToken);
+					await AsyncUtils.DelayNoThrow(maxWaitTime, cancellationToken);
 					break;
 				}
 
 				// Add the new lease to the agent
-				IAgent? NewAgent = await Agents.TryAddLeaseAsync(Agent, Lease);
-				if (NewAgent != null)
+				IAgent? newAgent = await Agents.TryAddLeaseAsync(agent, lease);
+				if (newAgent != null)
 				{
-					await Source.OnLeaseStartedAsync(NewAgent, Lease.Id, Any.Parser.ParseFrom(Lease.Payload), Agents.GetLogger(Agent.Id));
-					await CreateLeaseAsync(Agent, Lease);
-					return NewAgent;
+					await source.OnLeaseStartedAsync(newAgent, lease.Id, Any.Parser.ParseFrom(lease.Payload), Agents.GetLogger(agent.Id));
+					await CreateLeaseAsync(agent, lease);
+					return newAgent;
 				}
 
 				// Update the agent
-				Agent = await GetAgentAsync(Agent.Id);
+				agent = await GetAgentAsync(agent.Id);
 			}
-			return Agent;
+			return agent;
 		}
 
 		/// <summary>
 		/// Cancels the specified agent lease
 		/// </summary>
-		/// <param name="Agent">Agent to cancel the lease on</param>
-		/// <param name="LeaseId">The lease id to cancel</param>
+		/// <param name="agent">Agent to cancel the lease on</param>
+		/// <param name="leaseId">The lease id to cancel</param>
 		/// <returns></returns>
-		public async Task<bool> CancelLeaseAsync(IAgent Agent, LeaseId LeaseId)
+		public async Task<bool> CancelLeaseAsync(IAgent agent, LeaseId leaseId)
 		{		
-			int Index = 0;
-			while (Index < Agent.Leases.Count && Agent.Leases[Index].Id != LeaseId)
+			int index = 0;
+			while (index < agent.Leases.Count && agent.Leases[index].Id != leaseId)
 			{
-				Index++;
+				index++;
 			}
 
-			if (Index == Agent.Leases.Count)
+			if (index == agent.Leases.Count)
 			{
 				return false;
 			}
 
-			if (Agent.Leases[Index].State == LeaseState.Cancelled)
+			if (agent.Leases[index].State == LeaseState.Cancelled)
 			{
 				return false;
 			}
 
-			await Agents.TryCancelLeaseAsync(Agent, Index);
+			await Agents.TryCancelLeaseAsync(agent, index);
 			return true;
 		}
 
 		/// <summary>
 		/// 
 		/// </summary>
-		public async Task<IAgent?> UpdateSessionWithWaitAsync(IAgent InAgent, SessionId SessionId, AgentStatus Status, IReadOnlyList<string>? Properties, IReadOnlyDictionary<string, int>? Resources, IList<HordeCommon.Rpc.Messages.Lease> NewLeases, CancellationToken CancellationToken)
+		public async Task<IAgent?> UpdateSessionWithWaitAsync(IAgent inAgent, SessionId sessionId, AgentStatus status, IReadOnlyList<string>? properties, IReadOnlyDictionary<string, int>? resources, IList<HordeCommon.Rpc.Messages.Lease> newLeases, CancellationToken cancellationToken)
 		{
-			IAgent? Agent = InAgent;
+			IAgent? agent = inAgent;
 
 			// Capture the current agent update index. This allows us to detect if anything has changed.
-			uint UpdateIndex = Agent.UpdateIndex;
+			uint updateIndex = agent.UpdateIndex;
 
 			// Update the agent session and return to the caller if anything changes
-			Agent = await UpdateSessionAsync(Agent, SessionId, Status, Properties, Resources, NewLeases);
-			if (Agent != null && Agent.UpdateIndex == UpdateIndex && (Agent.Leases.Count > 0 || Agent.Status != AgentStatus.Stopping))
+			agent = await UpdateSessionAsync(agent, sessionId, status, properties, resources, newLeases);
+			if (agent != null && agent.UpdateIndex == updateIndex && (agent.Leases.Count > 0 || agent.Status != AgentStatus.Stopping))
 			{
-				Agent = await WaitForLeaseAsync(Agent, CancellationToken);
+				agent = await WaitForLeaseAsync(agent, cancellationToken);
 			}
-			return Agent;
+			return agent;
 		}
 
 		/// <summary>
 		/// Updates the state of the current agent session
 		/// </summary>
-		/// <param name="InAgent">The current agent state</param>
-		/// <param name="SessionId">Id of the session</param>
-		/// <param name="Status">New status for the agent</param>
-		/// <param name="Properties">New agent properties</param>
-		/// <param name="Resources">New agent resources</param>
-		/// <param name="NewLeases">New list of leases for this session</param>
+		/// <param name="inAgent">The current agent state</param>
+		/// <param name="sessionId">Id of the session</param>
+		/// <param name="status">New status for the agent</param>
+		/// <param name="properties">New agent properties</param>
+		/// <param name="resources">New agent resources</param>
+		/// <param name="newLeases">New list of leases for this session</param>
 		/// <returns>Updated agent state</returns>
-		public async Task<IAgent?> UpdateSessionAsync(IAgent InAgent, SessionId SessionId, AgentStatus Status, IReadOnlyList<string>? Properties, IReadOnlyDictionary<string, int>? Resources, IList<HordeCommon.Rpc.Messages.Lease> NewLeases)
+		public async Task<IAgent?> UpdateSessionAsync(IAgent inAgent, SessionId sessionId, AgentStatus status, IReadOnlyList<string>? properties, IReadOnlyDictionary<string, int>? resources, IList<HordeCommon.Rpc.Messages.Lease> newLeases)
 		{
-			DateTime UtcNow = Clock.UtcNow;
+			DateTime utcNow = _clock.UtcNow;
 
-			Stopwatch Timer = Stopwatch.StartNew();
+			Stopwatch timer = Stopwatch.StartNew();
 
-			IAgent? Agent = InAgent;
-			while (Agent != null)
+			IAgent? agent = inAgent;
+			while (agent != null)
 			{
 				// If the agent is stopping and doesn't have any leases, we can terminate the current session.
-				if (Status == AgentStatus.Stopping && NewLeases.Count == 0)
+				if (status == AgentStatus.Stopping && newLeases.Count == 0)
 				{
 					// If we've already decided to terminate the session, this update is redundant but harmless
-					if (Agent.SessionId != SessionId)
+					if (agent.SessionId != sessionId)
 					{
 						break;
 					}
 
 					// If the session is valid, we can terminate once the agent leases are also empty
-					if (Agent.Leases.Count == 0)
+					if (agent.Leases.Count == 0)
 					{
-						if (!await TryTerminateSessionAsync(Agent))
+						if (!await TryTerminateSessionAsync(agent))
 						{
-							Agent = await GetAgentAsync(Agent.Id);
+							agent = await GetAgentAsync(agent.Id);
 							continue;
 						}
 						break;
@@ -587,13 +585,13 @@ namespace Horde.Build.Services
 				}
 
 				// Check the session id is correct.
-				if (Agent.SessionId != SessionId)
+				if (agent.SessionId != sessionId)
 				{
-					throw new InvalidOperationException($"Invalid agent session {SessionId}");
+					throw new InvalidOperationException($"Invalid agent session {sessionId}");
 				}
 
 				// Check the session hasn't expired
-				if (!Agent.IsSessionValid(UtcNow))
+				if (!agent.IsSessionValid(utcNow))
 				{
 					throw new InvalidOperationException("Session has already expired");
 				}
@@ -601,44 +599,44 @@ namespace Horde.Build.Services
 				// Extend the current session time if we're within a time period of the current time expiring. This reduces
 				// unnecessary DB writes a little, but it also allows us to skip the update and jump into a long poll state
 				// if there's still time on the session left.
-				DateTime? SessionExpiresAt = null;
-				if (!Agent.SessionExpiresAt.HasValue || UtcNow > (Agent.SessionExpiresAt - SessionExpiryTime) + SessionRenewTime)
+				DateTime? sessionExpiresAt = null;
+				if (!agent.SessionExpiresAt.HasValue || utcNow > (agent.SessionExpiresAt - SessionExpiryTime) + SessionRenewTime)
 				{
-					SessionExpiresAt = UtcNow + SessionExpiryTime;
+					sessionExpiresAt = utcNow + SessionExpiryTime;
 				}
 
 				// Flag for whether the leases array should be updated
 				bool bUpdateLeases = false;
-				List<AgentLease> Leases = new List<AgentLease>(Agent.Leases);
+				List<AgentLease> leases = new List<AgentLease>(agent.Leases);
 
 				// Remove any completed leases from the agent
-				Dictionary<LeaseId, HordeCommon.Rpc.Messages.Lease> LeaseIdToNewState = NewLeases.ToDictionary(x => new LeaseId(x.Id), x => x);
-				for (int Idx = 0; Idx < Leases.Count; Idx++)
+				Dictionary<LeaseId, HordeCommon.Rpc.Messages.Lease> leaseIdToNewState = newLeases.ToDictionary(x => new LeaseId(x.Id), x => x);
+				for (int idx = 0; idx < leases.Count; idx++)
 				{
-					AgentLease Lease = Leases[Idx];
-					if (Lease.State == LeaseState.Cancelled)
+					AgentLease lease = leases[idx];
+					if (lease.State == LeaseState.Cancelled)
 					{
-						HordeCommon.Rpc.Messages.Lease? NewLease;
-						if (!LeaseIdToNewState.TryGetValue(Lease.Id, out NewLease) || NewLease.State == LeaseState.Cancelled || NewLease.State == LeaseState.Completed)
+						HordeCommon.Rpc.Messages.Lease? newLease;
+						if (!leaseIdToNewState.TryGetValue(lease.Id, out newLease) || newLease.State == LeaseState.Cancelled || newLease.State == LeaseState.Completed)
 						{
-							await RemoveLeaseAsync(Agent, Lease, UtcNow, LeaseOutcome.Cancelled, null);
-							Leases.RemoveAt(Idx--);
+							await RemoveLeaseAsync(agent, lease, utcNow, LeaseOutcome.Cancelled, null);
+							leases.RemoveAt(idx--);
 							bUpdateLeases = true;
 						}
 					}
 					else
 					{
-						HordeCommon.Rpc.Messages.Lease? NewLease;
-						if (LeaseIdToNewState.TryGetValue(Lease.Id, out NewLease) && NewLease.State != Lease.State)
+						HordeCommon.Rpc.Messages.Lease? newLease;
+						if (leaseIdToNewState.TryGetValue(lease.Id, out newLease) && newLease.State != lease.State)
 						{
-							if (NewLease.State == LeaseState.Cancelled || NewLease.State == LeaseState.Completed)
+							if (newLease.State == LeaseState.Cancelled || newLease.State == LeaseState.Completed)
 							{
-								await RemoveLeaseAsync(Agent, Lease, UtcNow, NewLease.Outcome, NewLease.Output.ToByteArray());
-								Leases.RemoveAt(Idx--);
+								await RemoveLeaseAsync(agent, lease, utcNow, newLease.Outcome, newLease.Output.ToByteArray());
+								leases.RemoveAt(idx--);
 							}
-							else if (NewLease.State == LeaseState.Active && Lease.State == LeaseState.Pending)
+							else if (newLease.State == LeaseState.Active && lease.State == LeaseState.Pending)
 							{
-								Lease.State = LeaseState.Active;
+								lease.State = LeaseState.Active;
 							}
 							bUpdateLeases = true;
 						}
@@ -646,75 +644,75 @@ namespace Horde.Build.Services
 				}
 
 				// If the agent is stopping, cancel all the leases. Clear out the current session once it's complete.
-				if (Status == AgentStatus.Stopping)
+				if (status == AgentStatus.Stopping)
 				{
-					foreach (AgentLease Lease in Leases)
+					foreach (AgentLease lease in leases)
 					{
-						if (Lease.State != LeaseState.Cancelled)
+						if (lease.State != LeaseState.Cancelled)
 						{
-							Lease.State = LeaseState.Cancelled;
+							lease.State = LeaseState.Cancelled;
 							bUpdateLeases = true;
 						}
 					}
 				}
 
 				// Get the new dynamic pools for the agent
-				List<PoolId> DynamicPools = await GetDynamicPoolsAsync(Agent);
+				List<PoolId> dynamicPools = await GetDynamicPoolsAsync(agent);
 
 				// Update the agent, and try to create new lease documents if we succeed
-				IAgent? NewAgent = await Agents.TryUpdateSessionAsync(Agent, Status, SessionExpiresAt, Properties, Resources, DynamicPools, bUpdateLeases ? Leases : null);
-				if (NewAgent != null)
+				IAgent? newAgent = await Agents.TryUpdateSessionAsync(agent, status, sessionExpiresAt, properties, resources, dynamicPools, bUpdateLeases ? leases : null);
+				if (newAgent != null)
 				{
-					Agent = NewAgent;
+					agent = newAgent;
 					break;
 				}
 
 				// Fetch the agent again
-				Agent = await GetAgentAsync(Agent.Id);
+				agent = await GetAgentAsync(agent.Id);
 			}
-			return Agent;
+			return agent;
 		}
 
 		/// <summary>
 		/// Terminates an existing session. Does not update the agent itself, if it's currently 
 		/// </summary>
-		/// <param name="Agent">The agent whose current session should be terminated</param>
+		/// <param name="agent">The agent whose current session should be terminated</param>
 		/// <returns>True if the session was terminated</returns>
-		private async Task<bool> TryTerminateSessionAsync(IAgent Agent)
+		private async Task<bool> TryTerminateSessionAsync(IAgent agent)
 		{
 			// Make sure the agent has a valid session id
-			if (Agent.SessionId == null)
+			if (agent.SessionId == null)
 			{
 				return true;
 			}
 
 			// Get the time that the session finishes at
-			DateTime FinishTime = Clock.UtcNow;
-			if (Agent.SessionExpiresAt.HasValue && Agent.SessionExpiresAt.Value < FinishTime)
+			DateTime finishTime = _clock.UtcNow;
+			if (agent.SessionExpiresAt.HasValue && agent.SessionExpiresAt.Value < finishTime)
 			{
-				FinishTime = Agent.SessionExpiresAt.Value;
+				finishTime = agent.SessionExpiresAt.Value;
 			}
 
 			// Save off the session id and current leases
-			SessionId SessionId = Agent.SessionId.Value;
-			List<AgentLease> Leases = new List<AgentLease>(Agent.Leases);
+			SessionId sessionId = agent.SessionId.Value;
+			List<AgentLease> leases = new List<AgentLease>(agent.Leases);
 
 			// Clear the current session
-			IAgent? NewAgent = await Agents.TryTerminateSessionAsync(Agent);
-			if (NewAgent != null)
+			IAgent? newAgent = await Agents.TryTerminateSessionAsync(agent);
+			if (newAgent != null)
 			{
-				Agent = NewAgent;
+				agent = newAgent;
 
 				// Remove any outstanding leases
-				foreach(AgentLease Lease in Leases)
+				foreach(AgentLease lease in leases)
 				{
-					Agents.GetLogger(Agent.Id).LogInformation("Removing lease {LeaseId} during session terminate...", Lease.Id);
-					await RemoveLeaseAsync(Agent, Lease, FinishTime, LeaseOutcome.Failed, null);
+					Agents.GetLogger(agent.Id).LogInformation("Removing lease {LeaseId} during session terminate...", lease.Id);
+					await RemoveLeaseAsync(agent, lease, finishTime, LeaseOutcome.Failed, null);
 				}
 
 				// Update the session document
-				Agents.GetLogger(Agent.Id).LogInformation("Terminated session {SessionId}", SessionId);
-				await Sessions.UpdateAsync(SessionId, FinishTime, Agent.Properties, Agent.Resources);
+				Agents.GetLogger(agent.Id).LogInformation("Terminated session {SessionId}", sessionId);
+				await _sessions.UpdateAsync(sessionId, finishTime, agent.Properties, agent.Resources);
 				return true;
 			}
 			return false;
@@ -723,18 +721,18 @@ namespace Horde.Build.Services
 		/// <summary>
 		/// Creates a new lease document
 		/// </summary>
-		/// <param name="Agent">Agent that will be executing the lease</param>
-		/// <param name="AgentLease">The new agent lease</param>
+		/// <param name="agent">Agent that will be executing the lease</param>
+		/// <param name="agentLease">The new agent lease</param>
 		/// <returns>New lease document</returns>
-		private Task<ILease> CreateLeaseAsync(IAgent Agent, AgentLease AgentLease)
+		private Task<ILease> CreateLeaseAsync(IAgent agent, AgentLease agentLease)
 		{
 			try
 			{
-				return Leases.AddAsync(AgentLease.Id, AgentLease.Name, Agent.Id, Agent.SessionId!.Value, AgentLease.StreamId, AgentLease.PoolId, AgentLease.LogId, AgentLease.StartTime, AgentLease.Payload!);
+				return _leases.AddAsync(agentLease.Id, agentLease.Name, agent.Id, agent.SessionId!.Value, agentLease.StreamId, agentLease.PoolId, agentLease.LogId, agentLease.StartTime, agentLease.Payload!);
 			}
-			catch (Exception Ex)
+			catch (Exception ex)
 			{
-				Logger.LogError(Ex, "Unable to create lease {LeaseId} for agent {AgentId}; lease already exists?", AgentLease.Id, Agent.Id);
+				_logger.LogError(ex, "Unable to create lease {LeaseId} for agent {AgentId}; lease already exists?", agentLease.Id, agent.Id);
 				throw;
 			}
 		}
@@ -742,167 +740,167 @@ namespace Horde.Build.Services
 		/// <summary>
 		/// Finds all leases matching a set of criteria
 		/// </summary>
-		/// <param name="AgentId">Unqiue id of the agent executing this lease</param>
-		/// <param name="SessionId">Unique id of the agent session</param>
-		/// <param name="StartTime">Start of the search window to return results for</param>
-		/// <param name="FinishTime">End of the search window to return results for</param>
-		/// <param name="Index">Index of the first result to return</param>
-		/// <param name="Count">Number of results to return</param>
+		/// <param name="agentId">Unqiue id of the agent executing this lease</param>
+		/// <param name="sessionId">Unique id of the agent session</param>
+		/// <param name="startTime">Start of the search window to return results for</param>
+		/// <param name="finishTime">End of the search window to return results for</param>
+		/// <param name="index">Index of the first result to return</param>
+		/// <param name="count">Number of results to return</param>
 		/// <returns>List of leases matching the given criteria</returns>
-		public Task<List<ILease>> FindLeasesAsync(AgentId? AgentId, SessionId? SessionId, DateTime? StartTime, DateTime? FinishTime, int Index, int Count)
+		public Task<List<ILease>> FindLeasesAsync(AgentId? agentId, SessionId? sessionId, DateTime? startTime, DateTime? finishTime, int index, int count)
 		{
-			using IScope Scope = GlobalTracer.Instance.BuildSpan($"{nameof(AgentService)}.{nameof(FindLeasesAsync)}").StartActive();
-			Scope.Span.SetTag("AgentId", AgentId?.ToString());
-			Scope.Span.SetTag("SessionId", SessionId?.ToString());
-			Scope.Span.SetTag("StartTime", StartTime?.ToString());
-			Scope.Span.SetTag("FinishTime", FinishTime?.ToString());
-			Scope.Span.SetTag("Index", Index);
-			Scope.Span.SetTag("Count", Count);
-			return Leases.FindLeasesAsync(AgentId, SessionId, StartTime, FinishTime, Index, Count);
+			using IScope scope = GlobalTracer.Instance.BuildSpan($"{nameof(AgentService)}.{nameof(FindLeasesAsync)}").StartActive();
+			scope.Span.SetTag("AgentId", agentId?.ToString());
+			scope.Span.SetTag("SessionId", sessionId?.ToString());
+			scope.Span.SetTag("StartTime", startTime?.ToString());
+			scope.Span.SetTag("FinishTime", finishTime?.ToString());
+			scope.Span.SetTag("Index", index);
+			scope.Span.SetTag("Count", count);
+			return _leases.FindLeasesAsync(agentId, sessionId, startTime, finishTime, index, count);
 		}
 
 		/// <summary>
 		/// Finds all leases by finish time
 		/// </summary>
-		/// <param name="MinFinishTime">Start of the search window to return results for</param>
-		/// <param name="MaxFinishTime">End of the search window to return results for</param>
-		/// <param name="Index">Index of the first result to return</param>
-		/// <param name="Count">Number of results to return</param>
+		/// <param name="minFinishTime">Start of the search window to return results for</param>
+		/// <param name="maxFinishTime">End of the search window to return results for</param>
+		/// <param name="index">Index of the first result to return</param>
+		/// <param name="count">Number of results to return</param>
 		/// <returns>List of leases matching the given criteria</returns>
-		public Task<List<ILease>> FindLeasesByFinishTimeAsync(DateTime? MinFinishTime, DateTime? MaxFinishTime, int? Index, int? Count)
+		public Task<List<ILease>> FindLeasesByFinishTimeAsync(DateTime? minFinishTime, DateTime? maxFinishTime, int? index, int? count)
 		{
-			return Leases.FindLeasesByFinishTimeAsync(MinFinishTime, MaxFinishTime, Index, Count, null, false);
+			return _leases.FindLeasesByFinishTimeAsync(minFinishTime, maxFinishTime, index, count, null, false);
 		}
 
 		/// <summary>
 		/// Gets a specific lease
 		/// </summary>
-		/// <param name="LeaseId">Unique id of the lease</param>
+		/// <param name="leaseId">Unique id of the lease</param>
 		/// <returns>The lease that was found, or null if it does not exist</returns>
-		public Task<ILease?> GetLeaseAsync(LeaseId LeaseId)
+		public Task<ILease?> GetLeaseAsync(LeaseId leaseId)
 		{
-			return Leases.GetAsync(LeaseId);
+			return _leases.GetAsync(leaseId);
 		}
 
 		/// <summary>
 		/// Removes a lease with the given id. Updates the lease state in the database, and removes the item from the agent's leases array.
 		/// </summary>
-		/// <param name="Agent">The agent to remove a lease from</param>
-		/// <param name="Lease">The lease to cancel</param>
-		/// <param name="UtcNow">The current time</param>
-		/// <param name="Outcome">Final status of the lease</param>
-		/// <param name="Output">Output from executing the task</param>
+		/// <param name="agent">The agent to remove a lease from</param>
+		/// <param name="lease">The lease to cancel</param>
+		/// <param name="utcNow">The current time</param>
+		/// <param name="outcome">Final status of the lease</param>
+		/// <param name="output">Output from executing the task</param>
 		/// <returns>Async task</returns>
-		private async Task RemoveLeaseAsync(IAgent Agent, AgentLease Lease, DateTime UtcNow, LeaseOutcome Outcome, byte[]? Output)
+		private async Task RemoveLeaseAsync(IAgent agent, AgentLease lease, DateTime utcNow, LeaseOutcome outcome, byte[]? output)
 		{
 			// Make sure the lease is terminated correctly
-			if (Lease.Payload == null)
+			if (lease.Payload == null)
 			{
-				Logger.LogWarning("Removing lease {LeaseId} (no payload)", Lease.Id);
+				_logger.LogWarning("Removing lease {LeaseId} (no payload)", lease.Id);
 			}
 			else
 			{
-				Any Any = Any.Parser.ParseFrom(Lease.Payload);
-				Logger.LogInformation("Removing lease {LeaseId} ({LeaseType})", Lease.Id, Any.TypeUrl);
+				Any any = Any.Parser.ParseFrom(lease.Payload);
+				_logger.LogInformation("Removing lease {LeaseId} ({LeaseType})", lease.Id, any.TypeUrl);
 
-				foreach (ITaskSource TaskSource in TaskSources)
+				foreach (ITaskSource taskSource in _taskSources)
 				{
-					if (Any.Is(TaskSource.Descriptor))
+					if (any.Is(taskSource.Descriptor))
 					{
-						await TaskSource.OnLeaseFinishedAsync(Agent, Lease.Id, Any, Outcome, Output, Agents.GetLogger(Agent.Id));
+						await taskSource.OnLeaseFinishedAsync(agent, lease.Id, any, outcome, output, Agents.GetLogger(agent.Id));
 						break;
 					}
 				}
 			}
 
 			// Figure out what time the lease finished
-			DateTime FinishTime = UtcNow;
-			if (Agent.SessionExpiresAt.HasValue && Agent.SessionExpiresAt.Value < FinishTime)
+			DateTime finishTime = utcNow;
+			if (agent.SessionExpiresAt.HasValue && agent.SessionExpiresAt.Value < finishTime)
 			{
-				FinishTime = Agent.SessionExpiresAt.Value;
+				finishTime = agent.SessionExpiresAt.Value;
 			}
-			if (Lease.ExpiryTime.HasValue && Lease.ExpiryTime.Value < FinishTime)
+			if (lease.ExpiryTime.HasValue && lease.ExpiryTime.Value < finishTime)
 			{
-				FinishTime = Lease.ExpiryTime.Value;
+				finishTime = lease.ExpiryTime.Value;
 			}
 
 			// Update the lease
-			await Leases.TrySetOutcomeAsync(Lease.Id, FinishTime, Outcome, Output);
+			await _leases.TrySetOutcomeAsync(lease.Id, finishTime, outcome, output);
 		}
 
 		/// <summary>
 		/// Gets information about a particular session
 		/// </summary>
-		/// <param name="SessionId">The unique session id</param>
+		/// <param name="sessionId">The unique session id</param>
 		/// <returns>The session information</returns>
-		public Task<ISession?> GetSessionAsync(SessionId SessionId)
+		public Task<ISession?> GetSessionAsync(SessionId sessionId)
 		{
-			return Sessions.GetAsync(SessionId);
+			return _sessions.GetAsync(sessionId);
 		}
 
 		/// <summary>
 		/// Find sessions for the given agent
 		/// </summary>
-		/// <param name="AgentId">The unique agent id</param>
-		/// <param name="StartTime">Start time to include in the search</param>
-		/// <param name="FinishTime">Finish time to include in the search</param>
-		/// <param name="Index">Index of the first result to return</param>
-		/// <param name="Count">Number of results to return</param>
+		/// <param name="agentId">The unique agent id</param>
+		/// <param name="startTime">Start time to include in the search</param>
+		/// <param name="finishTime">Finish time to include in the search</param>
+		/// <param name="index">Index of the first result to return</param>
+		/// <param name="count">Number of results to return</param>
 		/// <returns>List of sessions matching the given criteria</returns>
-		public Task<List<ISession>> FindSessionsAsync(AgentId AgentId, DateTime? StartTime, DateTime? FinishTime, int Index, int Count)
+		public Task<List<ISession>> FindSessionsAsync(AgentId agentId, DateTime? startTime, DateTime? finishTime, int index, int count)
 		{
-			return Sessions.FindAsync(AgentId, StartTime, FinishTime, Index, Count);
+			return _sessions.FindAsync(agentId, startTime, finishTime, index, count);
 		}
 
 		/// <summary>
 		/// Update the rate table for different agent types
 		/// </summary>
-		/// <param name="Entries">New entries for the rate table</param>
+		/// <param name="entries">New entries for the rate table</param>
 		/// <returns></returns>
-		public async Task UpdateRateTableAsync(List<AgentRateConfig> Entries)
+		public async Task UpdateRateTableAsync(List<AgentRateConfig> entries)
 		{
-			await AgentRateTableData.SetAsync(new AgentRateTable { Entries = Entries });
+			await _agentRateTableData.SetAsync(new AgentRateTable { Entries = entries });
 		}
 
 		/// <summary>
 		/// Gets the rate for the given agent
 		/// </summary>
-		/// <param name="AgentId">Agent id to query</param>
+		/// <param name="agentId">Agent id to query</param>
 		/// <returns>Hourly rate of running the given agent</returns>
-		public async ValueTask<double?> GetRateAsync(AgentId AgentId)
+		public async ValueTask<double?> GetRateAsync(AgentId agentId)
 		{
-			RedisKey Key = $"agent-rate/{AgentId}";
+			RedisKey key = $"agent-rate/{agentId}";
 
 			// Try to get the current value
-			RedisValue Value = await RedisService.Database.StringGetAsync(Key);
-			if (!Value.IsNull)
+			RedisValue value = await _redisService.Database.StringGetAsync(key);
+			if (!value.IsNull)
 			{
-				double Rate = (double)Value;
-				if (Rate == 0.0)
+				double rate = (double)value;
+				if (rate == 0.0)
 				{
 					return null;
 				}
 				else
 				{
-					return Rate;
+					return rate;
 				}
 			}
 			else
 			{
-				double Rate = 0.0;
+				double rate = 0.0;
 
 				// Get the rate table
-				AgentRateTable? RateTable = await AgentRateTable.GetAsync();
-				if (RateTable != null && RateTable.Entries.Count > 0)
+				AgentRateTable? rateTable = await _agentRateTable.GetAsync();
+				if (rateTable != null && rateTable.Entries.Count > 0)
 				{
-					IAgent? Agent = await GetAgentAsync(AgentId);
-					if (Agent != null)
+					IAgent? agent = await GetAgentAsync(agentId);
+					if (agent != null)
 					{
-						foreach (AgentRateConfig Config in RateTable.Entries)
+						foreach (AgentRateConfig config in rateTable.Entries)
 						{
-							if (Config.Condition != null && Config.Condition.Evaluate(x => Agent.GetPropertyValues(x)))
+							if (config.Condition != null && config.Condition.Evaluate(x => agent.GetPropertyValues(x)))
 							{
-								Rate = Config.Rate;
+								rate = config.Rate;
 								break;
 							}
 						}
@@ -910,35 +908,35 @@ namespace Horde.Build.Services
 				}
 
 				// Cache it for future reference
-				await RedisService.Database.StringSetAsync(Key, Rate, TimeSpan.FromMinutes(5.0), flags: CommandFlags.FireAndForget);
-				return Rate;
+				await _redisService.Database.StringSetAsync(key, rate, TimeSpan.FromMinutes(5.0), flags: CommandFlags.FireAndForget);
+				return rate;
 			}
 		}
 
 		/// <summary>
 		/// Terminate any sessions for agents that are offline
 		/// </summary>
-		/// <param name="StoppingToken">Token indicating the service is shutting down</param>
+		/// <param name="stoppingToken">Token indicating the service is shutting down</param>
 		/// <returns>Async task</returns>
-		async ValueTask TickAsync(CancellationToken StoppingToken)
+		async ValueTask TickAsync(CancellationToken stoppingToken)
 		{
-			while (!StoppingToken.IsCancellationRequested)
+			while (!stoppingToken.IsCancellationRequested)
 			{
 				// Find all the agents which are ready to be expired
 				const int MaxAgents = 100;
-				DateTime UtcNow = Clock.UtcNow;
-				List<IAgent> ExpiredAgents = await Agents.FindExpiredAsync(UtcNow, MaxAgents);
+				DateTime utcNow = _clock.UtcNow;
+				List<IAgent> expiredAgents = await Agents.FindExpiredAsync(utcNow, MaxAgents);
 
 				// Transition each agent to being offline
-				foreach (IAgent ExpiredAgent in ExpiredAgents)
+				foreach (IAgent expiredAgent in expiredAgents)
 				{
-					StoppingToken.ThrowIfCancellationRequested();
-					Logger.LogDebug("Terminating session {SessionId} for agent {Agent}", ExpiredAgent.SessionId, ExpiredAgent.Id);
-					await TryTerminateSessionAsync(ExpiredAgent);
+					stoppingToken.ThrowIfCancellationRequested();
+					_logger.LogDebug("Terminating session {SessionId} for agent {Agent}", expiredAgent.SessionId, expiredAgent.Id);
+					await TryTerminateSessionAsync(expiredAgent);
 				}
 
 				// Try again if we didn't fetch everything
-				if(ExpiredAgents.Count < MaxAgents)
+				if(expiredAgents.Count < MaxAgents)
 				{
 					break;
 				}
@@ -949,55 +947,55 @@ namespace Horde.Build.Services
 
 		private async Task CollectMetrics()
 		{
-			List<IAgent> AgentList = await Agents.FindAsync();
-			int NumAgentsTotal = AgentList.Count;
-			int NumAgentsTotalEnabled = AgentList.Count(a => a.Enabled);
-			int NumAgentsTotalDisabled = AgentList.Count(a => !a.Enabled);
-			int NumAgentsTotalOk = AgentList.Count(a => a.Enabled && a.Status == AgentStatus.Ok);
-			int NumAgentsTotalStopping = AgentList.Count(a => a.Enabled && a.Status == AgentStatus.Stopping);
-			int NumAgentsTotalUnhealthy = AgentList.Count(a => a.Enabled && a.Status == AgentStatus.Unhealthy);
-			int NumAgentsTotalUnspecified = AgentList.Count(a => a.Enabled && a.Status == AgentStatus.Unspecified);
+			List<IAgent> agentList = await Agents.FindAsync();
+			int numAgentsTotal = agentList.Count;
+			int numAgentsTotalEnabled = agentList.Count(a => a.Enabled);
+			int numAgentsTotalDisabled = agentList.Count(a => !a.Enabled);
+			int numAgentsTotalOk = agentList.Count(a => a.Enabled && a.Status == AgentStatus.Ok);
+			int numAgentsTotalStopping = agentList.Count(a => a.Enabled && a.Status == AgentStatus.Stopping);
+			int numAgentsTotalUnhealthy = agentList.Count(a => a.Enabled && a.Status == AgentStatus.Unhealthy);
+			int numAgentsTotalUnspecified = agentList.Count(a => a.Enabled && a.Status == AgentStatus.Unspecified);
 			
 			// TODO: utilize tags argument in a smarter way below
-			DogStatsd.Gauge("agents.total.count", NumAgentsTotal);
-			DogStatsd.Gauge("agents.total.enabled.count", NumAgentsTotalEnabled);
-			DogStatsd.Gauge("agents.total.disabled.count", NumAgentsTotalDisabled);
-			DogStatsd.Gauge("agents.total.ok.count", NumAgentsTotalOk);
-			DogStatsd.Gauge("agents.total.stopping.count", NumAgentsTotalStopping);
-			DogStatsd.Gauge("agents.total.unhealthy.count", NumAgentsTotalUnhealthy);
-			DogStatsd.Gauge("agents.total.unspecified.count", NumAgentsTotalUnspecified);
+			_dogStatsd.Gauge("agents.total.count", numAgentsTotal);
+			_dogStatsd.Gauge("agents.total.enabled.count", numAgentsTotalEnabled);
+			_dogStatsd.Gauge("agents.total.disabled.count", numAgentsTotalDisabled);
+			_dogStatsd.Gauge("agents.total.ok.count", numAgentsTotalOk);
+			_dogStatsd.Gauge("agents.total.stopping.count", numAgentsTotalStopping);
+			_dogStatsd.Gauge("agents.total.unhealthy.count", numAgentsTotalUnhealthy);
+			_dogStatsd.Gauge("agents.total.unspecified.count", numAgentsTotalUnspecified);
 		}
 
 		/// <summary>
 		/// Determines if the user is authorized to perform an action on a particular agent
 		/// </summary>
-		/// <param name="Agent">The agent to check</param>
-		/// <param name="Action">The action being performed</param>
-		/// <param name="User">The principal to authorize</param>
-		/// <param name="Cache">The permissions cache</param>
+		/// <param name="agent">The agent to check</param>
+		/// <param name="action">The action being performed</param>
+		/// <param name="user">The principal to authorize</param>
+		/// <param name="cache">The permissions cache</param>
 		/// <returns>True if the action is authorized</returns>
-		public async Task<bool> AuthorizeAsync(IAgent Agent, AclAction Action, ClaimsPrincipal User, GlobalPermissionsCache? Cache)
+		public async Task<bool> AuthorizeAsync(IAgent agent, AclAction action, ClaimsPrincipal user, GlobalPermissionsCache? cache)
 		{
-			bool? Result = Agent.Acl?.Authorize(Action, User);
-			if (Result == null)
+			bool? result = agent.Acl?.Authorize(action, user);
+			if (result == null)
 			{
-				return await AclService.AuthorizeAsync(Action, User, Cache);
+				return await _aclService.AuthorizeAsync(action, user, cache);
 			}
 			else
 			{
-				return Result.Value;
+				return result.Value;
 			}
 		}
 
 		/// <summary>
 		/// Determines if the user is authorized to perform an action on a particular agent
 		/// </summary>
-		/// <param name="Agent">The agent to check</param>
-		/// <param name="User">The principal to authorize</param>
+		/// <param name="agent">The agent to check</param>
+		/// <param name="user">The principal to authorize</param>
 		/// <returns>True if the action is authorized</returns>
-		public bool AuthorizeSession(IAgent Agent, ClaimsPrincipal User)
+		public bool AuthorizeSession(IAgent agent, ClaimsPrincipal user)
 		{
-			if (Agent.SessionId != null && User.HasSessionClaim(Agent.SessionId.Value) && Agent.IsSessionValid(Clock.UtcNow))
+			if (agent.SessionId != null && user.HasSessionClaim(agent.SessionId.Value) && agent.IsSessionValid(_clock.UtcNow))
 			{
 				return true;
 			}
diff --git a/Engine/Source/Programs/Horde/Horde.Build/Agents/Impl/AgentSoftwareCollection.cs b/Engine/Source/Programs/Horde/Horde.Build/Agents/Impl/AgentSoftwareCollection.cs
--- a/Engine/Source/Programs/Horde/Horde.Build/Agents/Impl/AgentSoftwareCollection.cs
+++ b/Engine/Source/Programs/Horde/Horde.Build/Agents/Impl/AgentSoftwareCollection.cs
@@ -1,26 +1,17 @@
 // Copyright Epic Games, Inc. All Rights Reserved.
 
-using EpicGames.Core;
-using HordeCommon;
-using Horde.Build.Models;
+using System.Collections.Generic;
+using System.IO;
+using System.Linq;
+using System.Threading.Tasks;
 using Horde.Build.Services;
 using Horde.Build.Utilities;
 using MongoDB.Bson;
 using MongoDB.Bson.Serialization.Attributes;
 using MongoDB.Driver;
-using System;
-using System.Collections.Generic;
-using System.Globalization;
-using System.IO;
-using System.IO.Compression;
-using System.Linq;
-using System.Text.Json;
-using System.Threading.Tasks;
 
 namespace Horde.Build.Collections.Impl
 {
-	using AgentSoftwareVersion = StringId<IAgentSoftwareCollection>;
-
 	/// <summary>
 	/// Collection of agent software documents
 	/// </summary>
@@ -43,16 +34,16 @@ namespace Horde.Build.Collections.Impl
             [BsonConstructor]
 			private AgentSoftwareDocument()
 			{
-				this.Id = null!;
-				this.Data = null;
-                this.ChunkIds = null;			
+				Id = null!;
+				Data = null;
+                ChunkIds = null;			
             }
 
-			public AgentSoftwareDocument(string Id, byte[]? Data, List<string>? ChunkIds = null)
+			public AgentSoftwareDocument(string id, byte[]? data, List<string>? chunkIds = null)
 			{
-				this.Id = Id;
-				this.Data = Data;
-                this.ChunkIds = ChunkIds;
+				Id = id;
+				Data = data;
+                ChunkIds = chunkIds;
             }
 		}
 
@@ -65,79 +56,75 @@ namespace Horde.Build.Collections.Impl
 
 			public string Version { get; set; }
 
-
             [BsonConstructor]
 			private ChunkDocument()
 			{
-				this.Id = null!;
-				this.Data = null!;
-                this.Version = null!;
+				Id = null!;
+				Data = null!;
+                Version = null!;
             }
 
-			public ChunkDocument(string Id, string Version, byte[] Data)
+			public ChunkDocument(string id, string version, byte[] data)
 			{
-				this.Id = Id;
-				this.Data = Data;
-                this.Version = Version;
+				Id = id;
+				Data = data;
+                Version = version;
             }
 		}
 
 		// MongoDB document size limit is 16 megs, we chunk at this with a 1k buffer for anything else in document
-		readonly int ChunkSize = 16 * 1024 * 1024 - 1024;
-
-
+		readonly int _chunkSize = 16 * 1024 * 1024 - 1024;
 
 		/// <summary>
 		/// Template documents
 		/// </summary>
-		IMongoCollection<AgentSoftwareDocument> Collection;
+		readonly IMongoCollection<AgentSoftwareDocument> _collection;
 
 		/// <summary>
 		/// Agent template chunks
 		/// </summary>
-		IMongoCollection<ChunkDocument> AgentChunks;
-
+		readonly IMongoCollection<ChunkDocument> _agentChunks;
 
 		/// <summary>
 		/// Constructor
 		/// </summary>
-		/// <param name="DatabaseService">The database service singleton</param>
-		public AgentSoftwareCollection(DatabaseService DatabaseService)
+		/// <param name="databaseService">The database service singleton</param>
+		public AgentSoftwareCollection(DatabaseService databaseService)
 		{
-			Collection = DatabaseService.GetCollection<AgentSoftwareDocument>("AgentSoftware");
-			AgentChunks = DatabaseService.GetCollection<ChunkDocument>("AgentSoftwareChunks");
+			_collection = databaseService.GetCollection<AgentSoftwareDocument>("AgentSoftware");
+			_agentChunks = databaseService.GetCollection<ChunkDocument>("AgentSoftwareChunks");
 		}
 
 		/// <inheritdoc/>
-		public async Task<bool> AddAsync(string Version, byte[] Data)
+		public async Task<bool> AddAsync(string version, byte[] data)
 		{
 
-			if (await ExistsAsync(Version))
+			if (await ExistsAsync(version))
 			{
                 return false;
             }
 
-            AgentSoftwareDocument Document;
+            AgentSoftwareDocument document;
 
             // Check if agent will fit in document, otherwise we need to chunk
-            if (Data.Length <= ChunkSize)
+            if (data.Length <= _chunkSize)
 			{
-				Document = new AgentSoftwareDocument(Version, Data);
-				await Collection.InsertOneAsync(Document);
+				document = new AgentSoftwareDocument(version, data);
+				await _collection.InsertOneAsync(document);
                 return true;
             }
 
-			int TotalRead = 0;            
-            byte[] Buffer = new byte[ChunkSize];            
-			List<ChunkDocument> Chunks = new List<ChunkDocument>();
-            using (MemoryStream ChunkStream = new MemoryStream(Data))
+			int totalRead = 0;            
+            byte[] buffer = new byte[_chunkSize];            
+			List<ChunkDocument> chunks = new List<ChunkDocument>();
+            using (MemoryStream chunkStream = new MemoryStream(data))
             {
                 for (; ; )
                 {
-					int BytesRead = await ChunkStream.ReadAsync(Buffer, 0, ChunkSize);
-					Chunks.Add(new ChunkDocument(ObjectId.GenerateNewId().ToString(), Version, Buffer.Take(BytesRead).ToArray()));
-					TotalRead += BytesRead;
-					if (TotalRead == Data.Length)
+					int bytesRead = await chunkStream.ReadAsync(buffer, 0, _chunkSize);
+					chunks.Add(new ChunkDocument(ObjectId.GenerateNewId().ToString(), version, buffer.Take(bytesRead).ToArray()));
+					totalRead += bytesRead;
+					if (totalRead == data.Length)
 					{
 						break;
 					}
@@ -145,62 +132,62 @@ namespace Horde.Build.Collections.Impl
             }
 
             // insert the chunks
-            await AgentChunks.InsertManyAsync(Chunks);
+            await _agentChunks.InsertManyAsync(chunks);
 
-            Document = new AgentSoftwareDocument(Version, null, Chunks.Select( Chunk => Chunk.Id).ToList());
-			await Collection.InsertOneAsync(Document);
+            document = new AgentSoftwareDocument(version, null, chunks.Select( chunk => chunk.Id).ToList());
+			await _collection.InsertOneAsync(document);
             return true;
         }
 
 		/// <inheritdoc/>
-		public async Task<bool> ExistsAsync(string Version)
+		public async Task<bool> ExistsAsync(string version)
 		{
-			return await Collection.Find(x => x.Id == Version).AnyAsync();
+			return await _collection.Find(x => x.Id == version).AnyAsync();
 		}
 
 		/// <inheritdoc/>
-		public async Task<bool> RemoveAsync(string Version)
+		public async Task<bool> RemoveAsync(string version)
 		{
 			// Delete any chunks
-			await AgentChunks.DeleteManyAsync(x => x.Version == Version);
+			await _agentChunks.DeleteManyAsync(x => x.Version == version);
 
-			DeleteResult Result = await Collection.DeleteOneAsync(x => x.Id == Version);			
-			return Result.DeletedCount > 0;
+			DeleteResult result = await _collection.DeleteOneAsync(x => x.Id == version);			
+			return result.DeletedCount > 0;
 		}
 
 		/// <inheritdoc/>
-		public async Task<byte[]?> GetAsync(string Version)
+		public async Task<byte[]?> GetAsync(string version)
 		{			
-			IFindFluent<AgentSoftwareDocument, AgentSoftwareDocument> Query = Collection.Find(x => x.Id == Version);
-			AgentSoftwareDocument? Software = await Query.FirstOrDefaultAsync();
+			IFindFluent<AgentSoftwareDocument, AgentSoftwareDocument> query = _collection.Find(x => x.Id == version);
+			AgentSoftwareDocument? software = await query.FirstOrDefaultAsync();
 
-			if (Software == null)
+			if (software == null)
 			{
                 return null;
             }
 
-			if (Software.Data != null && Software.Data.Length > 0)
+			if (software.Data != null && software.Data.Length > 0)
 			{
-                return Software.Data;
+                return software.Data;
             }
 
-			if (Software.ChunkIds != null)
+			if (software.ChunkIds != null)
 			{
-				FilterDefinitionBuilder<ChunkDocument> FilterBuilder = Builders<ChunkDocument>.Filter;
-				FilterDefinition<ChunkDocument> Filter = FilterBuilder.Empty;
-				Filter &= FilterBuilder.In(x => x.Id, Software.ChunkIds);
+				FilterDefinitionBuilder<ChunkDocument> filterBuilder = Builders<ChunkDocument>.Filter;
+				FilterDefinition<ChunkDocument> filter = filterBuilder.Empty;
+				filter &= filterBuilder.In(x => x.Id, software.ChunkIds);
 
-				IFindFluent<ChunkDocument, ChunkDocument> ChunkQuery = AgentChunks.Find(x => Software.ChunkIds.Contains(x.Id));				
-				List<ChunkDocument> Chunks = await ChunkQuery.ToListAsync();
+				IFindFluent<ChunkDocument, ChunkDocument> chunkQuery = _agentChunks.Find(x => software.ChunkIds.Contains(x.Id));				
+				List<ChunkDocument> chunks = await chunkQuery.ToListAsync();
 
-				using (MemoryStream DataStream = new MemoryStream())
+				using (MemoryStream dataStream = new MemoryStream())
 				{
-					foreach (string ChunkId in Software.ChunkIds)
+					foreach (string chunkId in software.ChunkIds)
 					{
-                        DataStream.Write(Chunks.First(x => x.Id == ChunkId).Data);
+                        dataStream.Write(chunks.First(x => x.Id == chunkId).Data);
                     }
 
-                    return DataStream.ToArray();
+                    return dataStream.ToArray();
                 }
 			}
 
diff --git a/Engine/Source/Programs/Horde/Horde.Build/Agents/Impl/AgentSoftwareService.cs b/Engine/Source/Programs/Horde/Horde.Build/Agents/Impl/AgentSoftwareService.cs
--- a/Engine/Source/Programs/Horde/Horde.Build/Agents/Impl/AgentSoftwareService.cs
+++ b/Engine/Source/Programs/Horde/Horde.Build/Agents/Impl/AgentSoftwareService.cs
@@ -1,27 +1,20 @@
 // Copyright Epic Games, Inc. All Rights Reserved.
 
-using HordeCommon;
-using Horde.Build.Collections;
-using Horde.Build.Models;
-using Horde.Build.Utilities;
-using Microsoft.IdentityModel.Tokens;
-using MongoDB.Bson;
-using MongoDB.Bson.Serialization.Attributes;
-using MongoDB.Driver;
 using System;
 using System.Collections.Generic;
-using System.Diagnostics;
 using System.IO;
 using System.Linq;
-using System.Threading.Channels;
+using System.Security.Cryptography.X509Certificates;
 using System.Threading.Tasks;
 using EpicGames.Core;
+using Horde.Build.Collections;
+using Horde.Build.Commands;
+using Horde.Build.Models;
+using Horde.Build.Utilities;
 using Microsoft.Extensions.Logging;
-using System.Runtime.InteropServices;
 using Microsoft.Extensions.Options;
-using System.IO.Compression;
-using System.Security.Cryptography.X509Certificates;
-using Horde.Build.Commands;
+using MongoDB.Bson.Serialization.Attributes;
+using MongoDB.Driver;
 
 namespace Horde.Build.Services
 {
@@ -63,10 +56,10 @@ namespace Horde.Build.Services
 		/// <summary>
 		/// Constructor
 		/// </summary>
-		/// <param name="Name">Name of the channel</param>
-		public AgentSoftwareChannel(AgentSoftwareChannelName Name)
+		/// <param name="name">Name of the channel</param>
+		public AgentSoftwareChannel(AgentSoftwareChannelName name)
 		{
-			this.Name = Name;
+			Name = name;
 		}
 	}
 
@@ -89,17 +82,17 @@ namespace Horde.Build.Services
 		/// <summary>
 		/// Finds an existing channel by the given name, or adds a new one
 		/// </summary>
-		/// <param name="Name"></param>
+		/// <param name="name"></param>
 		/// <returns></returns>
-		public AgentSoftwareChannel FindOrAddChannel(AgentSoftwareChannelName Name)
+		public AgentSoftwareChannel FindOrAddChannel(AgentSoftwareChannelName name)
 		{
-			AgentSoftwareChannel? Channel = Channels.FirstOrDefault(x => x.Name == Name);
-			if (Channel == null)
+			AgentSoftwareChannel? channel = Channels.FirstOrDefault(x => x.Name == name);
+			if (channel == null)
 			{
-				Channel = new AgentSoftwareChannel(Name);
-				Channels.Add(Channel);
+				channel = new AgentSoftwareChannel(name);
+				Channels.Add(channel);
 			}
-			return Channel;
+			return channel;
 		}
 	}
 
@@ -116,44 +109,43 @@ namespace Horde.Build.Services
 		/// <summary>
 		/// Collection of software documents
 		/// </summary>
-		IAgentSoftwareCollection Collection;
+		readonly IAgentSoftwareCollection _collection;
 
 		/// <summary>
 		/// Channels singleton
 		/// </summary>
-		ISingletonDocument<AgentSoftwareChannels> Singleton;
+		readonly ISingletonDocument<AgentSoftwareChannels> _singleton;
 
 		/// <summary>
 		/// Cached copy of the channels singleton
 		/// </summary>
-		LazyCachedValue<Task<AgentSoftwareChannels>> ChannelsDocument;
+		readonly LazyCachedValue<Task<AgentSoftwareChannels>> _channelsDocument;
 
 		/// <summary>
 		/// The server settings
 		/// </summary>
-		IOptionsMonitor<ServerSettings> Settings;
+		readonly IOptionsMonitor<ServerSettings> _settings;
 
 		/// <summary>
 		///  Logger for controller
 		/// </summary>
-		private readonly ILogger<AgentSoftwareService> Logger;
-
+		private readonly ILogger<AgentSoftwareService> _logger;
 
 		/// <summary>
 		/// Constructor
 		/// </summary>
-		/// <param name="Collection">The software collection</param>
-		/// <param name="Singleton">The channels singleton</param>		
-		/// <param name="Settings">The settings monitor</param>		
-		/// <param name="Logger">The logger instance</param>
-		public AgentSoftwareService(IAgentSoftwareCollection Collection, ISingletonDocument<AgentSoftwareChannels> Singleton, IOptionsMonitor<ServerSettings> Settings, ILogger<AgentSoftwareService> Logger)
+		/// <param name="collection">The software collection</param>
+		/// <param name="singleton">The channels singleton</param>		
+		/// <param name="settings">The settings monitor</param>		
+		/// <param name="logger">The logger instance</param>
+		public AgentSoftwareService(IAgentSoftwareCollection collection, ISingletonDocument<AgentSoftwareChannels> singleton, IOptionsMonitor<ServerSettings> settings, ILogger<AgentSoftwareService> logger)
 		{
-			this.Collection = Collection;
-			this.Singleton = Singleton;
-			this.Settings = Settings;
-			this.Logger = Logger;
+			_collection = collection;
+			_singleton = singleton;
+			_settings = settings;
+			_logger = logger;
 
-			ChannelsDocument = new LazyCachedValue<Task<AgentSoftwareChannels>>(() => Singleton.GetAsync(), TimeSpan.FromSeconds(20.0));
+			_channelsDocument = new LazyCachedValue<Task<AgentSoftwareChannels>>(() => singleton.GetAsync(), TimeSpan.FromSeconds(20.0));
 
 			Task.Run(() => RegisterDefaultAgent(5000));
 		}
@@ -163,56 +155,56 @@ namespace Horde.Build.Services
 		/// </summary>
 		public async Task<List<IAgentSoftwareChannel>> FindChannelsAsync()
 		{
-			AgentSoftwareChannels Current = await ChannelsDocument.GetLatest();
-			return Current.Channels.ConvertAll<IAgentSoftwareChannel>(x => x);
+			AgentSoftwareChannels current = await _channelsDocument.GetLatest();
+			return current.Channels.ConvertAll<IAgentSoftwareChannel>(x => x);
 		}
 
 		/// <summary>
 		/// Gets a single software archive
 		/// </summary>
-		/// <param name="Name">Unique id of the channel to find software for</param>
-		public async Task<IAgentSoftwareChannel?> GetChannelAsync(AgentSoftwareChannelName Name)
+		/// <param name="name">Unique id of the channel to find software for</param>
+		public async Task<IAgentSoftwareChannel?> GetChannelAsync(AgentSoftwareChannelName name)
 		{
-			AgentSoftwareChannels Current = await ChannelsDocument.GetLatest();
-			return Current.Channels.FirstOrDefault(x => x.Name == Name);
+			AgentSoftwareChannels current = await _channelsDocument.GetLatest();
+			return current.Channels.FirstOrDefault(x => x.Name == name);
 		}
 
 		/// <summary>
 		/// Gets a cached channel name
 		/// </summary>
-		/// <param name="Name">The cached channel name</param>
+		/// <param name="name">The cached channel name</param>
 		/// <returns></returns>
-		public async Task<IAgentSoftwareChannel?> GetCachedChannelAsync(AgentSoftwareChannelName Name)
+		public async Task<IAgentSoftwareChannel?> GetCachedChannelAsync(AgentSoftwareChannelName name)
 		{
-			AgentSoftwareChannels Document = await ChannelsDocument.GetCached();
-			return Document.Channels.FirstOrDefault(x => x.Name == Name);
+			AgentSoftwareChannels document = await _channelsDocument.GetCached();
+			return document.Channels.FirstOrDefault(x => x.Name == name);
 		}
 
 		/// <summary>
 		/// Removes a channel
 		/// </summary>
-		/// <param name="Name">The channel id</param>
+		/// <param name="name">The channel id</param>
 		/// <returns>Async task</returns>
-		public async Task DeleteChannelAsync(AgentSoftwareChannelName Name)
+		public async Task DeleteChannelAsync(AgentSoftwareChannelName name)
 		{
 			for (; ; )
 			{
-				AgentSoftwareChannels Current = await Singleton.GetAsync();
+				AgentSoftwareChannels current = await _singleton.GetAsync();
 
-				int ChannelIdx = Current.Channels.FindIndex(x => x.Name == Name);
-				if (ChannelIdx == -1)
+				int channelIdx = current.Channels.FindIndex(x => x.Name == name);
+				if (channelIdx == -1)
 				{
 					break;
 				}
 
-				string Version = Current.Channels[ChannelIdx].Version;
-				Current.Channels.RemoveAt(ChannelIdx);
+				string version = current.Channels[channelIdx].Version;
+				current.Channels.RemoveAt(channelIdx);
 
-				if (await Singleton.TryUpdateAsync(Current))
+				if (await _singleton.TryUpdateAsync(current))
 				{
-					if (!Current.Channels.Any(x => x.Version == Version))
+					if (!current.Channels.Any(x => x.Version == version))
 					{
-						await Collection.RemoveAsync(Version);
+						await _collection.RemoveAsync(version);
 					}
 					break;
 				}
@@ -222,127 +214,127 @@ namespace Horde.Build.Services
 		/// <summary>
 		/// Updates a new software revision
 		/// </summary>
-		/// <param name="Name">Name of the channel</param>
-		/// <param name="Author">Name of the user uploading this file</param>
-		/// <param name="Data">The input data stream. This should be a zip archive containing the HordeAgent executable.</param>
+		/// <param name="name">Name of the channel</param>
+		/// <param name="author">Name of the user uploading this file</param>
+		/// <param name="data">The input data stream. This should be a zip archive containing the HordeAgent executable.</param>
 		/// <returns>Unique id for the file</returns>
-		public async Task<string> SetArchiveAsync(AgentSoftwareChannelName Name, string? Author, byte[] Data)
+		public async Task<string> SetArchiveAsync(AgentSoftwareChannelName name, string? author, byte[] data)
 		{
 			// Upload the software
-			string Version = AgentUtilities.ReadVersion(Data);
-			await Collection.AddAsync(Version, Data);
+			string version = AgentUtilities.ReadVersion(data);
+			await _collection.AddAsync(version, data);
 
 			// Update the channel
 			for(; ;)
 			{
-				AgentSoftwareChannels Instance = await Singleton.GetAsync();
+				AgentSoftwareChannels instance = await _singleton.GetAsync();
 
-				AgentSoftwareChannel Channel = Instance.FindOrAddChannel(Name);
-				Channel.ModifiedBy = Author;
-				Channel.ModifiedTime = DateTime.UtcNow;
-				Channel.Version = Version;
+				AgentSoftwareChannel channel = instance.FindOrAddChannel(name);
+				channel.ModifiedBy = author;
+				channel.ModifiedTime = DateTime.UtcNow;
+				channel.Version = version;
 
-				if (await Singleton.TryUpdateAsync(Instance))
+				if (await _singleton.TryUpdateAsync(instance))
 				{
 					break;
 				}
 			}
-			return Version;
+			return version;
 		}
 
 		/// <summary>
 		/// Gets the zip file for a given channel
 		/// </summary>
-		/// <param name="Name">The channel name</param>
+		/// <param name="name">The channel name</param>
 		/// <returns>Data for the given archive</returns>
-		public async Task<byte[]?> GetArchiveAsync(AgentSoftwareChannelName Name)
+		public async Task<byte[]?> GetArchiveAsync(AgentSoftwareChannelName name)
 		{
-			AgentSoftwareChannels Instance = await Singleton.GetAsync();
+			AgentSoftwareChannels instance = await _singleton.GetAsync();
 
-			AgentSoftwareChannel? Channel = Instance.Channels.FirstOrDefault(x => x.Name == Name);
-			if (Channel == null)
+			AgentSoftwareChannel? channel = instance.Channels.FirstOrDefault(x => x.Name == name);
+			if (channel == null)
 			{
 				return null;
 			}
 			else
 			{
-				return await Collection.GetAsync(Channel.Version);
+				return await _collection.GetAsync(channel.Version);
 			}
 		}
 
 		/// <summary>
 		/// Gets the zip file with a given version number
 		/// </summary>
-		/// <param name="Version">The version</param>
+		/// <param name="version">The version</param>
 		/// <returns>Data for the given archive</returns>
-		public Task<byte[]?> GetArchiveAsync(string Version)
+		public Task<byte[]?> GetArchiveAsync(string version)
 		{
-			return Collection.GetAsync(Version);
+			return _collection.GetAsync(version);
 		}
-		async Task RegisterDefaultAgent(int DelayMs)
+		async Task RegisterDefaultAgent(int delayMs)
 		{
-			await Task.Delay(DelayMs);
+			await Task.Delay(delayMs);
 
 			// Check whether we have an installed agent zip
-			FileReference AgentZip = FileReference.Combine(Program.AppDir, "DefaultAgent/Agent.zip");
-			if (!Settings.CurrentValue.SingleInstance || !FileReference.Exists(AgentZip))
+			FileReference agentZip = FileReference.Combine(Program.AppDir, "DefaultAgent/Agent.zip");
+			if (!_settings.CurrentValue.SingleInstance || !FileReference.Exists(agentZip))
 			{
 				return;
 			}
 
-			Logger.LogInformation("Checking for default agent software update");
+			_logger.LogInformation("Checking for default agent software update");
 
-			string AgentHash = ContentHash.MD5(AgentZip).ToString();
+			string agentHash = ContentHash.MD5(agentZip).ToString();
 
-			FileReference AgentHashFile = FileReference.Combine(Program.DataDir, "Agent/DefaultAgentHash");
+			FileReference agentHashFile = FileReference.Combine(Program.DataDir, "Agent/DefaultAgentHash");
 
 			try
 			{
 
-				if (FileReference.Exists(AgentHashFile) && FileReference.ReadAllText(AgentHashFile).Trim() == AgentHash)
+				if (FileReference.Exists(agentHashFile) && FileReference.ReadAllText(agentHashFile).Trim() == agentHash)
 				{
-					Logger.LogInformation("Default agent software is up to date");
+					_logger.LogInformation("Default agent software is up to date");
 					return;
 				}
 			}
-			catch (Exception Ex)
+			catch (Exception ex)
 			{
-				Logger.LogError(Ex, "Error checking default agent software, {Message}", Ex.Message);
+				_logger.LogError(ex, "Error checking default agent software, {Message}", ex.Message);
 			}
 
-            byte[] Bytes = await File.ReadAllBytesAsync(AgentZip.ToString());
+            byte[] bytes = await File.ReadAllBytesAsync(agentZip.ToString());
 
-			using X509Certificate2? GrpcCertificate = ServerCommand.ReadGrpcCertificate(Settings.CurrentValue);
+			using X509Certificate2? grpcCertificate = ServerCommand.ReadGrpcCertificate(_settings.CurrentValue);
 
-			if (GrpcCertificate == null)
+			if (grpcCertificate == null)
 			{
 				throw new Exception("Unable to register default agent without valid grpc certicate");
 			}
 
 			// construct agent app settings
 			
-			string ProfileName = "Default";
+			string profileName = "Default";
 
-			string ServerUrl = $"https://{System.Net.Dns.GetHostName()}:{Settings.CurrentValue.HttpsPort}";			
+			string serverUrl = $"https://{System.Net.Dns.GetHostName()}:{_settings.CurrentValue.HttpsPort}";			
 
-			Dictionary<string, object> DefaultProfile = new Dictionary<string, object>() { { "Name", ProfileName }, { "Environment", "prod" }, { "Thumbprint", GrpcCertificate.Thumbprint }, { "Url", ServerUrl } };
+			Dictionary<string, object> defaultProfile = new Dictionary<string, object>() { { "Name", profileName }, { "Environment", "prod" }, { "Thumbprint", grpcCertificate.Thumbprint }, { "Url", serverUrl } };
 
-			List<object> ServerProfiles = new List<object>() { DefaultProfile };
+			List<object> serverProfiles = new List<object>() { defaultProfile };
 
-			Dictionary<string, object> AgentSettings = new Dictionary<string, object>() { { "ServerProfiles", ServerProfiles }, { "Server", ProfileName }};
+			Dictionary<string, object> agentSettings = new Dictionary<string, object>() { { "ServerProfiles", serverProfiles }, { "Server", profileName }};
 
-			Bytes = AgentUtilities.UpdateAppSettings(Bytes, AgentSettings);
+			bytes = AgentUtilities.UpdateAppSettings(bytes, agentSettings);
 
-			string Version = await SetArchiveAsync(new AgentSoftwareChannelName("default"), null, Bytes);
+			string version = await SetArchiveAsync(new AgentSoftwareChannelName("default"), null, bytes);
 
-			if (!DirectoryReference.Exists(AgentHashFile.Directory))
+			if (!DirectoryReference.Exists(agentHashFile.Directory))
 			{
-				DirectoryReference.CreateDirectory(AgentHashFile.Directory);
+				DirectoryReference.CreateDirectory(agentHashFile.Directory);
 			}
 
-			await FileReference.WriteAllTextAsync(AgentHashFile, AgentHash);
+			await FileReference.WriteAllTextAsync(agentHashFile, agentHash);
 
-			Logger.LogInformation("Updated default agent software to {AgentHash}", AgentHash);
+			_logger.LogInformation("Updated default agent software to {AgentHash}", agentHash);
 
 		}
 	}
diff --git a/Engine/Source/Programs/Horde/Horde.Build/Agents/Impl/ConformList.cs b/Engine/Source/Programs/Horde/Horde.Build/Agents/Impl/ConformList.cs
--- a/Engine/Source/Programs/Horde/Horde.Build/Agents/Impl/ConformList.cs
+++ b/Engine/Source/Programs/Horde/Horde.Build/Agents/Impl/ConformList.cs
@@ -1,12 +1,10 @@
 // Copyright Epic Games, Inc. All Rights Reserved.
 
+using System;
+using System.Collections.Generic;
 using Horde.Build.Utilities;
 using MongoDB.Bson;
 using MongoDB.Bson.Serialization.Attributes;
-using System;
-using System.Collections.Generic;
-using System.Linq;
-using System.Threading.Tasks;
 
 namespace Horde.Build.Models
 {
diff --git a/Engine/Source/Programs/Horde/Horde.Build/Agents/Impl/LeaseCollection.cs b/Engine/Source/Programs/Horde/Horde.Build/Agents/Impl/LeaseCollection.cs
--- a/Engine/Source/Programs/Horde/Horde.Build/Agents/Impl/LeaseCollection.cs
+++ b/Engine/Source/Programs/Horde/Horde.Build/Agents/Impl/LeaseCollection.cs
@@ -1,16 +1,16 @@
 // Copyright Epic Games, Inc. All Rights Reserved.
 
-using HordeCommon;
+using System;
+using System.Collections.Generic;
+using System.Linq;
+using System.Threading.Tasks;
 using Horde.Build.Models;
 using Horde.Build.Services;
 using Horde.Build.Utilities;
+using HordeCommon;
 using MongoDB.Bson;
 using MongoDB.Bson.Serialization.Attributes;
 using MongoDB.Driver;
-using System;
-using System.Collections.Generic;
-using System.Linq;
-using System.Threading.Tasks;
 
 namespace Horde.Build.Collections.Impl
 {
@@ -55,17 +55,17 @@ namespace Horde.Build.Collections.Impl
 				Payload = null!;
 			}
 
-			public LeaseDocument(LeaseId Id, string Name, AgentId AgentId, SessionId SessionId, StreamId? StreamId, PoolId? PoolId, LogId? LogId, DateTime StartTime, byte[] Payload)
+			public LeaseDocument(LeaseId id, string name, AgentId agentId, SessionId sessionId, StreamId? streamId, PoolId? poolId, LogId? logId, DateTime startTime, byte[] payload)
 			{
-				this.Id = Id;
-				this.Name = Name;
-				this.AgentId = AgentId;
-				this.SessionId = SessionId;
-				this.StreamId = StreamId;
-				this.PoolId = PoolId;
-				this.LogId = LogId;
-				this.StartTime = StartTime;
-				this.Payload = Payload;
+				Id = id;
+				Name = name;
+				AgentId = agentId;
+				SessionId = sessionId;
+				StreamId = streamId;
+				PoolId = poolId;
+				LogId = logId;
+				StartTime = startTime;
+				Payload = payload;
 			}
 		}
 		
@@ -77,139 +77,145 @@ namespace Horde.Build.Collections.Impl
 			public readonly string FinishTime;
 			public readonly string FinishTimeStartTimeCompound;
 
-			public DatabaseIndexes(IMongoCollection<LeaseDocument> Leases, bool DatabaseReadOnlyMode) : base(DatabaseReadOnlyMode)
+			public DatabaseIndexes(IMongoCollection<LeaseDocument> leases, bool databaseReadOnlyMode) : base(databaseReadOnlyMode)
 			{
-				AgentId = CreateOrGetIndex(Leases, "AgentId_1", Builders<LeaseDocument>.IndexKeys.Ascending(x => x.AgentId));
-				SessionId = CreateOrGetIndex(Leases, "SessionId_1", Builders<LeaseDocument>.IndexKeys.Ascending(x => x.SessionId));
-				StartTime = CreateOrGetIndex(Leases, "StartTime_1", Builders<LeaseDocument>.IndexKeys.Ascending(x => x.StartTime));
-				FinishTime = CreateOrGetIndex(Leases, "FinishTime_1", Builders<LeaseDocument>.IndexKeys.Ascending(x => x.FinishTime));
-				FinishTimeStartTimeCompound = CreateOrGetIndex(Leases, "FinishTime_1_StartTime_-1", Builders<LeaseDocument>.IndexKeys.Ascending(x => x.FinishTime).Descending(x => x.StartTime));
+				AgentId = CreateOrGetIndex(leases, "AgentId_1", Builders<LeaseDocument>.IndexKeys.Ascending(x => x.AgentId));
+				SessionId = CreateOrGetIndex(leases, "SessionId_1", Builders<LeaseDocument>.IndexKeys.Ascending(x => x.SessionId));
+				StartTime = CreateOrGetIndex(leases, "StartTime_1", Builders<LeaseDocument>.IndexKeys.Ascending(x => x.StartTime));
+				FinishTime = CreateOrGetIndex(leases, "FinishTime_1", Builders<LeaseDocument>.IndexKeys.Ascending(x => x.FinishTime));
+				FinishTimeStartTimeCompound = CreateOrGetIndex(leases, "FinishTime_1_StartTime_-1", Builders<LeaseDocument>.IndexKeys.Ascending(x => x.FinishTime).Descending(x => x.StartTime));
 			}
 		}
 
 		/// <summary>
 		/// Collection of lease documents
 		/// </summary>
-		readonly IMongoCollection<LeaseDocument> Leases;
-		
+		readonly IMongoCollection<LeaseDocument> _leases;
+
 		/// <summary>
 		/// Creation and lookup of indexes
 		/// </summary>
-		DatabaseIndexes Indexes;
+		readonly DatabaseIndexes _indexes;
 
 		/// <summary>
 		/// Constructor
 		/// </summary>
-		/// <param name="DatabaseService">The database service instance</param>
-		public LeaseCollection(DatabaseService DatabaseService)
+		/// <param name="databaseService">The database service instance</param>
+		public LeaseCollection(DatabaseService databaseService)
 		{
-			Leases = DatabaseService.GetCollection<LeaseDocument>("Leases");
-			Indexes = new DatabaseIndexes(Leases, DatabaseService.ReadOnlyMode);
+			_leases = databaseService.GetCollection<LeaseDocument>("Leases");
+			_indexes = new DatabaseIndexes(_leases, databaseService.ReadOnlyMode);
 		}
 
 		/// <inheritdoc/>
-		public async Task<ILease> AddAsync(LeaseId Id, string Name, AgentId AgentId, SessionId SessionId, StreamId? StreamId, PoolId? PoolId, LogId? LogId, DateTime StartTime, byte[] Payload)
+		public async Task<ILease> AddAsync(LeaseId id, string name, AgentId agentId, SessionId sessionId, StreamId? streamId, PoolId? poolId, LogId? logId, DateTime startTime, byte[] payload)
 		{
-			LeaseDocument Lease = new LeaseDocument(Id, Name, AgentId, SessionId, StreamId, PoolId, LogId, StartTime, Payload);
-			await Leases.ReplaceOneAsync(x => x.Id == Id, Lease, new ReplaceOptions { IsUpsert = true });
-			return Lease;
+			LeaseDocument lease = new LeaseDocument(id, name, agentId, sessionId, streamId, poolId, logId, startTime, payload);
+			await _leases.ReplaceOneAsync(x => x.Id == id, lease, new ReplaceOptions { IsUpsert = true });
+			return lease;
 		}
 
 		/// <inheritdoc/>
-		public async Task DeleteAsync(LeaseId LeaseId)
+		public async Task DeleteAsync(LeaseId leaseId)
 		{
-			await Leases.DeleteOneAsync(x => x.Id == LeaseId);
+			await _leases.DeleteOneAsync(x => x.Id == leaseId);
 		}
 
 		/// <inheritdoc/>
-		public async Task<ILease?> GetAsync(LeaseId LeaseId)
+		public async Task<ILease?> GetAsync(LeaseId leaseId)
 		{
-			return await Leases.Find(x => x.Id == LeaseId).FirstOrDefaultAsync();
+			return await _leases.Find(x => x.Id == leaseId).FirstOrDefaultAsync();
 		}
 
 		/// <inheritdoc/>
-		public async Task<List<ILease>> FindLeasesAsync(AgentId? AgentId, SessionId? SessionId, DateTime? MinTime, DateTime? MaxTime, int? Index, int? Count, string? IndexHint = null, bool ConsistentRead = true)
+		public async Task<List<ILease>> FindLeasesAsync(AgentId? agentId, SessionId? sessionId, DateTime? minTime, DateTime? maxTime, int? index, int? count, string? indexHint = null, bool consistentRead = true)
 		{
-			IMongoCollection<LeaseDocument> Collection = ConsistentRead ? Leases : Leases.WithReadPreference(ReadPreference.SecondaryPreferred);
+			IMongoCollection<LeaseDocument> collection = consistentRead ? _leases : _leases.WithReadPreference(ReadPreference.SecondaryPreferred);
 			
-			FilterDefinitionBuilder<LeaseDocument> FilterBuilder = Builders<LeaseDocument>.Filter;
-			FilterDefinition<LeaseDocument> Filter = FilterDefinition<LeaseDocument>.Empty;
-			if (AgentId != null)
+			FilterDefinitionBuilder<LeaseDocument> filterBuilder = Builders<LeaseDocument>.Filter;
+			FilterDefinition<LeaseDocument> filter = FilterDefinition<LeaseDocument>.Empty;
+			if (agentId != null)
 			{
-				Filter &= FilterBuilder.Eq(x => x.AgentId, AgentId.Value);
+				filter &= filterBuilder.Eq(x => x.AgentId, agentId.Value);
 			}
-			if (SessionId != null)
+			if (sessionId != null)
 			{
-				Filter &= FilterBuilder.Eq(x => x.SessionId, SessionId.Value);
+				filter &= filterBuilder.Eq(x => x.SessionId, sessionId.Value);
 			}
-			if (MinTime != null)
+			if (minTime != null)
 			{
-				Filter &= FilterBuilder.Or(FilterBuilder.Eq(x => x.FinishTime!, null), FilterBuilder.Gt(x => x.FinishTime!, MinTime.Value));
+				filter &= filterBuilder.Or(filterBuilder.Eq(x => x.FinishTime!, null), filterBuilder.Gt(x => x.FinishTime!, minTime.Value));
 			}
-			if (MaxTime != null)
+			if (maxTime != null)
 			{
-				Filter &= FilterBuilder.Lt(x => x.StartTime, MaxTime.Value);
+				filter &= filterBuilder.Lt(x => x.StartTime, maxTime.Value);
 			}
-			FindOptions? FindOptions = null;
-			if (IndexHint != null)
+			FindOptions? findOptions = null;
+			if (indexHint != null)
 			{
-				FindOptions = new FindOptions { Hint = new BsonString(IndexHint) };
+				findOptions = new FindOptions { Hint = new BsonString(indexHint) };
 			}
 
-			List<LeaseDocument> Results = await Collection.Find(Filter, FindOptions).SortByDescending(x => x.StartTime).Range(Index, Count).ToListAsync();
-			return Results.ConvertAll<ILease>(x => x);
+			List<LeaseDocument> results = await collection.Find(filter, findOptions).SortByDescending(x => x.StartTime).Range(index, count).ToListAsync();
+			return results.ConvertAll<ILease>(x => x);
 		}
 		
 		/// <inheritdoc/>
-		public async Task<List<ILease>> FindLeasesByFinishTimeAsync(DateTime? MinFinishTime, DateTime? MaxFinishTime, int? Index, int? Count, string? IndexHint, bool ConsistentRead)
+		public async Task<List<ILease>> FindLeasesByFinishTimeAsync(DateTime? minFinishTime, DateTime? maxFinishTime, int? index, int? count, string? indexHint, bool consistentRead)
 		{
-			IMongoCollection<LeaseDocument> Collection = ConsistentRead ? Leases : Leases.WithReadPreference(ReadPreference.SecondaryPreferred);
-			FilterDefinitionBuilder<LeaseDocument> FilterBuilder = Builders<LeaseDocument>.Filter;
-			FilterDefinition<LeaseDocument> Filter = FilterDefinition<LeaseDocument>.Empty;
+			IMongoCollection<LeaseDocument> collection = consistentRead ? _leases : _leases.WithReadPreference(ReadPreference.SecondaryPreferred);
+			FilterDefinitionBuilder<LeaseDocument> filterBuilder = Builders<LeaseDocument>.Filter;
+			FilterDefinition<LeaseDocument> filter = FilterDefinition<LeaseDocument>.Empty;
 
-			if (MinFinishTime == null && MaxFinishTime == null)
+			if (minFinishTime == null && maxFinishTime == null)
 			{
-				throw new ArgumentException($"Both {nameof(MinFinishTime)} and {nameof(MaxFinishTime)} cannot be null");
+				throw new ArgumentException($"Both {nameof(minFinishTime)} and {nameof(maxFinishTime)} cannot be null");
 			}
 
-			if (MinFinishTime != null) Filter &= FilterBuilder.Gt(x => x.FinishTime, MinFinishTime.Value);
-			if (MaxFinishTime != null) Filter &= FilterBuilder.Lt(x => x.FinishTime, MaxFinishTime.Value);
+			if (minFinishTime != null)
+			{
+				filter &= filterBuilder.Gt(x => x.FinishTime, minFinishTime.Value);
+			}
+			if (maxFinishTime != null)
+			{
+				filter &= filterBuilder.Lt(x => x.FinishTime, maxFinishTime.Value);
+			}
 
-			FindOptions? FindOptions = IndexHint == null ? null : new FindOptions { Hint = new BsonString(IndexHint) };
-			List<LeaseDocument> Results = await Collection.Find(Filter, FindOptions).SortByDescending(x => x.FinishTime).Range(Index, Count).ToListAsync();
+			FindOptions? findOptions = indexHint == null ? null : new FindOptions { Hint = new BsonString(indexHint) };
+			List<LeaseDocument> results = await collection.Find(filter, findOptions).SortByDescending(x => x.FinishTime).Range(index, count).ToListAsync();
 
-			return Results.ConvertAll<ILease>(x => x);
+			return results.ConvertAll<ILease>(x => x);
 		}
 
 		/// <inheritdoc/>
-		public async Task<List<ILease>> FindLeasesAsync(DateTime? MinTime, DateTime? MaxTime)
+		public async Task<List<ILease>> FindLeasesAsync(DateTime? minTime, DateTime? maxTime)
 		{
-			return await FindLeasesAsync(null, null, MinTime, MaxTime, null, null, Indexes.FinishTimeStartTimeCompound, false);
+			return await FindLeasesAsync(null, null, minTime, maxTime, null, null, _indexes.FinishTimeStartTimeCompound, false);
 		}
 
 		/// <inheritdoc/>
-		public async Task<List<ILease>> FindActiveLeasesAsync(int? Index = null, int? Count = null)
+		public async Task<List<ILease>> FindActiveLeasesAsync(int? index = null, int? count = null)
 		{
-			List<LeaseDocument> Results = await Leases.Find(x => x.FinishTime == null).Range(Index, Count).ToListAsync();
-			return Results.ConvertAll<ILease>(x => x);
+			List<LeaseDocument> results = await _leases.Find(x => x.FinishTime == null).Range(index, count).ToListAsync();
+			return results.ConvertAll<ILease>(x => x);
 		}
 
 		/// <inheritdoc/>
-		public async Task<bool> TrySetOutcomeAsync(LeaseId LeaseId, DateTime FinishTime, LeaseOutcome Outcome, byte[]? Output)
+		public async Task<bool> TrySetOutcomeAsync(LeaseId leaseId, DateTime finishTime, LeaseOutcome outcome, byte[]? output)
 		{
-			FilterDefinitionBuilder<LeaseDocument> FilterBuilder = Builders<LeaseDocument>.Filter;
-			FilterDefinition<LeaseDocument> Filter = FilterBuilder.Eq(x => x.Id, LeaseId) & FilterBuilder.Eq(x => x.FinishTime, null);
+			FilterDefinitionBuilder<LeaseDocument> filterBuilder = Builders<LeaseDocument>.Filter;
+			FilterDefinition<LeaseDocument> filter = filterBuilder.Eq(x => x.Id, leaseId) & filterBuilder.Eq(x => x.FinishTime, null);
 
-			UpdateDefinitionBuilder<LeaseDocument> UpdateBuilder = Builders<LeaseDocument>.Update;
-			UpdateDefinition<LeaseDocument> Update = UpdateBuilder.Set(x => x.FinishTime, FinishTime).Set(x => x.Outcome, Outcome);
+			UpdateDefinitionBuilder<LeaseDocument> updateBuilder = Builders<LeaseDocument>.Update;
+			UpdateDefinition<LeaseDocument> update = updateBuilder.Set(x => x.FinishTime, finishTime).Set(x => x.Outcome, outcome);
 
-			if(Output != null && Output.Length > 0)
+			if(output != null && output.Length > 0)
 			{
-				Update = Update.Set(x => x.Output, Output);
+				update = update.Set(x => x.Output, output);
 			}
 
-			UpdateResult Result = await Leases.UpdateOneAsync(Filter, Update);
-			return Result.ModifiedCount > 0;
+			UpdateResult result = await _leases.UpdateOneAsync(filter, update);
+			return result.ModifiedCount > 0;
 		}
 	}
 }
