Horde: Fix coding conventions within Horde.Build.

diff --git a/Engine/Source/Programs/Horde/Horde.Build/Agents/Impl/SessionCollection.cs b/Engine/Source/Programs/Horde/Horde.Build/Agents/Impl/SessionCollection.cs
--- a/Engine/Source/Programs/Horde/Horde.Build/Agents/Impl/SessionCollection.cs
+++ b/Engine/Source/Programs/Horde/Horde.Build/Agents/Impl/SessionCollection.cs
@@ -1,16 +1,14 @@
 // Copyright Epic Games, Inc. All Rights Reserved.
 
+using System;
+using System.Collections.Generic;
+using System.Linq;
+using System.Threading.Tasks;
 using Horde.Build.Models;
 using Horde.Build.Services;
 using Horde.Build.Utilities;
-using Microsoft.Extensions.DependencyInjection;
-using MongoDB.Bson;
 using MongoDB.Bson.Serialization.Attributes;
 using MongoDB.Driver;
-using System;
-using System.Collections.Generic;
-using System.Linq;
-using System.Threading.Tasks;
 
 namespace Horde.Build.Collections.Impl
 {
@@ -45,99 +43,99 @@ namespace Horde.Build.Collections.Impl
 			{
 			}
 
-			public SessionDocument(SessionId Id, AgentId AgentId, DateTime StartTime, IReadOnlyList<string>? Properties, IReadOnlyDictionary<string, int>? Resources, string? Version)
+			public SessionDocument(SessionId id, AgentId agentId, DateTime startTime, IReadOnlyList<string>? properties, IReadOnlyDictionary<string, int>? resources, string? version)
 			{
-				this.Id = Id;
-				this.AgentId = AgentId;
-				this.StartTime = StartTime;
-				if (Properties != null)
+				Id = id;
+				AgentId = agentId;
+				StartTime = startTime;
+				if (properties != null)
 				{
-					this.Properties = new List<string>(Properties);
+					Properties = new List<string>(properties);
 				}
-				if (Resources != null)
+				if (resources != null)
 				{
-					this.Resources = new Dictionary<string, int>(Resources);
+					Resources = new Dictionary<string, int>(resources);
 				}
-				this.Version = Version;
+				Version = version;
 			}
 		}
 
 		/// <summary>
 		/// Collection of session documents
 		/// </summary>
-		readonly IMongoCollection<SessionDocument> Sessions;
+		readonly IMongoCollection<SessionDocument> _sessions;
 
 		/// <summary>
 		/// Constructor
 		/// </summary>
-		/// <param name="DatabaseService">The database service</param>
-		public SessionCollection(DatabaseService DatabaseService)
+		/// <param name="databaseService">The database service</param>
+		public SessionCollection(DatabaseService databaseService)
 		{
-			Sessions = DatabaseService.GetCollection<SessionDocument>("Sessions");
+			_sessions = databaseService.GetCollection<SessionDocument>("Sessions");
 
-			if (!DatabaseService.ReadOnlyMode)
+			if (!databaseService.ReadOnlyMode)
 			{
-				Sessions.Indexes.CreateOne(new CreateIndexModel<SessionDocument>(Builders<SessionDocument>.IndexKeys.Ascending(x => x.AgentId)));
-				Sessions.Indexes.CreateOne(new CreateIndexModel<SessionDocument>(Builders<SessionDocument>.IndexKeys.Ascending(x => x.StartTime)));
-				Sessions.Indexes.CreateOne(new CreateIndexModel<SessionDocument>(Builders<SessionDocument>.IndexKeys.Ascending(x => x.FinishTime)));
+				_sessions.Indexes.CreateOne(new CreateIndexModel<SessionDocument>(Builders<SessionDocument>.IndexKeys.Ascending(x => x.AgentId)));
+				_sessions.Indexes.CreateOne(new CreateIndexModel<SessionDocument>(Builders<SessionDocument>.IndexKeys.Ascending(x => x.StartTime)));
+				_sessions.Indexes.CreateOne(new CreateIndexModel<SessionDocument>(Builders<SessionDocument>.IndexKeys.Ascending(x => x.FinishTime)));
 			}
 		}
 
 		/// <inheritdoc/>
-		public async Task<ISession> AddAsync(SessionId Id, AgentId AgentId, DateTime StartTime, IReadOnlyList<string>? Properties, IReadOnlyDictionary<string, int>? Resources, string? Version)
+		public async Task<ISession> AddAsync(SessionId id, AgentId agentId, DateTime startTime, IReadOnlyList<string>? properties, IReadOnlyDictionary<string, int>? resources, string? version)
 		{
-			SessionDocument NewSession = new SessionDocument(Id, AgentId, StartTime, Properties, Resources, Version);
-			await Sessions.InsertOneAsync(NewSession);
-			return NewSession;
+			SessionDocument newSession = new SessionDocument(id, agentId, startTime, properties, resources, version);
+			await _sessions.InsertOneAsync(newSession);
+			return newSession;
 		}
 
 		/// <inheritdoc/>
-		public async Task<ISession?> GetAsync(SessionId SessionId)
+		public async Task<ISession?> GetAsync(SessionId sessionId)
 		{
-			return await Sessions.Find(x => x.Id == SessionId).FirstOrDefaultAsync();
+			return await _sessions.Find(x => x.Id == sessionId).FirstOrDefaultAsync();
 		}
 
 		/// <inheritdoc/>
-		public async Task<List<ISession>> FindAsync(AgentId AgentId, DateTime? StartTime, DateTime? FinishTime, int Index, int Count)
+		public async Task<List<ISession>> FindAsync(AgentId agentId, DateTime? startTime, DateTime? finishTime, int index, int count)
 		{
-			FilterDefinitionBuilder<SessionDocument> FilterBuilder = Builders<SessionDocument>.Filter;
+			FilterDefinitionBuilder<SessionDocument> filterBuilder = Builders<SessionDocument>.Filter;
 
-			FilterDefinition<SessionDocument> Filter = FilterBuilder.Eq(x => x.AgentId, AgentId);
-			if (StartTime != null)
+			FilterDefinition<SessionDocument> filter = filterBuilder.Eq(x => x.AgentId, agentId);
+			if (startTime != null)
 			{
-				Filter &= FilterBuilder.Gte(x => x.StartTime, StartTime.Value);
+				filter &= filterBuilder.Gte(x => x.StartTime, startTime.Value);
 			}
-			if (FinishTime != null)
+			if (finishTime != null)
 			{
-				Filter &= FilterBuilder.Or(FilterBuilder.Eq(x => x.FinishTime, null), FilterBuilder.Lte(x => x.FinishTime, FinishTime.Value));
+				filter &= filterBuilder.Or(filterBuilder.Eq(x => x.FinishTime, null), filterBuilder.Lte(x => x.FinishTime, finishTime.Value));
 			}
 
-			List<SessionDocument> Results = await Sessions.Find(Filter).SortByDescending(x => x.StartTime).Skip(Index).Limit(Count).ToListAsync();
-			return Results.ConvertAll<ISession>(x => x);
+			List<SessionDocument> results = await _sessions.Find(filter).SortByDescending(x => x.StartTime).Skip(index).Limit(count).ToListAsync();
+			return results.ConvertAll<ISession>(x => x);
 		}
 
 		/// <inheritdoc/>
-		public async Task<List<ISession>> FindActiveSessionsAsync(int? Index, int? Count)
+		public async Task<List<ISession>> FindActiveSessionsAsync(int? index, int? count)
 		{
-			List<SessionDocument> Results = await Sessions.Find(x => x.FinishTime == null).Range(Index, Count).ToListAsync();
-			return Results.ConvertAll<ISession>(x => x);
+			List<SessionDocument> results = await _sessions.Find(x => x.FinishTime == null).Range(index, count).ToListAsync();
+			return results.ConvertAll<ISession>(x => x);
 		}
 
 		/// <inheritdoc/>
-		public Task UpdateAsync(SessionId SessionId, DateTime FinishTime, IReadOnlyList<string> Properties, IReadOnlyDictionary<string, int> Resources)
+		public Task UpdateAsync(SessionId sessionId, DateTime finishTime, IReadOnlyList<string> properties, IReadOnlyDictionary<string, int> resources)
 		{
-			UpdateDefinition<SessionDocument> Update = Builders<SessionDocument>.Update
-				.Set(x => x.FinishTime, FinishTime)
-				.Set(x => x.Properties, new List<string>(Properties))
-				.Set(x => x.Resources, new Dictionary<string, int>(Resources));
+			UpdateDefinition<SessionDocument> update = Builders<SessionDocument>.Update
+				.Set(x => x.FinishTime, finishTime)
+				.Set(x => x.Properties, new List<string>(properties))
+				.Set(x => x.Resources, new Dictionary<string, int>(resources));
 
-			return Sessions.FindOneAndUpdateAsync(x => x.Id == SessionId, Update);
+			return _sessions.FindOneAndUpdateAsync(x => x.Id == sessionId, update);
 		}
 
 		/// <inheritdoc/>
-		public Task DeleteAsync(SessionId SessionId)
+		public Task DeleteAsync(SessionId sessionId)
 		{
-			return Sessions.DeleteOneAsync(x => x.Id == SessionId);
+			return _sessions.DeleteOneAsync(x => x.Id == sessionId);
 		}
 	}
 }
diff --git a/Engine/Source/Programs/Horde/Horde.Build/Agents/Impl/Tasks/ConformTaskSource.cs b/Engine/Source/Programs/Horde/Horde.Build/Agents/Impl/Tasks/ConformTaskSource.cs
--- a/Engine/Source/Programs/Horde/Horde.Build/Agents/Impl/Tasks/ConformTaskSource.cs
+++ b/Engine/Source/Programs/Horde/Horde.Build/Agents/Impl/Tasks/ConformTaskSource.cs
@@ -1,23 +1,21 @@
 // Copyright Epic Games, Inc. All Rights Reserved.
 
+using System;
+using System.Collections.Generic;
+using System.Linq;
+using System.Threading;
+using System.Threading.Tasks;
 using Google.Protobuf;
-using Google.Protobuf.Reflection;
 using Google.Protobuf.WellKnownTypes;
-using HordeCommon;
-using HordeCommon.Rpc.Tasks;
 using Horde.Build.Api;
 using Horde.Build.Collections;
 using Horde.Build.Models;
 using Horde.Build.Services;
 using Horde.Build.Utilities;
+using HordeCommon;
+using HordeCommon.Rpc.Tasks;
 using Microsoft.Extensions.Hosting;
 using Microsoft.Extensions.Logging;
-using MongoDB.Bson;
-using System;
-using System.Collections.Generic;
-using System.Linq;
-using System.Threading;
-using System.Threading.Tasks;
 
 namespace Horde.Build.Tasks.Impl
 {
@@ -36,149 +34,149 @@ namespace Horde.Build.Tasks.Impl
 		/// <inheritdoc/>
 		public override TaskSourceFlags Flags => TaskSourceFlags.AllowWhenDisabled;
 
-		DatabaseService DatabaseService;
-		IAgentCollection AgentCollection;
-		PoolService PoolService;
-		SingletonDocument<ConformList> ConformList;
-		PerforceLoadBalancer PerforceLoadBalancer;
-		ILogFileService LogService;
-		ILogger Logger;
-		ITicker TickConformList;
+		readonly DatabaseService _databaseService;
+		readonly IAgentCollection _agentCollection;
+		readonly PoolService _poolService;
+		readonly SingletonDocument<ConformList> _conformList;
+		readonly PerforceLoadBalancer _perforceLoadBalancer;
+		readonly ILogFileService _logService;
+		readonly ILogger _logger;
+		readonly ITicker _tickConformList;
 
 		/// <summary>
 		/// Constructor
 		/// </summary>
-		public ConformTaskSource(DatabaseService DatabaseService, IAgentCollection AgentCollection, PoolService PoolService, ILogFileService LogService, PerforceLoadBalancer PerforceLoadBalancer, IClock Clock, ILogger<ConformTaskSource> Logger)
+		public ConformTaskSource(DatabaseService databaseService, IAgentCollection agentCollection, PoolService poolService, ILogFileService logService, PerforceLoadBalancer perforceLoadBalancer, IClock clock, ILogger<ConformTaskSource> logger)
 		{
-			this.DatabaseService = DatabaseService;
-			this.AgentCollection = AgentCollection;
-			this.PoolService = PoolService;
-			this.ConformList = new SingletonDocument<ConformList>(DatabaseService);
-			this.PerforceLoadBalancer = PerforceLoadBalancer;
-			this.LogService = LogService;
-			this.Logger = Logger;
-			this.TickConformList = Clock.AddSharedTicker<ConformTaskSource>(TimeSpan.FromMinutes(1.0), CleanConformListAsync, Logger);
-
-			this.OnLeaseStartedProperties.Add(nameof(ConformTask.LogId), x => new LogId(x.LogId));
+			_databaseService = databaseService;
+			_agentCollection = agentCollection;
+			_poolService = poolService;
+			_conformList = new SingletonDocument<ConformList>(databaseService);
+			_perforceLoadBalancer = perforceLoadBalancer;
+			_logService = logService;
+			_logger = logger;
+			_tickConformList = clock.AddSharedTicker<ConformTaskSource>(TimeSpan.FromMinutes(1.0), CleanConformListAsync, logger);
+
+			OnLeaseStartedProperties.Add(nameof(ConformTask.LogId), x => new LogId(x.LogId));
 		}
 
 		/// <inheritdoc/>
-		public Task StartAsync(CancellationToken CancellationToken) => TickConformList.StartAsync();
+		public Task StartAsync(CancellationToken cancellationToken) => _tickConformList.StartAsync();
 
 		/// <inheritdoc/>
-		public Task StopAsync(CancellationToken CancellationToken) => TickConformList.StopAsync();
+		public Task StopAsync(CancellationToken cancellationToken) => _tickConformList.StopAsync();
 
 		/// <inheritdoc/>
-		public void Dispose() => TickConformList.Dispose();
+		public void Dispose() => _tickConformList.Dispose();
 
 		/// <summary>
 		/// Clean up the conform list of any outdated entries
 		/// </summary>
-		/// <param name="CancellationToken">Cancellation token</param>
+		/// <param name="cancellationToken">Cancellation token</param>
 		/// <returns>Async task</returns>
-		async ValueTask CleanConformListAsync(CancellationToken CancellationToken)
+		async ValueTask CleanConformListAsync(CancellationToken cancellationToken)
 		{
-			DateTime UtcNow = DateTime.UtcNow;
-			DateTime LastCheckTimeUtc = UtcNow - TimeSpan.FromMinutes(30.0);
+			DateTime utcNow = DateTime.UtcNow;
+			DateTime lastCheckTimeUtc = utcNow - TimeSpan.FromMinutes(30.0);
 
 			// Get the current state of the conform list
-			ConformList List = await ConformList.GetAsync();
+			ConformList list = await _conformList.GetAsync();
 
 			// Update any leases that are older than LastCheckTimeUtc
-			Dictionary<LeaseId, bool> RemoveLeases = new Dictionary<LeaseId, bool>();
-			foreach (ConformListEntry Entry in Enumerable.Concat(List.Entries, List.Servers.SelectMany(x => x.Entries)))
+			Dictionary<LeaseId, bool> removeLeases = new Dictionary<LeaseId, bool>();
+			foreach (ConformListEntry entry in Enumerable.Concat(list.Entries, list.Servers.SelectMany(x => x.Entries)))
 			{
-				if (Entry.LastCheckTimeUtc < LastCheckTimeUtc)
+				if (entry.LastCheckTimeUtc < lastCheckTimeUtc)
 				{
-					IAgent? Agent = await AgentCollection.GetAsync(Entry.AgentId);
+					IAgent? agent = await _agentCollection.GetAsync(entry.AgentId);
 
-					bool Remove = false;
-					if (Agent == null || !Agent.Leases.Any(x => x.Id == Entry.LeaseId))
+					bool remove = false;
+					if (agent == null || !agent.Leases.Any(x => x.Id == entry.LeaseId))
 					{
-						Logger.LogWarning("Removing invalid lease from conform list: {AgentId} lease {LeaseId}", Entry.AgentId, Entry.LeaseId);
-						Remove = true;
+						_logger.LogWarning("Removing invalid lease from conform list: {AgentId} lease {LeaseId}", entry.AgentId, entry.LeaseId);
+						remove = true;
 					}
 
-					RemoveLeases[Entry.LeaseId] = Remove;
+					removeLeases[entry.LeaseId] = remove;
 				}
 			}
 
 			// If there's anything to change, update the list
-			if (RemoveLeases.Count > 0)
+			if (removeLeases.Count > 0)
 			{
-				await ConformList.UpdateAsync(List => UpdateConformList(List, UtcNow, RemoveLeases));
+				await _conformList.UpdateAsync(list => UpdateConformList(list, utcNow, removeLeases));
 			}
 		}
 
 		/// <summary>
 		/// Remove items from the conform list, and update timestamps for them
 		/// </summary>
-		/// <param name="List">The list to update</param>
-		/// <param name="UtcNow">Current time</param>
-		/// <param name="RemoveLeases">List of leases to update. Entries with values set to true will be removed, entries with values set to false will have their timestamp updated.</param>
-		static void UpdateConformList(ConformList List, DateTime UtcNow, Dictionary<LeaseId, bool> RemoveLeases)
+		/// <param name="list">The list to update</param>
+		/// <param name="utcNow">Current time</param>
+		/// <param name="removeLeases">List of leases to update. Entries with values set to true will be removed, entries with values set to false will have their timestamp updated.</param>
+		static void UpdateConformList(ConformList list, DateTime utcNow, Dictionary<LeaseId, bool> removeLeases)
 		{
-			UpdateConformList(List.Entries, UtcNow, RemoveLeases);
-			foreach (ConformListServer Server in List.Servers)
+			UpdateConformList(list.Entries, utcNow, removeLeases);
+			foreach (ConformListServer server in list.Servers)
 			{
-				UpdateConformList(Server.Entries, UtcNow, RemoveLeases);
+				UpdateConformList(server.Entries, utcNow, removeLeases);
 			}
 		}
 
 		/// <summary>
 		/// Remove items from the conform list, and update timestamps for them
 		/// </summary>
-		/// <param name="Entries">The list to update</param>
-		/// <param name="UtcNow">Current time</param>
-		/// <param name="RemoveLeases">List of leases to update. Entries with values set to true will be removed, entries with values set to false will have their timestamp updated.</param>
-		static void UpdateConformList(List<ConformListEntry> Entries, DateTime UtcNow, Dictionary<LeaseId, bool> RemoveLeases)
+		/// <param name="entries">The list to update</param>
+		/// <param name="utcNow">Current time</param>
+		/// <param name="removeLeases">List of leases to update. Entries with values set to true will be removed, entries with values set to false will have their timestamp updated.</param>
+		static void UpdateConformList(List<ConformListEntry> entries, DateTime utcNow, Dictionary<LeaseId, bool> removeLeases)
 		{
-			for (int Idx = 0; Idx < Entries.Count; Idx++)
+			for (int idx = 0; idx < entries.Count; idx++)
 			{
-				bool Remove;
-				if (RemoveLeases.TryGetValue(Entries[Idx].LeaseId, out Remove))
+				bool remove;
+				if (removeLeases.TryGetValue(entries[idx].LeaseId, out remove))
 				{
-					if (Remove)
+					if (remove)
 					{
-						Entries.RemoveAt(Idx--);
+						entries.RemoveAt(idx--);
 					}
-					else if (Entries[Idx].LastCheckTimeUtc < UtcNow)
+					else if (entries[idx].LastCheckTimeUtc < utcNow)
 					{
-						Entries[Idx].LastCheckTimeUtc = UtcNow;
+						entries[idx].LastCheckTimeUtc = utcNow;
 					}
 				}
 			}
 		}
 
 		/// <inheritdoc/>
-		public override async Task<AgentLease?> AssignLeaseAsync(IAgent Agent, CancellationToken CancellationToken)
+		public override async Task<AgentLease?> AssignLeaseAsync(IAgent agent, CancellationToken cancellationToken)
 		{
-			DateTime UtcNow = DateTime.UtcNow;
-			if (!await IsConformPendingAsync(Agent, UtcNow))
+			DateTime utcNow = DateTime.UtcNow;
+			if (!await IsConformPendingAsync(agent, utcNow))
 			{
 				return null;
 			}
 
-			if (Agent.Leases.Count == 0)
+			if (agent.Leases.Count == 0)
 			{
-				ConformTask Task = new ConformTask();
-				if (await GetWorkspacesAsync(Agent, Task.Workspaces))
+				ConformTask task = new ConformTask();
+				if (await GetWorkspacesAsync(agent, task.Workspaces))
 				{
-					LeaseId LeaseId = LeaseId.GenerateNewId();
-					if (await AllocateConformLeaseAsync(Agent.Id, Task.Workspaces, LeaseId))
+					LeaseId leaseId = LeaseId.GenerateNewId();
+					if (await AllocateConformLeaseAsync(agent.Id, task.Workspaces, leaseId))
 					{
-						ILogFile Log = await LogService.CreateLogFileAsync(JobId.Empty, Agent.SessionId, LogType.Json);
-						Task.LogId = Log.Id.ToString();
-						Task.RemoveUntrackedFiles = Agent.RequestFullConform;
+						ILogFile log = await _logService.CreateLogFileAsync(JobId.Empty, agent.SessionId, LogType.Json);
+						task.LogId = log.Id.ToString();
+						task.RemoveUntrackedFiles = agent.RequestFullConform;
 
-						byte[] Payload = Any.Pack(Task).ToByteArray();
+						byte[] payload = Any.Pack(task).ToByteArray();
 
-						return new AgentLease(LeaseId, "Updating workspaces", null, null, Log.Id, LeaseState.Pending, null, true, Payload);
+						return new AgentLease(leaseId, "Updating workspaces", null, null, log.Id, LeaseState.Pending, null, true, payload);
 					}
 				}
 			}
 
-			if (Agent.RequestConform || Agent.RequestFullConform)
+			if (agent.RequestConform || agent.RequestFullConform)
 			{
 				return AgentLease.Drain;
 			}
@@ -187,21 +185,21 @@ namespace Horde.Build.Tasks.Impl
 		}
 
 		/// <inheritdoc/>
-		public override Task CancelLeaseAsync(IAgent Agent, LeaseId LeaseId, ConformTask Payload)
+		public override Task CancelLeaseAsync(IAgent agent, LeaseId leaseId, ConformTask payload)
 		{
-			return ReleaseConformLeaseAsync(LeaseId);
+			return ReleaseConformLeaseAsync(leaseId);
 		}
 
 		/// <inheritdoc/>
-		public async Task<bool> GetWorkspacesAsync(IAgent Agent, IList<HordeCommon.Rpc.Messages.AgentWorkspace> Workspaces)
+		public async Task<bool> GetWorkspacesAsync(IAgent agent, IList<HordeCommon.Rpc.Messages.AgentWorkspace> workspaces)
 		{
-			Globals Globals = await DatabaseService.GetGlobalsAsync();
+			Globals globals = await _databaseService.GetGlobalsAsync();
 
-			HashSet<AgentWorkspace> ConformWorkspaces = await PoolService.GetWorkspacesAsync(Agent, DateTime.UtcNow);
-			foreach (AgentWorkspace ConformWorkspace in ConformWorkspaces)
+			HashSet<AgentWorkspace> conformWorkspaces = await _poolService.GetWorkspacesAsync(agent, DateTime.UtcNow);
+			foreach (AgentWorkspace conformWorkspace in conformWorkspaces)
 			{
-				PerforceCluster? Cluster = Globals.FindPerforceCluster(ConformWorkspace.Cluster);
-				if (Cluster == null || !await Agent.TryAddWorkspaceMessage(ConformWorkspace, Cluster, PerforceLoadBalancer, Workspaces))
+				PerforceCluster? cluster = globals.FindPerforceCluster(conformWorkspace.Cluster);
+				if (cluster == null || !await agent.TryAddWorkspaceMessage(conformWorkspace, cluster, _perforceLoadBalancer, workspaces))
 				{
 					return false;
 				}
@@ -211,71 +209,71 @@ namespace Horde.Build.Tasks.Impl
 		}
 
 		/// <inheritdoc/>
-		public override Task OnLeaseFinishedAsync(IAgent Agent, LeaseId LeaseId, ConformTask Payload, LeaseOutcome Outcome, ReadOnlyMemory<byte> Output, ILogger Logger)
+		public override Task OnLeaseFinishedAsync(IAgent agent, LeaseId leaseId, ConformTask payload, LeaseOutcome outcome, ReadOnlyMemory<byte> output, ILogger logger)
 		{
-			return ReleaseConformLeaseAsync(LeaseId);
+			return ReleaseConformLeaseAsync(leaseId);
 		}
 
 		/// <summary>
 		/// Atempt to allocate a conform resource for the given lease
 		/// </summary>
-		/// <param name="AgentId">The agent id</param>
-		/// <param name="Workspaces">List of workspaces that are required</param>
-		/// <param name="LeaseId">The lease id</param>
+		/// <param name="agentId">The agent id</param>
+		/// <param name="workspaces">List of workspaces that are required</param>
+		/// <param name="leaseId">The lease id</param>
 		/// <returns>True if the resource was allocated, false otherwise</returns>
-		private async Task<bool> AllocateConformLeaseAsync(AgentId AgentId, IEnumerable<HordeCommon.Rpc.Messages.AgentWorkspace> Workspaces, LeaseId LeaseId)
+		private async Task<bool> AllocateConformLeaseAsync(AgentId agentId, IEnumerable<HordeCommon.Rpc.Messages.AgentWorkspace> workspaces, LeaseId leaseId)
 		{
-			Globals Globals = await DatabaseService.GetGlobalsAsync();
+			Globals globals = await _databaseService.GetGlobalsAsync();
 			for (; ; )
 			{
-				ConformList CurrentValue = await ConformList.GetAsync();
-				if (Globals.MaxConformCount != 0 && CurrentValue.Entries.Count + CurrentValue.Servers.Sum(x => x.Entries.Count) >= Globals.MaxConformCount)
+				ConformList currentValue = await _conformList.GetAsync();
+				if (globals.MaxConformCount != 0 && currentValue.Entries.Count + currentValue.Servers.Sum(x => x.Entries.Count) >= globals.MaxConformCount)
 				{
 					return false;
 				}
 
-				HashSet<string> Servers = new HashSet<string>(StringComparer.OrdinalIgnoreCase);
-				foreach (HordeCommon.Rpc.Messages.AgentWorkspace Workspace in Workspaces)
+				HashSet<string> servers = new HashSet<string>(StringComparer.OrdinalIgnoreCase);
+				foreach (HordeCommon.Rpc.Messages.AgentWorkspace workspace in workspaces)
 				{
-					if (Servers.Add(Workspace.ServerAndPort))
+					if (servers.Add(workspace.ServerAndPort))
 					{
-						PerforceCluster? Cluster = Globals.FindPerforceCluster(Workspace.Cluster);
-						if (Cluster == null)
+						PerforceCluster? cluster = globals.FindPerforceCluster(workspace.Cluster);
+						if (cluster == null)
 						{
-							Logger.LogWarning("Unable to find perforce cluster '{Cluster}' for conform", Workspace.Cluster);
+							_logger.LogWarning("Unable to find perforce cluster '{Cluster}' for conform", workspace.Cluster);
 							return false;
 						}
 
-						PerforceServer? Server = Cluster.Servers.FirstOrDefault(x => x.ServerAndPort.Equals(Workspace.BaseServerAndPort, StringComparison.Ordinal));
-						if (Server == null)
+						PerforceServer? server = cluster.Servers.FirstOrDefault(x => x.ServerAndPort.Equals(workspace.BaseServerAndPort, StringComparison.Ordinal));
+						if (server == null)
 						{
-							Logger.LogWarning("Unable to find perforce server '{Server}' in cluster '{Cluster}' for conform", Workspace.BaseServerAndPort, Workspace.Cluster);
+							_logger.LogWarning("Unable to find perforce server '{Server}' in cluster '{Cluster}' for conform", workspace.BaseServerAndPort, workspace.Cluster);
 							return false;
 						}
 
-						ConformListServer? ConformServer = CurrentValue.Servers.FirstOrDefault(x => x.Cluster.Equals(Workspace.Cluster, StringComparison.Ordinal) && x.ServerAndPort.Equals(Workspace.ServerAndPort, StringComparison.Ordinal));
-						if (ConformServer == null)
+						ConformListServer? conformServer = currentValue.Servers.FirstOrDefault(x => x.Cluster.Equals(workspace.Cluster, StringComparison.Ordinal) && x.ServerAndPort.Equals(workspace.ServerAndPort, StringComparison.Ordinal));
+						if (conformServer == null)
 						{
-							ConformServer = new ConformListServer { Cluster = Workspace.Cluster, ServerAndPort = Workspace.ServerAndPort };
-							CurrentValue.Servers.Add(ConformServer);
+							conformServer = new ConformListServer { Cluster = workspace.Cluster, ServerAndPort = workspace.ServerAndPort };
+							currentValue.Servers.Add(conformServer);
 						}
 
-						if (Server.MaxConformCount != 0 && ConformServer.Entries.Count > Server.MaxConformCount)
+						if (server.MaxConformCount != 0 && conformServer.Entries.Count > server.MaxConformCount)
 						{
 							return false;
 						}
 
-						ConformListEntry Entry = new ConformListEntry();
-						Entry.AgentId = AgentId;
-						Entry.LeaseId = LeaseId;
-						Entry.LastCheckTimeUtc = DateTime.UtcNow;
-						ConformServer.Entries.Add(Entry);
+						ConformListEntry entry = new ConformListEntry();
+						entry.AgentId = agentId;
+						entry.LeaseId = leaseId;
+						entry.LastCheckTimeUtc = DateTime.UtcNow;
+						conformServer.Entries.Add(entry);
 					}
 				}
 
-				if (await ConformList.TryUpdateAsync(CurrentValue))
+				if (await _conformList.TryUpdateAsync(currentValue))
 				{
-					Logger.LogInformation("Added conform lease {LeaseId}", LeaseId);
+					_logger.LogInformation("Added conform lease {LeaseId}", leaseId);
 					return true;
 				}
 			}
@@ -284,21 +282,21 @@ namespace Horde.Build.Tasks.Impl
 		/// <summary>
 		/// Terminate a conform lease
 		/// </summary>
-		/// <param name="LeaseId">The lease id</param>
+		/// <param name="leaseId">The lease id</param>
 		/// <returns>Async task</returns>
-		public async Task ReleaseConformLeaseAsync(LeaseId LeaseId)
+		public async Task ReleaseConformLeaseAsync(LeaseId leaseId)
 		{
 			for (; ; )
 			{
-				ConformList CurrentValue = await ConformList.GetAsync();
-				if (CurrentValue.Entries.RemoveAll(x => x.LeaseId == LeaseId) + CurrentValue.Servers.Sum(x => x.Entries.RemoveAll(x => x.LeaseId == LeaseId)) == 0)
+				ConformList currentValue = await _conformList.GetAsync();
+				if (currentValue.Entries.RemoveAll(x => x.LeaseId == leaseId) + currentValue.Servers.Sum(x => x.Entries.RemoveAll(x => x.LeaseId == leaseId)) == 0)
 				{
-					Logger.LogInformation("Conform lease {LeaseId} is not in singelton", LeaseId);
+					_logger.LogInformation("Conform lease {LeaseId} is not in singelton", leaseId);
 					break;
 				}
-				if (await ConformList.TryUpdateAsync(CurrentValue))
+				if (await _conformList.TryUpdateAsync(currentValue))
 				{
-					Logger.LogInformation("Removed conform lease {LeaseId}", LeaseId);
+					_logger.LogInformation("Removed conform lease {LeaseId}", leaseId);
 					break;
 				}
 			}
@@ -307,38 +305,38 @@ namespace Horde.Build.Tasks.Impl
 		/// <summary>
 		/// Determine if an agent should be conformed
 		/// </summary>
-		/// <param name="Agent">The agent to test</param>
-		/// <param name="UtcNow">Current time</param>
+		/// <param name="agent">The agent to test</param>
+		/// <param name="utcNow">Current time</param>
 		/// <returns>True if the agent should be conformed</returns>
-		private async Task<bool> IsConformPendingAsync(IAgent Agent, DateTime UtcNow)
+		private async Task<bool> IsConformPendingAsync(IAgent agent, DateTime utcNow)
 		{
 			// If a conform was manually requested, allow it to run even if the agent is disabled
-			if (Agent.RequestConform || Agent.RequestFullConform)
+			if (agent.RequestConform || agent.RequestFullConform)
 			{
-				return !IsConformCoolDownPeriod(Agent, UtcNow);
+				return !IsConformCoolDownPeriod(agent, utcNow);
 			}
 
 			// Otherwise only run if the agent is enabled
-			if (Agent.Enabled)
+			if (agent.Enabled)
 			{
 				// If we've attempted (and failed) a conform, run again after a certain time
-				if (Agent.ConformAttemptCount.HasValue)
+				if (agent.ConformAttemptCount.HasValue)
 				{
-					return !IsConformCoolDownPeriod(Agent, UtcNow);
+					return !IsConformCoolDownPeriod(agent, utcNow);
 				}
 
 				// Always run a conform every 24h
-				if (UtcNow > Agent.LastConformTime + TimeSpan.FromDays(1.0))
+				if (utcNow > agent.LastConformTime + TimeSpan.FromDays(1.0))
 				{
 					return true;
 				}
 
 				// Check if the workspaces have changed (first check against a cached list of workspaces, then an accurate one)
-				HashSet<AgentWorkspace> Workspaces = await PoolService.GetWorkspacesAsync(Agent, UtcNow - TimeSpan.FromSeconds(30.0));
-				if (!Workspaces.SetEquals(Agent.Workspaces))
+				HashSet<AgentWorkspace> workspaces = await _poolService.GetWorkspacesAsync(agent, utcNow - TimeSpan.FromSeconds(30.0));
+				if (!workspaces.SetEquals(agent.Workspaces))
 				{
-					Workspaces = await PoolService.GetWorkspacesAsync(Agent, UtcNow);
-					if (!Workspaces.SetEquals(Agent.Workspaces))
+					workspaces = await _poolService.GetWorkspacesAsync(agent, utcNow);
+					if (!workspaces.SetEquals(agent.Workspaces))
 					{
 						return true;
 					}
@@ -350,33 +348,33 @@ namespace Horde.Build.Tasks.Impl
 		/// <summary>
 		/// Determines if the given agent is in a cooldown period, waiting to retry after a failed conform
 		/// </summary>
-		/// <param name="Agent">The agent to test</param>
-		/// <param name="UtcNow">Current time</param>
+		/// <param name="agent">The agent to test</param>
+		/// <param name="utcNow">Current time</param>
 		/// <returns></returns>
-		private static bool IsConformCoolDownPeriod(IAgent Agent, DateTime UtcNow)
+		private static bool IsConformCoolDownPeriod(IAgent agent, DateTime utcNow)
 		{
-			if (!Agent.ConformAttemptCount.HasValue)
+			if (!agent.ConformAttemptCount.HasValue)
 			{
 				return false;
 			}
 
-			TimeSpan RetryTime;
-			switch (Agent.ConformAttemptCount.Value)
+			TimeSpan retryTime;
+			switch (agent.ConformAttemptCount.Value)
 			{
 				case 1:
-					RetryTime = TimeSpan.FromMinutes(5.0);
+					retryTime = TimeSpan.FromMinutes(5.0);
 					break;
 				case 2:
-					RetryTime = TimeSpan.FromMinutes(20.0);
+					retryTime = TimeSpan.FromMinutes(20.0);
 					break;
 				case 3:
-					RetryTime = TimeSpan.FromHours(1.0);
+					retryTime = TimeSpan.FromHours(1.0);
 					break;
 				default:
-					RetryTime = TimeSpan.FromHours(6.0);
+					retryTime = TimeSpan.FromHours(6.0);
 					break;
 			}
-			return UtcNow < Agent.LastConformTime + RetryTime;
+			return utcNow < agent.LastConformTime + retryTime;
 		}
 	}
 }
diff --git a/Engine/Source/Programs/Horde/Horde.Build/Agents/Impl/Tasks/RestartTaskSource.cs b/Engine/Source/Programs/Horde/Horde.Build/Agents/Impl/Tasks/RestartTaskSource.cs
--- a/Engine/Source/Programs/Horde/Horde.Build/Agents/Impl/Tasks/RestartTaskSource.cs
+++ b/Engine/Source/Programs/Horde/Horde.Build/Agents/Impl/Tasks/RestartTaskSource.cs
@@ -1,20 +1,15 @@
 // Copyright Epic Games, Inc. All Rights Reserved.
 
+using System.Threading;
+using System.Threading.Tasks;
 using Google.Protobuf;
-using Google.Protobuf.Reflection;
 using Google.Protobuf.WellKnownTypes;
-using HordeCommon;
-using HordeCommon.Rpc.Tasks;
 using Horde.Build.Api;
 using Horde.Build.Models;
 using Horde.Build.Services;
 using Horde.Build.Utilities;
-using MongoDB.Bson;
-using System;
-using System.Collections.Generic;
-using System.Linq;
-using System.Threading;
-using System.Threading.Tasks;
+using HordeCommon;
+using HordeCommon.Rpc.Tasks;
 
 namespace Horde.Build.Tasks.Impl
 {
@@ -28,34 +23,34 @@ namespace Horde.Build.Tasks.Impl
 
 		public override TaskSourceFlags Flags => TaskSourceFlags.AllowWhenDisabled | TaskSourceFlags.AllowDuringDowntime;
 
-		ILogFileService LogService;
+		readonly ILogFileService _logService;
 
-		public RestartTaskSource(ILogFileService LogService)
+		public RestartTaskSource(ILogFileService logService)
 		{
-			this.LogService = LogService;
+			_logService = logService;
 
-			this.OnLeaseStartedProperties.Add(nameof(RestartTask.LogId), x => new LogId(x.LogId));
+			OnLeaseStartedProperties.Add(nameof(RestartTask.LogId), x => new LogId(x.LogId));
 		}
 
-		public override async Task<AgentLease?> AssignLeaseAsync(IAgent Agent, CancellationToken CancellationToken)
+		public override async Task<AgentLease?> AssignLeaseAsync(IAgent agent, CancellationToken cancellationToken)
 		{
-			if (!Agent.RequestRestart)
+			if (!agent.RequestRestart)
 			{
 				return null;
 			}
-			if (Agent.Leases.Count > 0)
+			if (agent.Leases.Count > 0)
 			{
 				return AgentLease.Drain;
 			}
 
-			ILogFile Log = await LogService.CreateLogFileAsync(JobId.Empty, Agent.SessionId, LogType.Json);
+			ILogFile log = await _logService.CreateLogFileAsync(JobId.Empty, agent.SessionId, LogType.Json);
 
-			RestartTask Task = new RestartTask();
-			Task.LogId = Log.Id.ToString();
+			RestartTask task = new RestartTask();
+			task.LogId = log.Id.ToString();
 
-			byte[] Payload = Any.Pack(Task).ToByteArray();
+			byte[] payload = Any.Pack(task).ToByteArray();
 
-			return new AgentLease(LeaseId.GenerateNewId(), "Restart", null, null, Log.Id, LeaseState.Pending, null, true, Payload);
+			return new AgentLease(LeaseId.GenerateNewId(), "Restart", null, null, log.Id, LeaseState.Pending, null, true, payload);
 		}
 	}
 }
diff --git a/Engine/Source/Programs/Horde/Horde.Build/Agents/Impl/Tasks/ShutdownTaskSource.cs b/Engine/Source/Programs/Horde/Horde.Build/Agents/Impl/Tasks/ShutdownTaskSource.cs
--- a/Engine/Source/Programs/Horde/Horde.Build/Agents/Impl/Tasks/ShutdownTaskSource.cs
+++ b/Engine/Source/Programs/Horde/Horde.Build/Agents/Impl/Tasks/ShutdownTaskSource.cs
@@ -1,21 +1,15 @@
 // Copyright Epic Games, Inc. All Rights Reserved.
 
-using EpicGames.Core;
+using System.Threading;
+using System.Threading.Tasks;
 using Google.Protobuf;
-using Google.Protobuf.Reflection;
 using Google.Protobuf.WellKnownTypes;
-using HordeCommon;
-using HordeCommon.Rpc.Tasks;
 using Horde.Build.Api;
 using Horde.Build.Models;
 using Horde.Build.Services;
 using Horde.Build.Utilities;
-using MongoDB.Bson;
-using System;
-using System.Collections.Generic;
-using System.Linq;
-using System.Threading;
-using System.Threading.Tasks;
+using HordeCommon;
+using HordeCommon.Rpc.Tasks;
 
 namespace Horde.Build.Tasks.Impl
 {
@@ -29,33 +23,33 @@ namespace Horde.Build.Tasks.Impl
 
 		public override TaskSourceFlags Flags => TaskSourceFlags.AllowWhenDisabled | TaskSourceFlags.AllowDuringDowntime;
 
-		ILogFileService LogService;
+		readonly ILogFileService _logService;
 
-		public ShutdownTaskSource(ILogFileService LogService)
+		public ShutdownTaskSource(ILogFileService logService)
 		{
-			this.LogService = LogService;
-			this.OnLeaseStartedProperties.Add(nameof(ShutdownTask.LogId), x => new LogId(x.LogId));
+			_logService = logService;
+			OnLeaseStartedProperties.Add(nameof(ShutdownTask.LogId), x => new LogId(x.LogId));
 		}
 
-		public override async Task<AgentLease?> AssignLeaseAsync(IAgent Agent, CancellationToken CancellationToken)
+		public override async Task<AgentLease?> AssignLeaseAsync(IAgent agent, CancellationToken cancellationToken)
 		{
-			if (!Agent.RequestShutdown)
+			if (!agent.RequestShutdown)
 			{
 				return null;
 			}
-			if (Agent.Leases.Count > 0)
+			if (agent.Leases.Count > 0)
 			{
 				return AgentLease.Drain;
 			}
 
-			ILogFile Log = await LogService.CreateLogFileAsync(JobId.Empty, Agent.SessionId, LogType.Json);
+			ILogFile log = await _logService.CreateLogFileAsync(JobId.Empty, agent.SessionId, LogType.Json);
 
-			ShutdownTask Task = new ShutdownTask();
-			Task.LogId = Log.Id.ToString();
+			ShutdownTask task = new ShutdownTask();
+			task.LogId = log.Id.ToString();
 
-			byte[] Payload = Any.Pack(Task).ToByteArray();
+			byte[] payload = Any.Pack(task).ToByteArray();
 
-			return new AgentLease(LeaseId.GenerateNewId(), "Shutdown", null, null, Log.Id, LeaseState.Pending, null, true, Payload);
+			return new AgentLease(LeaseId.GenerateNewId(), "Shutdown", null, null, log.Id, LeaseState.Pending, null, true, payload);
 		}
 	}
 }
diff --git a/Engine/Source/Programs/Horde/Horde.Build/Agents/Impl/Tasks/UpgradeTaskSource.cs b/Engine/Source/Programs/Horde/Horde.Build/Agents/Impl/Tasks/UpgradeTaskSource.cs
--- a/Engine/Source/Programs/Horde/Horde.Build/Agents/Impl/Tasks/UpgradeTaskSource.cs
+++ b/Engine/Source/Programs/Horde/Horde.Build/Agents/Impl/Tasks/UpgradeTaskSource.cs
@@ -1,20 +1,16 @@
 // Copyright Epic Games, Inc. All Rights Reserved.
 
+using System;
+using System.Threading;
+using System.Threading.Tasks;
 using Google.Protobuf;
-using Google.Protobuf.Reflection;
 using Google.Protobuf.WellKnownTypes;
-using HordeCommon;
-using HordeCommon.Rpc.Tasks;
 using Horde.Build.Api;
 using Horde.Build.Models;
 using Horde.Build.Services;
 using Horde.Build.Utilities;
-using MongoDB.Bson;
-using System;
-using System.Collections.Generic;
-using System.Linq;
-using System.Threading;
-using System.Threading.Tasks;
+using HordeCommon;
+using HordeCommon.Rpc.Tasks;
 
 namespace Horde.Build.Tasks.Impl
 {
@@ -29,49 +25,49 @@ namespace Horde.Build.Tasks.Impl
 
 		public override TaskSourceFlags Flags => TaskSourceFlags.AllowWhenDisabled | TaskSourceFlags.AllowDuringDowntime;
 
-		AgentSoftwareService AgentSoftwareService;
-		ILogFileService LogService;
-		IClock Clock;
+		readonly AgentSoftwareService _agentSoftwareService;
+		readonly ILogFileService _logService;
+		readonly IClock _clock;
 
-		public UpgradeTaskSource(AgentSoftwareService AgentSoftwareService, ILogFileService LogService, IClock Clock)
+		public UpgradeTaskSource(AgentSoftwareService agentSoftwareService, ILogFileService logService, IClock clock)
 		{
-			this.AgentSoftwareService = AgentSoftwareService;
-			this.LogService = LogService;
-			this.Clock = Clock;
+			_agentSoftwareService = agentSoftwareService;
+			_logService = logService;
+			_clock = clock;
 
-			this.OnLeaseStartedProperties.Add(nameof(UpgradeTask.LogId), x => new LogId(x.LogId));
+			OnLeaseStartedProperties.Add(nameof(UpgradeTask.LogId), x => new LogId(x.LogId));
 		}
 
-		public async override Task<AgentLease?> AssignLeaseAsync(IAgent Agent, CancellationToken CancellationToken)
+		public override async Task<AgentLease?> AssignLeaseAsync(IAgent agent, CancellationToken cancellationToken)
 		{
-			string? RequiredVersion = await GetRequiredSoftwareVersion(Agent);
-			if (RequiredVersion != null && Agent.Version != RequiredVersion)
+			string? requiredVersion = await GetRequiredSoftwareVersion(agent);
+			if (requiredVersion != null && agent.Version != requiredVersion)
 			{
-				AgentLease? Lease = null;
-				if (Agent.Leases.Count == 0 && (Agent.LastUpgradeTime == null || Agent.LastUpgradeTime.Value + TimeSpan.FromMinutes(5.0) < Clock.UtcNow || Agent.LastUpgradeVersion != RequiredVersion.ToString()))
+				AgentLease? lease = null;
+				if (agent.Leases.Count == 0 && (agent.LastUpgradeTime == null || agent.LastUpgradeTime.Value + TimeSpan.FromMinutes(5.0) < _clock.UtcNow || agent.LastUpgradeVersion != requiredVersion.ToString()))
 				{
-					ILogFile LogFile = await LogService.CreateLogFileAsync(JobId.Empty, Agent.SessionId, LogType.Json);
+					ILogFile logFile = await _logService.CreateLogFileAsync(JobId.Empty, agent.SessionId, LogType.Json);
 
-					UpgradeTask Task = new UpgradeTask();
-					Task.SoftwareId = RequiredVersion.ToString();
-					Task.LogId = LogFile.Id.ToString();
+					UpgradeTask task = new UpgradeTask();
+					task.SoftwareId = requiredVersion.ToString();
+					task.LogId = logFile.Id.ToString();
 
-					byte[] Payload;
-					if (Agent.Version == "5.0.0-17425336" || Agent.Version == "5.0.0-17448746")
+					byte[] payload;
+					if (agent.Version == "5.0.0-17425336" || agent.Version == "5.0.0-17448746")
 					{
-						Any Any = new Any();
-						Any.TypeUrl = "type.googleapis.com/Horde.UpgradeTask";
-						Any.Value = Task.ToByteString();
-						Payload = Any.ToByteArray();
+						Any any = new Any();
+						any.TypeUrl = "type.googleapis.com/Horde.UpgradeTask";
+						any.Value = task.ToByteString();
+						payload = any.ToByteArray();
 					}
 					else
 					{
-						Payload = Any.Pack(Task).ToByteArray();
+						payload = Any.Pack(task).ToByteArray();
 					}
 
-					Lease = new AgentLease(LeaseId.GenerateNewId(), $"Upgrade to {RequiredVersion}", null, null, LogFile.Id, LeaseState.Pending, null, true, Payload);
+					lease = new AgentLease(LeaseId.GenerateNewId(), $"Upgrade to {requiredVersion}", null, null, logFile.Id, LeaseState.Pending, null, true, payload);
 				}
-				return Lease;
+				return lease;
 			}
 			return null;
 		}
@@ -79,13 +75,13 @@ namespace Horde.Build.Tasks.Impl
 		/// <summary>
 		/// Determines the client software version that should be installed on an agent
 		/// </summary>
-		/// <param name="Agent">The agent instance</param>
+		/// <param name="agent">The agent instance</param>
 		/// <returns>Unique id of the client version this agent should be running</returns>
-		public async Task<string?> GetRequiredSoftwareVersion(IAgent Agent)
+		public async Task<string?> GetRequiredSoftwareVersion(IAgent agent)
 		{
-			AgentSoftwareChannelName ChannelName = Agent.Channel ?? AgentSoftwareService.DefaultChannelName;
-			IAgentSoftwareChannel? Channel = await AgentSoftwareService.GetCachedChannelAsync(ChannelName);
-			return Channel?.Version;
+			AgentSoftwareChannelName channelName = agent.Channel ?? AgentSoftwareService.DefaultChannelName;
+			IAgentSoftwareChannel? channel = await _agentSoftwareService.GetCachedChannelAsync(channelName);
+			return channel?.Version;
 		}
 	}
 }
diff --git a/Engine/Source/Programs/Horde/Horde.Build/Agents/Impl/UpgradeService.cs b/Engine/Source/Programs/Horde/Horde.Build/Agents/Impl/UpgradeService.cs
--- a/Engine/Source/Programs/Horde/Horde.Build/Agents/Impl/UpgradeService.cs
+++ b/Engine/Source/Programs/Horde/Horde.Build/Agents/Impl/UpgradeService.cs
@@ -1,16 +1,12 @@
 // Copyright Epic Games, Inc. All Rights Reserved.
 
+using System;
+using System.Threading.Tasks;
 using Horde.Build.Collections;
 using Horde.Build.Collections.Impl;
 using Horde.Build.Models;
 using Microsoft.Extensions.DependencyInjection;
 using Microsoft.Extensions.Logging;
-using MongoDB.Driver;
-using Serilog.Events;
-using System;
-using System.Collections.Generic;
-using System.Linq;
-using System.Threading.Tasks;
 
 namespace Horde.Build.Services
 {
@@ -32,39 +28,39 @@ namespace Horde.Build.Services
 		/// <summary>
 		/// The DI service provider
 		/// </summary>
-		IServiceProvider ServiceProvider;
+		readonly IServiceProvider _serviceProvider;
 
 		/// <summary>
 		/// Logger instance
 		/// </summary>
-		ILogger<UpgradeService> Logger;
+		readonly ILogger<UpgradeService> _logger;
 
 		/// <summary>
 		/// Constructor
 		/// </summary>
-		/// <param name="DatabaseService">The database service singleton</param>
-		/// <param name="ServiceProvider">The DI service provider</param>
-		/// <param name="Logger">Logger instance</param>
-		public UpgradeService(DatabaseService DatabaseService, IServiceProvider ServiceProvider, ILogger<UpgradeService> Logger)
+		/// <param name="databaseService">The database service singleton</param>
+		/// <param name="serviceProvider">The DI service provider</param>
+		/// <param name="logger">Logger instance</param>
+		public UpgradeService(DatabaseService databaseService, IServiceProvider serviceProvider, ILogger<UpgradeService> logger)
 		{
-			this.DatabaseService = DatabaseService;
-			this.ServiceProvider = ServiceProvider;
-			this.Logger = Logger;
+			DatabaseService = databaseService;
+			_serviceProvider = serviceProvider;
+			_logger = logger;
 		}
 
 		/// <summary>
 		/// Update the database schema
 		/// </summary>
-		/// <param name="FromVersion">The version to upgrade from. Automatically detected if this is null.</param>
+		/// <param name="fromVersion">The version to upgrade from. Automatically detected if this is null.</param>
 		/// <returns>Async task</returns>
-		public async Task UpgradeSchemaAsync(int? FromVersion)
+		public async Task UpgradeSchemaAsync(int? fromVersion)
 		{
-			Globals Globals = await DatabaseService.GetGlobalsAsync();
-			int SchemaVersion = FromVersion ?? Globals.SchemaVersion ?? 0;
+			Globals globals = await DatabaseService.GetGlobalsAsync();
+			int schemaVersion = fromVersion ?? globals.SchemaVersion ?? 0;
 
-			while (SchemaVersion < LatestSchemaVersion)
+			while (schemaVersion < LatestSchemaVersion)
 			{
-				Logger.LogInformation("Upgrading from schema version {Version}", SchemaVersion);
+				_logger.LogInformation("Upgrading from schema version {Version}", schemaVersion);
 
 				// Handle the current version
 				//				if (SchemaVersion == 0)
@@ -79,30 +75,30 @@ namespace Horde.Build.Services
 				//					await IssueCollection.FixResolvedIssuesAsync(Logger);
 				//				}
 
-				if (SchemaVersion == 3)
+				if (schemaVersion == 3)
 				{
-					IJobCollection JobCollection = ServiceProvider.GetRequiredService<IJobCollection>();
-					await JobCollection.UpgradeDocumentsAsync();
+					IJobCollection jobCollection = _serviceProvider.GetRequiredService<IJobCollection>();
+					await jobCollection.UpgradeDocumentsAsync();
 				}
-				if (SchemaVersion == 4)
+				if (schemaVersion == 4)
 				{
-					UserCollectionV1 UserCollectionV1 = new UserCollectionV1(ServiceProvider.GetRequiredService<DatabaseService>());
-					using UserCollectionV2 UserCollectionV2 = new UserCollectionV2(ServiceProvider.GetRequiredService<DatabaseService>(), ServiceProvider.GetRequiredService<ILogger<UserCollectionV2>>());
-					await UserCollectionV2.ResaveDocumentsAsync(UserCollectionV1);
+					UserCollectionV1 userCollectionV1 = new UserCollectionV1(_serviceProvider.GetRequiredService<DatabaseService>());
+					using UserCollectionV2 userCollectionV2 = new UserCollectionV2(_serviceProvider.GetRequiredService<DatabaseService>(), _serviceProvider.GetRequiredService<ILogger<UserCollectionV2>>());
+					await userCollectionV2.ResaveDocumentsAsync(userCollectionV1);
 				}
 
 				// Increment the version number
-				SchemaVersion++;
+				schemaVersion++;
 
 				// Try to update the current schema version number
-				while (Globals.SchemaVersion == null || Globals.SchemaVersion < SchemaVersion)
+				while (globals.SchemaVersion == null || globals.SchemaVersion < schemaVersion)
 				{
-					Globals.SchemaVersion = SchemaVersion;
-					if (await DatabaseService.TryUpdateSingletonAsync(Globals))
+					globals.SchemaVersion = schemaVersion;
+					if (await DatabaseService.TryUpdateSingletonAsync(globals))
 					{
 						break;
 					}
-					Globals = await DatabaseService.GetGlobalsAsync();
+					globals = await DatabaseService.GetGlobalsAsync();
 				}
 			}
 		}
diff --git a/Engine/Source/Programs/Horde/Horde.Build/Api/Acls.cs b/Engine/Source/Programs/Horde/Horde.Build/Api/Acls.cs
--- a/Engine/Source/Programs/Horde/Horde.Build/Api/Acls.cs
+++ b/Engine/Source/Programs/Horde/Horde.Build/Api/Acls.cs
@@ -1,13 +1,8 @@
 ﻿// Copyright Epic Games, Inc. All Rights Reserved.
 
-using Horde.Build.Acls;
-using Horde.Build.Models;
-using System;
 using System.Collections.Generic;
 using System.ComponentModel.DataAnnotations;
-using System.Linq;
-using System.Security.Claims;
-using System.Threading.Tasks;
+using Horde.Build.Acls;
 
 namespace Horde.Build.Api
 {
@@ -86,21 +81,21 @@ namespace Horde.Build.Api
 		/// <summary>
 		/// Constructor
 		/// </summary>
-		/// <param name="Claim">The claim to construct from</param>
-		public GetAclClaimResponse(AclClaim Claim)
-			: this(Claim.Type, Claim.Value)
+		/// <param name="claim">The claim to construct from</param>
+		public GetAclClaimResponse(AclClaim claim)
+			: this(claim.Type, claim.Value)
 		{
 		}
 
 		/// <summary>
 		/// Constructor
 		/// </summary>
-		/// <param name="Type">The claim type</param>
-		/// <param name="Value">The claim value</param>
-		public GetAclClaimResponse(string Type, string Value)
+		/// <param name="type">The claim type</param>
+		/// <param name="value">The claim value</param>
+		public GetAclClaimResponse(string type, string value)
 		{
-			this.Type = Type;
-			this.Value = Value;
+			Type = type;
+			Value = value;
 		}
 	}
 
@@ -122,11 +117,11 @@ namespace Horde.Build.Api
 		/// <summary>
 		/// Constructor
 		/// </summary>
-		/// <param name="AclEntry">The acl entry to construct from</param>
-		public GetAclEntryResponse(AclEntry AclEntry)
+		/// <param name="aclEntry">The acl entry to construct from</param>
+		public GetAclEntryResponse(AclEntry aclEntry)
 		{
-			this.Claim = new GetAclClaimResponse(AclEntry.Claim.Type, AclEntry.Claim.Value);
-			this.Actions = AclEntry.GetActionNames(AclEntry.Actions);
+			Claim = new GetAclClaimResponse(aclEntry.Claim.Type, aclEntry.Claim.Value);
+			Actions = AclEntry.GetActionNames(aclEntry.Actions);
 		}
 	}
 
@@ -154,12 +149,12 @@ namespace Horde.Build.Api
 		/// <summary>
 		/// Constructor
 		/// </summary>
-		/// <param name="Acl">The ACL to construct from</param>
-		public GetAclResponse(Acl Acl)
+		/// <param name="acl">The ACL to construct from</param>
+		public GetAclResponse(Acl acl)
 		{
-			this.Entries = Acl.Entries.ConvertAll(x => new GetAclEntryResponse(x));
-			this.Inherit = Acl.Inherit;
-			this.Exceptions = Acl.Exceptions?.ConvertAll(x => x.ToString());
+			Entries = acl.Entries.ConvertAll(x => new GetAclEntryResponse(x));
+			Inherit = acl.Inherit;
+			Exceptions = acl.Exceptions?.ConvertAll(x => x.ToString());
 		}
 	}
 }
diff --git a/Engine/Source/Programs/Horde/Horde.Build/Api/AgentSoftwareChannels.cs b/Engine/Source/Programs/Horde/Horde.Build/Api/AgentSoftwareChannels.cs
--- a/Engine/Source/Programs/Horde/Horde.Build/Api/AgentSoftwareChannels.cs
+++ b/Engine/Source/Programs/Horde/Horde.Build/Api/AgentSoftwareChannels.cs
@@ -1,10 +1,7 @@
 // Copyright Epic Games, Inc. All Rights Reserved.
 
-using Horde.Build.Models;
 using System;
-using System.Collections.Generic;
-using System.Linq;
-using System.Threading.Tasks;
+using Horde.Build.Models;
 
 namespace Horde.Build.Api
 {
@@ -36,13 +33,13 @@ namespace Horde.Build.Api
 		/// <summary>
 		/// Constructor
 		/// </summary>
-		/// <param name="Channel">The channel information</param>
-		public GetAgentSoftwareChannelResponse(IAgentSoftwareChannel Channel)
+		/// <param name="channel">The channel information</param>
+		public GetAgentSoftwareChannelResponse(IAgentSoftwareChannel channel)
 		{
-			this.Name = Channel.Name.ToString();
-			this.ModifiedBy = Channel.ModifiedBy;
-			this.ModifiedTime = Channel.ModifiedTime;
-			this.Version = Channel.Version.ToString();
+			Name = channel.Name.ToString();
+			ModifiedBy = channel.ModifiedBy;
+			ModifiedTime = channel.ModifiedTime;
+			Version = channel.Version.ToString();
 		}
 	}
 }
diff --git a/Engine/Source/Programs/Horde/Horde.Build/Api/Agents.cs b/Engine/Source/Programs/Horde/Horde.Build/Api/Agents.cs
--- a/Engine/Source/Programs/Horde/Horde.Build/Api/Agents.cs
+++ b/Engine/Source/Programs/Horde/Horde.Build/Api/Agents.cs
@@ -1,19 +1,15 @@
 // Copyright Epic Games, Inc. All Rights Reserved.
 
-using EpicGames.Core;
-using HordeCommon;
-using Horde.Build.Models;
-using Horde.Build.Utilities;
 using System;
 using System.Collections.Generic;
 using System.ComponentModel.DataAnnotations;
 using System.Linq;
-using System.Threading.Tasks;
+using EpicGames.Core;
+using Horde.Build.Models;
+using HordeCommon;
 
 namespace Horde.Build.Api
 {
-	using PoolId = StringId<IPool>;
-
 	/// <summary>
 	/// Parameters to register a new agent
 	/// </summary>
@@ -59,10 +55,10 @@ namespace Horde.Build.Api
 		/// <summary>
 		/// Constructor
 		/// </summary>
-		/// <param name="Id">Unique id for this agent</param>
-		public CreateAgentResponse(string Id)
+		/// <param name="id">Unique id for this agent</param>
+		public CreateAgentResponse(string id)
 		{
-			this.Id = Id;
+			Id = id;
 		}
 	}
 
@@ -185,38 +181,38 @@ namespace Horde.Build.Api
 		/// <summary>
 		/// Constructor
 		/// </summary>
-		/// <param name="Lease">The lease to initialize from</param>
-		/// <param name="Details">The payload details</param>
-		public GetAgentLeaseResponse(AgentLease Lease, Dictionary<string, string>? Details)
+		/// <param name="lease">The lease to initialize from</param>
+		/// <param name="details">The payload details</param>
+		public GetAgentLeaseResponse(AgentLease lease, Dictionary<string, string>? details)
 		{
-			this.Id = Lease.Id.ToString();
-			this.Name = Lease.Name;
-			this.LogId = Lease.LogId?.ToString();
-			this.State = Lease.State;
-			this.StartTime = Lease.StartTime;
-			this.Executing = Lease.Active;
-			this.FinishTime = Lease.ExpiryTime;
-			this.Details = Details;
+			Id = lease.Id.ToString();
+			Name = lease.Name;
+			LogId = lease.LogId?.ToString();
+			State = lease.State;
+			StartTime = lease.StartTime;
+			Executing = lease.Active;
+			FinishTime = lease.ExpiryTime;
+			Details = details;
 		}
 
 		/// <summary>
 		/// Converts this lease to a response object
 		/// </summary>
-		/// <param name="Lease">The lease to initialize from</param>
-		/// <param name="Details">The payload details</param>
-		/// <param name="AgentRate">Rate for running this agent</param>
-		public GetAgentLeaseResponse(ILease Lease, Dictionary<string, string>? Details, double? AgentRate)
+		/// <param name="lease">The lease to initialize from</param>
+		/// <param name="details">The payload details</param>
+		/// <param name="agentRate">Rate for running this agent</param>
+		public GetAgentLeaseResponse(ILease lease, Dictionary<string, string>? details, double? agentRate)
 		{
-			this.Id = Lease.Id.ToString();
-			this.AgentId = Lease.AgentId.ToString();
-			this.AgentRate = AgentRate;
-			this.Name = Lease.Name;
-			this.LogId = Lease.LogId?.ToString();
-			this.StartTime = Lease.StartTime;
-			this.Executing = (Lease.FinishTime == null);
-			this.FinishTime = Lease.FinishTime;
-			this.Details = Details;
-			this.Outcome = Lease.Outcome;
+			Id = lease.Id.ToString();
+			AgentId = lease.AgentId.ToString();
+			AgentRate = agentRate;
+			Name = lease.Name;
+			LogId = lease.LogId?.ToString();
+			StartTime = lease.StartTime;
+			Executing = (lease.FinishTime == null);
+			FinishTime = lease.FinishTime;
+			Details = details;
+			Outcome = lease.Outcome;
 		}
 	}
 
@@ -255,14 +251,14 @@ namespace Horde.Build.Api
 		/// <summary>
 		/// Constructor
 		/// </summary>
-		/// <param name="Session">The session to construct from</param>
-		public GetAgentSessionResponse(ISession Session)
+		/// <param name="session">The session to construct from</param>
+		public GetAgentSessionResponse(ISession session)
 		{
-			this.Id = Session.Id.ToString();
-			this.StartTime = Session.StartTime;
-			this.FinishTime = Session.FinishTime;
-			this.Properties = (Session.Properties != null) ? new List<string>(Session.Properties) : null;
-			this.Version = Session.Version?.ToString();
+			Id = session.Id.ToString();
+			StartTime = session.StartTime;
+			FinishTime = session.FinishTime;
+			Properties = (session.Properties != null) ? new List<string>(session.Properties) : null;
+			Version = session.Version?.ToString();
 		}
 	}
 
@@ -299,20 +295,20 @@ namespace Horde.Build.Api
 		/// <summary>
 		/// Whether to use an incremental workspace
 		/// </summary>
-		public bool bIncremental { get; set; }
+		public bool BIncremental { get; set; }
 
 		/// <summary>
 		/// Constructor
 		/// </summary>
-		/// <param name="Workspace">The workspace to construct from</param>
-		public GetAgentWorkspaceResponse(AgentWorkspace Workspace)
+		/// <param name="workspace">The workspace to construct from</param>
+		public GetAgentWorkspaceResponse(AgentWorkspace workspace)
 		{
-			this.Cluster = Workspace.Cluster;
-			this.UserName = Workspace.UserName;
-			this.Identifier = Workspace.Identifier;
-			this.Stream = Workspace.Stream;
-			this.View = Workspace.View;
-			this.bIncremental = Workspace.bIncremental;
+			Cluster = workspace.Cluster;
+			UserName = workspace.UserName;
+			Identifier = workspace.Identifier;
+			Stream = workspace.Stream;
+			View = workspace.View;
+			BIncremental = workspace.BIncremental;
 		}
 	}
 
@@ -459,43 +455,43 @@ namespace Horde.Build.Api
 		/// <summary>
 		/// Constructor
 		/// </summary>
-		/// <param name="Agent">The agent to construct from</param>
-		/// <param name="Leases">Active leases</param>
-		/// <param name="Rate">Rate for this agent</param>
+		/// <param name="agent">The agent to construct from</param>
+		/// <param name="leases">Active leases</param>
+		/// <param name="rate">Rate for this agent</param>
 		/// <param name="bIncludeAcl">Whether to include the ACL in the response</param>
-		public GetAgentResponse(IAgent Agent, List<GetAgentLeaseResponse> Leases, double? Rate, bool bIncludeAcl)
+		public GetAgentResponse(IAgent agent, List<GetAgentLeaseResponse> leases, double? rate, bool bIncludeAcl)
 		{
-			this.Id = Agent.Id.ToString();
-			this.Name = Agent.Id.ToString();
-			this.Enabled = Agent.Enabled;
-			this.Rate = Rate;
-			this.Properties = new List<string>(Agent.Properties);
-			this.Resources = new Dictionary<string, int>(Agent.Resources);
-			this.SessionId = Agent.SessionId?.ToString();
-			this.Online = Agent.IsSessionValid(DateTime.UtcNow);
-			this.Deleted = Agent.Deleted;
-			this.PendingConform = Agent.RequestConform;
-			this.PendingFullConform = Agent.RequestFullConform;
-			this.PendingRestart = Agent.RequestRestart;
-			this.PendingShutdown = Agent.RequestShutdown;
-			this.LastShutdownReason = Agent.LastShutdownReason ?? "Unknown";
-			this.LastConformTime = Agent.LastConformTime;
-			this.NextConformTime =   Agent.LastConformTime;
-			this.ConformAttemptCount = Agent.ConformAttemptCount;
-			this.Version = Agent.Version?.ToString() ?? "Unknown";
-			if(Agent.Channel != null)
+			Id = agent.Id.ToString();
+			Name = agent.Id.ToString();
+			Enabled = agent.Enabled;
+			Rate = rate;
+			Properties = new List<string>(agent.Properties);
+			Resources = new Dictionary<string, int>(agent.Resources);
+			SessionId = agent.SessionId?.ToString();
+			Online = agent.IsSessionValid(DateTime.UtcNow);
+			Deleted = agent.Deleted;
+			PendingConform = agent.RequestConform;
+			PendingFullConform = agent.RequestFullConform;
+			PendingRestart = agent.RequestRestart;
+			PendingShutdown = agent.RequestShutdown;
+			LastShutdownReason = agent.LastShutdownReason ?? "Unknown";
+			LastConformTime = agent.LastConformTime;
+			NextConformTime =   agent.LastConformTime;
+			ConformAttemptCount = agent.ConformAttemptCount;
+			Version = agent.Version?.ToString() ?? "Unknown";
+			if(agent.Channel != null)
 			{
-				this.Version += $" ({Agent.Channel})";
+				Version += $" ({agent.Channel})";
 			}
-			this.Version = Agent.Version?.ToString();
-			this.Channel = Agent.Channel?.ToString();
-			this.UpdateTime = Agent.UpdateTime;
-			this.Pools = Agent.GetPools().Select(x => x.ToString()).ToList();
-			this.Workspaces = Agent.Workspaces.ConvertAll(x => new GetAgentWorkspaceResponse(x));
-			this.Capabilities = new { Devices = new[] { new { Properties = Agent.Properties, Resources = Agent.Resources } } };
-			this.Leases = Leases;
-			this.Acl = (bIncludeAcl && Agent.Acl != null) ? new GetAclResponse(Agent.Acl) : null;
-			this.Comment = Agent.Comment;
+			Version = agent.Version?.ToString();
+			Channel = agent.Channel?.ToString();
+			UpdateTime = agent.UpdateTime;
+			Pools = agent.GetPools().Select(x => x.ToString()).ToList();
+			Workspaces = agent.Workspaces.ConvertAll(x => new GetAgentWorkspaceResponse(x));
+			Capabilities = new { Devices = new[] { new { agent.Properties, agent.Resources } } };
+			Leases = leases;
+			Acl = (bIncludeAcl && agent.Acl != null) ? new GetAclResponse(agent.Acl) : null;
+			Comment = agent.Comment;
 		}
 	}
 }
diff --git a/Engine/Source/Programs/Horde/Horde.Build/Api/Artifacts.cs b/Engine/Source/Programs/Horde/Horde.Build/Api/Artifacts.cs
--- a/Engine/Source/Programs/Horde/Horde.Build/Api/Artifacts.cs
+++ b/Engine/Source/Programs/Horde/Horde.Build/Api/Artifacts.cs
@@ -1,8 +1,7 @@
 // Copyright Epic Games, Inc. All Rights Reserved.
 
-using Horde.Build.Models;
 using System.Collections.Generic;
-using System.ComponentModel.DataAnnotations;
+using Horde.Build.Models;
 
 namespace Horde.Build.Api
 {
@@ -19,10 +18,10 @@ namespace Horde.Build.Api
 		/// <summary>
 		/// Constructor
 		/// </summary>
-		/// <param name="Id">The artifact file id</param>
-		public CreateArtifactResponse(string Id)
+		/// <param name="id">The artifact file id</param>
+		public CreateArtifactResponse(string id)
 		{
-			this.Id = Id;
+			Id = id;
 		}
 	}
 
@@ -49,14 +48,14 @@ namespace Horde.Build.Api
 		/// <summary>
 		/// Constructor
 		/// </summary>
-		/// <param name="JobId">Job id to get all artifacts for</param>
-		/// <param name="StepId">step to filter by</param>
-		/// <param name="ArtifactIds">The artifact ids.  Returns all artifacts for a job </param>
-		public GetArtifactZipRequest(string? JobId, string? StepId, List<string>? ArtifactIds)
+		/// <param name="jobId">Job id to get all artifacts for</param>
+		/// <param name="stepId">step to filter by</param>
+		/// <param name="artifactIds">The artifact ids.  Returns all artifacts for a job </param>
+		public GetArtifactZipRequest(string? jobId, string? stepId, List<string>? artifactIds)
 		{
-			this.JobId = JobId;
-			this.StepId = StepId;
-			this.ArtifactIds = ArtifactIds;
+			JobId = jobId;
+			StepId = stepId;
+			ArtifactIds = artifactIds;
 		}
 	}
 
@@ -103,17 +102,17 @@ namespace Horde.Build.Api
 		/// <summary>
 		/// Constructor
 		/// </summary>
-		/// <param name="Artifact">The artifact to construct from</param>
-		/// <param name="Code">The direct download code</param>
-		public GetArtifactResponse(IArtifact Artifact, string? Code)
+		/// <param name="artifact">The artifact to construct from</param>
+		/// <param name="code">The direct download code</param>
+		public GetArtifactResponse(IArtifact artifact, string? code)
 		{
-			this.Id = Artifact.Id.ToString();
-			this.Name = Artifact.Name;
-			this.JobId = Artifact.JobId.ToString();
-			this.StepId = Artifact.StepId.ToString();
-			this.MimeType = Artifact.MimeType;
-			this.Length = Artifact.Length;
-			this.Code = Code;
+			Id = artifact.Id.ToString();
+			Name = artifact.Name;
+			JobId = artifact.JobId.ToString();
+			StepId = artifact.StepId.ToString();
+			MimeType = artifact.MimeType;
+			Length = artifact.Length;
+			Code = code;
 		}
 	}
 }
diff --git a/Engine/Source/Programs/Horde/Horde.Build/Api/Changes.cs b/Engine/Source/Programs/Horde/Horde.Build/Api/Changes.cs
--- a/Engine/Source/Programs/Horde/Horde.Build/Api/Changes.cs
+++ b/Engine/Source/Programs/Horde/Horde.Build/Api/Changes.cs
@@ -1,10 +1,7 @@
 // Copyright Epic Games, Inc. All Rights Reserved.
 
-using Horde.Build.Models;
-using System;
 using System.Collections.Generic;
-using System.Linq;
-using System.Threading.Tasks;
+using Horde.Build.Models;
 
 namespace Horde.Build.Api
 {
@@ -36,13 +33,13 @@ namespace Horde.Build.Api
 		/// <summary>
 		/// Constructor
 		/// </summary>
-		/// <param name="ChangeSummary">The commit to construct from</param>
-		public GetChangeSummaryResponse(ChangeSummary ChangeSummary)
+		/// <param name="changeSummary">The commit to construct from</param>
+		public GetChangeSummaryResponse(ChangeSummary changeSummary)
 		{
-			this.Number = ChangeSummary.Number;
-			this.Author = ChangeSummary.Author.Name;
-			this.AuthorInfo = new GetThinUserInfoResponse(ChangeSummary.Author);
-			this.Description = ChangeSummary.Description;
+			Number = changeSummary.Number;
+			Author = changeSummary.Author.Name;
+			AuthorInfo = new GetThinUserInfoResponse(changeSummary.Author);
+			Description = changeSummary.Description;
 		}
 	}
 
@@ -79,14 +76,14 @@ namespace Horde.Build.Api
 		/// <summary>
 		/// Constructor
 		/// </summary>
-		/// <param name="ChangeDetails">The commit to construct from</param>
-		public GetChangeDetailsResponse(ChangeDetails ChangeDetails)
+		/// <param name="changeDetails">The commit to construct from</param>
+		public GetChangeDetailsResponse(ChangeDetails changeDetails)
 		{
-			this.Number = ChangeDetails.Number;
-			this.Author = ChangeDetails.Author.Name;
-			this.AuthorInfo = new GetThinUserInfoResponse(ChangeDetails.Author);
-			this.Description = ChangeDetails.Description;
-			this.Files = ChangeDetails.Files.ConvertAll(x => x.Path);
+			Number = changeDetails.Number;
+			Author = changeDetails.Author.Name;
+			AuthorInfo = new GetThinUserInfoResponse(changeDetails.Author);
+			Description = changeDetails.Description;
+			Files = changeDetails.Files.ConvertAll(x => x.Path);
 		}
 	}
 }
diff --git a/Engine/Source/Programs/Horde/Horde.Build/Api/Config.cs b/Engine/Source/Programs/Horde/Horde.Build/Api/Config.cs
--- a/Engine/Source/Programs/Horde/Horde.Build/Api/Config.cs
+++ b/Engine/Source/Programs/Horde/Horde.Build/Api/Config.cs
@@ -2,7 +2,6 @@
 
 using System;
 using System.Collections.Generic;
-using Horde.Build.Models;
 
 namespace Horde.Build.Api
 {
@@ -141,7 +140,6 @@ namespace Horde.Build.Api
 		/// </summary>
 		public string LocalLogsDir { get; set; } = null!;
 
-
 		/// <summary>
 		/// Local blob storage directory, if using type filesystem
 		/// </summary>
@@ -285,7 +283,6 @@ namespace Horde.Build.Api
 		/// </summary>
 		public string? UserServerSettingsPath { get; set; }
 
-
 		/// <summary>
 		/// Lazily computed timezone value
 		/// </summary>
@@ -294,93 +291,88 @@ namespace Horde.Build.Api
 		/// <summary>
 		/// Constructor
 		/// </summary>		
-		public GetServerSettingsResponse(ServerSettings Settings, string? GlobalConfigPath = null)
+		public GetServerSettingsResponse(ServerSettings settings, string? globalConfigPath = null)
 		{
-			this.GlobalConfigPath = GlobalConfigPath;
+			GlobalConfigPath = globalConfigPath;
 
-			if (this.GlobalConfigPath == null)
+			if (GlobalConfigPath == null)
 			{
-				this.GlobalConfigPath = Settings.ConfigPath;
+				GlobalConfigPath = settings.ConfigPath;
 			}
 
-			this.NumServerUpdates = 0;
-
-			this.UserServerSettingsPath = Program.UserConfigFile.ToString();
-			this.DatabaseConnectionString = Settings.DatabaseConnectionString;
-			this.DatabaseName = Settings.DatabaseName;
-			this.DatabasePublicCert = Settings.DatabasePublicCert;
-			this.DatabaseReadOnlyMode = Settings.DatabaseReadOnlyMode;
-
-			this.AdminClaimType = Settings.AdminClaimType;
-			this.AdminClaimValue = Settings.AdminClaimValue;
-
-			this.ServerPrivateCert = Settings.ServerPrivateCert;
-			this.OidcAuthority = Settings.OidcAuthority;
-			this.OidcClientId = Settings.OidcClientId;
-			this.OidcSigninRedirect = Settings.OidcSigninRedirect;
-
-			this.JwtIssuer = Settings.JwtIssuer;
-			this.JwtSecret = Settings.JwtSecret;
-			this.JwtExpiryTimeHours = Settings.JwtExpiryTimeHours;
-
-			this.CorsEnabled = Settings.CorsEnabled;
-			this.CorsOrigin = Settings.CorsOrigin;
-			this.EnableScheduleInTestData = Settings.EnableScheduleInTestData;
-			this.SchedulePollingInterval = Settings.SchedulePollingInterval;
-
-
-			this.NoResourceBackOffTime = Settings.NoResourceBackOffTime;
-			this.InitiateJobBackOffTime = Settings.InitiateJobBackOffTime;
-			this.UnknownErrorBackOffTime = Settings.UnknownErrorBackOffTime;
-
-
-			this.RedisConnectionConfig = Settings.RedisConnectionConfig;
-			this.LogServiceWriteCacheType = Settings.LogServiceWriteCacheType;
-			this.ExternalStorageProviderType = Settings.ExternalStorageProviderType;
-			this.LocalLogsDir = Settings.LocalLogsDir;
-			this.LocalBlobsDir = Settings.LocalBlobsDir;
-			this.LocalArtifactsDir = Settings.LocalArtifactsDir;
-
-			this.S3BucketRegion = Settings.S3BucketRegion;
-			this.S3CredentialType = Settings.S3CredentialType;
-			this.S3ClientKeyId = Settings.S3ClientKeyId;
-			this.S3ClientSecret = Settings.S3ClientSecret;
-			this.S3AssumeArn = Settings.S3AssumeArn;
-			this.S3LogBucketName = Settings.S3LogBucketName;
-			this.S3ArtifactBucketName = Settings.S3ArtifactBucketName;
-
-			this.LogRelayServer = Settings.LogRelayServer;
-			this.LogRelayBearerToken = Settings.LogRelayBearerToken;
-			this.LogJsonToStdOut = Settings.LogJsonToStdOut;
-
-			this.FleetManager = Settings.FleetManager;
-			this.DisableSchedules = Settings.DisableSchedules;
-			this.ScheduleTimeZone = Settings.ScheduleTimeZone;
-
-			this.SlackToken = Settings.SlackToken;
-			this.SlackSocketToken = Settings.SlackSocketToken;
-
-
-			this.SlackToken = Settings.SlackToken;
-			this.SlackToken = Settings.SlackToken;
-			this.UpdateStreamsNotificationChannel = Settings.UpdateStreamsNotificationChannel;
-
-			this.SmtpServer = Settings.SmtpServer;
-			this.EmailSenderAddress = Settings.EmailSenderAddress;
-			this.EmailSenderName = Settings.EmailSenderName;
-
-			this.UpdateStreamsNotificationChannel = Settings.UpdateStreamsNotificationChannel;
-			this.UpdateStreamsNotificationChannel = Settings.UpdateStreamsNotificationChannel;
-			this.UpdateStreamsNotificationChannel = Settings.UpdateStreamsNotificationChannel;
-
-			this.P4BridgeServiceUsername = Settings.P4BridgeServiceUsername;
-			this.P4BridgeServicePassword = Settings.P4BridgeServicePassword;
-			this.P4BridgeCanImpersonate = Settings.P4BridgeCanImpersonate;
-
-
-			this.GlobalThreadPoolMinSize = Settings.GlobalThreadPoolMinSize;
-			this.TimeZoneInfo = Settings.TimeZoneInfo;
-
+			NumServerUpdates = 0;
+
+			UserServerSettingsPath = Program.UserConfigFile.ToString();
+			DatabaseConnectionString = settings.DatabaseConnectionString;
+			DatabaseName = settings.DatabaseName;
+			DatabasePublicCert = settings.DatabasePublicCert;
+			DatabaseReadOnlyMode = settings.DatabaseReadOnlyMode;
+
+			AdminClaimType = settings.AdminClaimType;
+			AdminClaimValue = settings.AdminClaimValue;
+
+			ServerPrivateCert = settings.ServerPrivateCert;
+			OidcAuthority = settings.OidcAuthority;
+			OidcClientId = settings.OidcClientId;
+			OidcSigninRedirect = settings.OidcSigninRedirect;
+
+			JwtIssuer = settings.JwtIssuer;
+			JwtSecret = settings.JwtSecret;
+			JwtExpiryTimeHours = settings.JwtExpiryTimeHours;
+
+			CorsEnabled = settings.CorsEnabled;
+			CorsOrigin = settings.CorsOrigin;
+			EnableScheduleInTestData = settings.EnableScheduleInTestData;
+			SchedulePollingInterval = settings.SchedulePollingInterval;
+
+			NoResourceBackOffTime = settings.NoResourceBackOffTime;
+			InitiateJobBackOffTime = settings.InitiateJobBackOffTime;
+			UnknownErrorBackOffTime = settings.UnknownErrorBackOffTime;
+
+			RedisConnectionConfig = settings.RedisConnectionConfig;
+			LogServiceWriteCacheType = settings.LogServiceWriteCacheType;
+			ExternalStorageProviderType = settings.ExternalStorageProviderType;
+			LocalLogsDir = settings.LocalLogsDir;
+			LocalBlobsDir = settings.LocalBlobsDir;
+			LocalArtifactsDir = settings.LocalArtifactsDir;
+
+			S3BucketRegion = settings.S3BucketRegion;
+			S3CredentialType = settings.S3CredentialType;
+			S3ClientKeyId = settings.S3ClientKeyId;
+			S3ClientSecret = settings.S3ClientSecret;
+			S3AssumeArn = settings.S3AssumeArn;
+			S3LogBucketName = settings.S3LogBucketName;
+			S3ArtifactBucketName = settings.S3ArtifactBucketName;
+
+			LogRelayServer = settings.LogRelayServer;
+			LogRelayBearerToken = settings.LogRelayBearerToken;
+			LogJsonToStdOut = settings.LogJsonToStdOut;
+
+			FleetManager = settings.FleetManager;
+			DisableSchedules = settings.DisableSchedules;
+			ScheduleTimeZone = settings.ScheduleTimeZone;
+
+			SlackToken = settings.SlackToken;
+			SlackSocketToken = settings.SlackSocketToken;
+
+			SlackToken = settings.SlackToken;
+			SlackToken = settings.SlackToken;
+			UpdateStreamsNotificationChannel = settings.UpdateStreamsNotificationChannel;
+
+			SmtpServer = settings.SmtpServer;
+			EmailSenderAddress = settings.EmailSenderAddress;
+			EmailSenderName = settings.EmailSenderName;
+
+			UpdateStreamsNotificationChannel = settings.UpdateStreamsNotificationChannel;
+			UpdateStreamsNotificationChannel = settings.UpdateStreamsNotificationChannel;
+			UpdateStreamsNotificationChannel = settings.UpdateStreamsNotificationChannel;
+
+			P4BridgeServiceUsername = settings.P4BridgeServiceUsername;
+			P4BridgeServicePassword = settings.P4BridgeServicePassword;
+			P4BridgeCanImpersonate = settings.P4BridgeCanImpersonate;
+
+			GlobalThreadPoolMinSize = settings.GlobalThreadPoolMinSize;
+			TimeZoneInfo = settings.TimeZoneInfo;
 		}
 	}
 
@@ -393,7 +385,6 @@ namespace Horde.Build.Api
 		/// Delta settings to update
 		/// </summary>
 		public Dictionary<string, object>? Settings { get; set; }
-
 	}
 
 	/// <summary>
@@ -401,7 +392,6 @@ namespace Horde.Build.Api
 	/// </summary>
 	public class ServerUpdateResponse
 	{
-
 		/// <summary>
 		/// List of any error messages when updating
 		/// </summary>
@@ -411,10 +401,8 @@ namespace Horde.Build.Api
 		/// Whether a server restart is required
 		/// </summary>
 		public bool RestartRequired { get; set; } = false;
-
 	}
 
-
 	/// <summary>
 	/// Parameters to update global settings
 	/// </summary>
@@ -444,8 +432,6 @@ namespace Horde.Build.Api
 		/// Default Pool Name
 		/// </summary>
 		public string? DefaultPoolName {get; set;}		
-
 	}
-
 }
 
diff --git a/Engine/Source/Programs/Horde/Horde.Build/Api/Credentials.cs b/Engine/Source/Programs/Horde/Horde.Build/Api/Credentials.cs
--- a/Engine/Source/Programs/Horde/Horde.Build/Api/Credentials.cs
+++ b/Engine/Source/Programs/Horde/Horde.Build/Api/Credentials.cs
@@ -1,11 +1,8 @@
 // Copyright Epic Games, Inc. All Rights Reserved.
 
-using Horde.Build.Models;
-using System;
 using System.Collections.Generic;
 using System.ComponentModel.DataAnnotations;
-using System.Linq;
-using System.Threading.Tasks;
+using Horde.Build.Models;
 
 namespace Horde.Build.Api
 {
@@ -39,10 +36,10 @@ namespace Horde.Build.Api
 		/// <summary>
 		/// Constructor
 		/// </summary>
-		/// <param name="Id">Unique id for the new credential</param>
-		public CreateCredentialResponse(string Id)
+		/// <param name="id">Unique id for the new credential</param>
+		public CreateCredentialResponse(string id)
 		{
-			this.Id = Id;
+			Id = id;
 		}
 	}
 
@@ -95,14 +92,14 @@ namespace Horde.Build.Api
 		/// <summary>
 		/// Constructor
 		/// </summary>
-		/// <param name="Credential">The credential to construct from</param>
+		/// <param name="credential">The credential to construct from</param>
 		/// <param name="bIncludeAcl">Whether to include the ACL in the response</param>
-		public GetCredentialResponse(Credential Credential, bool bIncludeAcl)
+		public GetCredentialResponse(Credential credential, bool bIncludeAcl)
 		{
-			this.Id = Credential.Id.ToString();
-			this.Name = Credential.Name;
-			this.Properties = Credential.Properties;
-			this.Acl = (bIncludeAcl && Credential.Acl != null)? new GetAclResponse(Credential.Acl) : null;
+			Id = credential.Id.ToString();
+			Name = credential.Name;
+			Properties = credential.Properties;
+			Acl = (bIncludeAcl && credential.Acl != null)? new GetAclResponse(credential.Acl) : null;
 		}
 	}
 }
diff --git a/Engine/Source/Programs/Horde/Horde.Build/Api/Devices.cs b/Engine/Source/Programs/Horde/Horde.Build/Api/Devices.cs
--- a/Engine/Source/Programs/Horde/Horde.Build/Api/Devices.cs
+++ b/Engine/Source/Programs/Horde/Horde.Build/Api/Devices.cs
@@ -3,7 +3,6 @@
 using System;
 using System.Collections.Generic;
 using System.ComponentModel.DataAnnotations;
-using System.Diagnostics.CodeAnalysis;
 using System.Linq;
 using Horde.Build.Models;
 
@@ -36,9 +35,9 @@ namespace Horde.Build.Api
 		/// <summary>
 		/// Constructor
 		/// </summary>
-		public CreateDevicePlatformResponse(string Id)
+		public CreateDevicePlatformResponse(string id)
 		{
-			this.Id = Id;
+			Id = id;
 		}
 	}
 
@@ -53,7 +52,6 @@ namespace Horde.Build.Api
 		public string[] ModelIds { get; set; } = null!;
 	}
 
-
 	/// <summary>
 	/// Get object response which describes a device platform
 	/// </summary>
@@ -77,13 +75,12 @@ namespace Horde.Build.Api
 		/// <summary>
 		/// Response constructor
 		/// </summary>
-		public GetDevicePlatformResponse(string Id, string Name, string[] ModelIds)
+		public GetDevicePlatformResponse(string id, string name, string[] modelIds)
 		{
-			this.Id = Id;
-			this.Name = Name;
-			this.ModelIds = ModelIds;
+			Id = id;
+			Name = name;
+			ModelIds = modelIds;
 		}
-
 	}
 
 	/// <summary>
@@ -107,7 +104,6 @@ namespace Horde.Build.Api
 		/// Projects associated with this device pool
 		/// </summary>
 		public List<string>? ProjectIds { get; set; }
-
 	}
 
 	/// <summary>
@@ -124,7 +120,6 @@ namespace Horde.Build.Api
 		/// Projects associated with this device pool
 		/// </summary>
 		public List<string>? ProjectIds { get; set; }
-
 	}
 
 	/// <summary>
@@ -140,9 +135,9 @@ namespace Horde.Build.Api
 		/// <summary>
 		/// Constructor
 		/// </summary>
-		public CreateDevicePoolResponse(string Id)
+		public CreateDevicePoolResponse(string id)
 		{
-			this.Id = Id;
+			Id = id;
 		}
 	}
 
@@ -174,12 +169,12 @@ namespace Horde.Build.Api
 		/// <summary>
 		/// Constructor
 		/// </summary>
-		public GetDevicePoolResponse(string Id, string Name, DevicePoolType PoolType, bool WriteAccess)
+		public GetDevicePoolResponse(string id, string name, DevicePoolType poolType, bool writeAccess)
 		{
-			this.Id = Id;
-			this.Name = Name;
-            this.PoolType = PoolType;
-			this.WriteAccess = WriteAccess;
+			Id = id;
+			Name = name;
+            PoolType = poolType;
+			WriteAccess = writeAccess;
         }
 	}
 
@@ -222,10 +217,8 @@ namespace Horde.Build.Api
 		/// The vendor model id of the device
 		/// </summary>
 		public string? ModelId { get; set; }
-
 	}
 
-
 	/// <summary>
 	/// Device creation response object
 	/// </summary>
@@ -240,10 +233,10 @@ namespace Horde.Build.Api
 		/// <summary>
 		/// Constructor
 		/// </summary>
-		/// <param name="Id"></param>
-		public CreateDeviceResponse(string Id)
+		/// <param name="id"></param>
+		public CreateDeviceResponse(string id)
 		{
-			this.Id = Id;
+			Id = id;
 		}
 	}
 
@@ -275,18 +268,16 @@ namespace Horde.Build.Api
 		/// <summary>
 		/// Constructor
 		/// </summary>
-		/// <param name="Telemetry"></param>
-		public GetDeviceUtilizationResponse(DeviceUtilizationTelemetry Telemetry)
+		/// <param name="telemetry"></param>
+		public GetDeviceUtilizationResponse(DeviceUtilizationTelemetry telemetry)
 		{
-			this.JobId = Telemetry.JobId;
-			this.StepId = Telemetry.StepId;
-			this.ReservationStartUtc = Telemetry.ReservationStartUtc;
-			this.ReservationFinishUtc = Telemetry.ReservationFinishUtc;
+			JobId = telemetry.JobId;
+			StepId = telemetry.StepId;
+			ReservationStartUtc = telemetry.ReservationStartUtc;
+			ReservationFinishUtc = telemetry.ReservationFinishUtc;
 		}
-
 	}
 
-
 	/// <summary>
 	/// Get response object which describes a device
 	/// </summary>
@@ -362,26 +353,25 @@ namespace Horde.Build.Api
 		/// </summary>
 		public List<GetDeviceUtilizationResponse>? Utilization { get; set;}
 
-
-	/// <summary>
-	/// Device response constructor
-	/// </summary>
-	public GetDeviceResponse(string Id, string PlatformId, string PoolId, string Name, bool Enabled, string? Address, string? ModelId, string? ModifiedByUser, string? Notes, DateTime? ProblemTime, DateTime? MaintenanceTime, List<DeviceUtilizationTelemetry>? Utilization, string? CheckedOutByUser = null, DateTime? CheckOutTime = null)
+		/// <summary>
+		/// Device response constructor
+		/// </summary>
+		public GetDeviceResponse(string id, string platformId, string poolId, string name, bool enabled, string? address, string? modelId, string? modifiedByUser, string? notes, DateTime? problemTime, DateTime? maintenanceTime, List<DeviceUtilizationTelemetry>? utilization, string? checkedOutByUser = null, DateTime? checkOutTime = null)
 		{
-			this.Id = Id;
-			this.Name = Name;
-			this.PlatformId = PlatformId;
-			this.PoolId = PoolId;
-			this.Enabled = Enabled;
-			this.Address = Address;
-			this.ModelId = ModelId;
-            this.ModifiedByUser = ModifiedByUser;
-            this.Notes = Notes;
-			this.ProblemTime = ProblemTime;
-			this.MaintenanceTime = MaintenanceTime;
-			this.Utilization = Utilization?.Select(U => new GetDeviceUtilizationResponse(U)).ToList();
-            this.CheckedOutByUserId = CheckedOutByUser;
-            this.CheckOutTime = CheckOutTime;
+			Id = id;
+			Name = name;
+			PlatformId = platformId;
+			PoolId = poolId;
+			Enabled = enabled;
+			Address = address;
+			ModelId = modelId;
+            ModifiedByUser = modifiedByUser;
+            Notes = notes;
+			ProblemTime = problemTime;
+			MaintenanceTime = maintenanceTime;
+			Utilization = utilization?.Select(u => new GetDeviceUtilizationResponse(u)).ToList();
+            CheckedOutByUserId = checkedOutByUser;
+            CheckOutTime = checkOutTime;
         }
 	}
 
@@ -442,7 +432,6 @@ namespace Horde.Build.Api
         public bool Checkout { get; set; }
     }
 
-
 	/// <summary>
 	/// Device reservation request object
 	/// </summary>
@@ -463,7 +452,6 @@ namespace Horde.Build.Api
 		/// The optional vendor model ids to exclude for this device
 		/// </summary>
 		public List<string>? ExcludeModels { get; set; }
-
 	}
 
 	/// <summary>
@@ -482,7 +470,6 @@ namespace Horde.Build.Api
 		/// </summary>
 		[Required]
 		public List<DeviceReservationRequest> Devices { get; set; } = null!;
-
 	}
 
 	/// <summary>
@@ -499,7 +486,6 @@ namespace Horde.Build.Api
 		/// The devices that were reserved
 		/// </summary>
 		public List<GetDeviceResponse> Devices { get; set; } = null!;
-
 	}
 
 	/// <summary>
@@ -561,7 +547,6 @@ namespace Horde.Build.Api
 		public string LegacyGuid { get; set; } = null!;
 	}
 
-
 	// Legacy clients
 
 	/// <summary>
@@ -654,12 +639,10 @@ namespace Horde.Build.Api
 		/// </summary>
 		public string? StepId { get; set; } = null;
 
-
 		/// <summary>
 		/// The legacy guid of the reservation
 		/// </summary>
 		[Required]
-		[SuppressMessage("Compiler", "CA1720: Identifier 'Guid' contains type name")]
 		public string Guid { get; set; } = null!;
 	}
 
@@ -684,7 +667,7 @@ namespace Horde.Build.Api
 		/// The IP or hostname of device
 		/// </summary>
 		[Required]
-		public string IPOrHostName { get; set; } = null!;
+		public string IpOrHostName { get; set; } = null!;
 
 		/// <summary>
 		/// The (legacy) perf spec of the device
@@ -739,8 +722,5 @@ namespace Horde.Build.Api
 		/// </summary>
 		[Required]
 		public Dictionary<string, string> PerfSpecHighMap { get; set; } = null!;
-
 	}
-
-
 }
diff --git a/Engine/Source/Programs/Horde/Horde.Build/Api/Events.cs b/Engine/Source/Programs/Horde/Horde.Build/Api/Events.cs
--- a/Engine/Source/Programs/Horde/Horde.Build/Api/Events.cs
+++ b/Engine/Source/Programs/Horde/Horde.Build/Api/Events.cs
@@ -1,15 +1,11 @@
 // Copyright Epic Games, Inc. All Rights Reserved.
 
+using System;
+using System.ComponentModel.DataAnnotations;
 using EpicGames.Core;
-using HordeCommon;
 using Horde.Build.Models;
 using Horde.Build.Utilities;
-using Microsoft.AspNetCore.Mvc;
-using System;
-using System.Collections.Generic;
-using System.ComponentModel.DataAnnotations;
-using System.Linq;
-using System.Threading.Tasks;
+using HordeCommon;
 
 namespace Horde.Build.Api
 {
@@ -59,14 +55,14 @@ namespace Horde.Build.Api
 		/// <summary>
 		/// Constructor
 		/// </summary>
-		/// <param name="StreamId">The stream id</param>
-		/// <param name="TemplateId">The template id</param>
-		/// <param name="Outcome">Outcome of the job</param>
-		public JobCompleteEventRecord(StreamId StreamId, TemplateRefId TemplateId, LabelOutcome Outcome)
+		/// <param name="streamId">The stream id</param>
+		/// <param name="templateId">The template id</param>
+		/// <param name="outcome">Outcome of the job</param>
+		public JobCompleteEventRecord(StreamId streamId, TemplateRefId templateId, LabelOutcome outcome)
 		{
-			this.StreamId = StreamId.ToString();
-			this.TemplateId = TemplateId.ToString();
-			this.Outcome = Outcome;
+			StreamId = streamId.ToString();
+			TemplateId = templateId.ToString();
+			Outcome = outcome;
 		}
 	}
 
@@ -109,26 +105,26 @@ namespace Horde.Build.Api
 		/// </summary>
 		public LabelCompleteEventRecord()
 		{
-			this.StreamId = String.Empty;
-			this.TemplateId = String.Empty;
-			this.LabelName = String.Empty;
+			StreamId = String.Empty;
+			TemplateId = String.Empty;
+			LabelName = String.Empty;
 		}
 
 		/// <summary>
 		/// Constructor
 		/// </summary>
-		/// <param name="StreamId">The stream id</param>
-		/// <param name="TemplateId">The template id</param>
-		/// <param name="CategoryName">Name of the category</param>
-		/// <param name="LabelName">The label name</param>
-		/// <param name="Outcome">Outcome of the label</param>
-		public LabelCompleteEventRecord(StreamId StreamId, TemplateRefId TemplateId, string? CategoryName, string LabelName, LabelOutcome Outcome)
+		/// <param name="streamId">The stream id</param>
+		/// <param name="templateId">The template id</param>
+		/// <param name="categoryName">Name of the category</param>
+		/// <param name="labelName">The label name</param>
+		/// <param name="outcome">Outcome of the label</param>
+		public LabelCompleteEventRecord(StreamId streamId, TemplateRefId templateId, string? categoryName, string labelName, LabelOutcome outcome)
 		{
-			this.StreamId = StreamId.ToString();
-			this.TemplateId = TemplateId.ToString();
-			this.CategoryName = CategoryName;
-			this.LabelName = LabelName;
-			this.Outcome = Outcome;
+			StreamId = streamId.ToString();
+			TemplateId = templateId.ToString();
+			CategoryName = categoryName;
+			LabelName = labelName;
+			Outcome = outcome;
 		}
 	}
 
@@ -174,16 +170,16 @@ namespace Horde.Build.Api
 		/// <summary>
 		/// Constructor
 		/// </summary>
-		/// <param name="StreamId">The stream id</param>
-		/// <param name="TemplateId">The template id</param>
-		/// <param name="StepName">The label name</param>
-		/// <param name="Outcome">Outcome of the step</param>
-		public StepCompleteEventRecord(StreamId StreamId, TemplateRefId TemplateId, string StepName, JobStepOutcome Outcome)
+		/// <param name="streamId">The stream id</param>
+		/// <param name="templateId">The template id</param>
+		/// <param name="stepName">The label name</param>
+		/// <param name="outcome">Outcome of the step</param>
+		public StepCompleteEventRecord(StreamId streamId, TemplateRefId templateId, string stepName, JobStepOutcome outcome)
 		{
-			this.StreamId = StreamId.ToString();
-			this.TemplateId = TemplateId.ToString();
-			this.StepName = StepName;
-			this.Outcome = Outcome;
+			StreamId = streamId.ToString();
+			TemplateId = templateId.ToString();
+			StepName = stepName;
+			Outcome = outcome;
 		}
 	}
 }
diff --git a/Engine/Source/Programs/Horde/Horde.Build/Api/HelixSwarm.cs b/Engine/Source/Programs/Horde/Horde.Build/Api/HelixSwarm.cs
--- a/Engine/Source/Programs/Horde/Horde.Build/Api/HelixSwarm.cs
+++ b/Engine/Source/Programs/Horde/Horde.Build/Api/HelixSwarm.cs
@@ -85,10 +85,10 @@ namespace Horde.Build.Api
 		/// <summary>
 		/// Constructor
 		/// </summary>
-		/// <param name="Id">Id of the new document</param>
-		public SwarmAutoTestResponse(string Id)
+		/// <param name="id">Id of the new document</param>
+		public SwarmAutoTestResponse(string id)
 		{
-			this.Id = Id;
+			Id = id;
 		}
 	}
 }
