Horde: Fix coding conventions within Horde.Build.

diff --git a/Engine/Source/Programs/Horde/Horde.Build/Api/Issues.cs b/Engine/Source/Programs/Horde/Horde.Build/Api/Issues.cs
--- a/Engine/Source/Programs/Horde/Horde.Build/Api/Issues.cs
+++ b/Engine/Source/Programs/Horde/Horde.Build/Api/Issues.cs
@@ -1,16 +1,12 @@
 // Copyright Epic Games, Inc. All Rights Reserved.
 
+using System;
+using System.Collections.Generic;
+using System.Linq;
 using EpicGames.Core;
 using Horde.Build.Models;
-using Horde.Build.Services;
 using Horde.Build.Utilities;
 using MongoDB.Bson;
-using System;
-using System.Collections.Generic;
-using System.Diagnostics.CodeAnalysis;
-using System.Linq;
-using System.Text.RegularExpressions;
-using System.Threading.Tasks;
 
 namespace Horde.Build.Api
 {
@@ -65,17 +61,17 @@ namespace Horde.Build.Api
 		/// <summary>
 		/// Constructor
 		/// </summary>
-		/// <param name="IssueStep">The issue step to construct from</param>
-		public GetIssueStepResponse(IIssueStep IssueStep)
+		/// <param name="issueStep">The issue step to construct from</param>
+		public GetIssueStepResponse(IIssueStep issueStep)
 		{
-			this.Change = IssueStep.Change;
-			this.Severity = IssueStep.Severity;
-			this.JobName = IssueStep.JobName;
-			this.JobId = IssueStep.JobId.ToString();
-			this.BatchId = IssueStep.BatchId.ToString();
-			this.StepId = IssueStep.StepId.ToString();
-			this.StepTime = IssueStep.StepTime;
-			this.LogId = IssueStep.LogId?.ToString();
+			Change = issueStep.Change;
+			Severity = issueStep.Severity;
+			JobName = issueStep.JobName;
+			JobId = issueStep.JobId.ToString();
+			BatchId = issueStep.BatchId.ToString();
+			StepId = issueStep.StepId.ToString();
+			StepTime = issueStep.StepTime;
+			LogId = issueStep.LogId?.ToString();
 		}
 	}
 
@@ -117,16 +113,16 @@ namespace Horde.Build.Api
 		/// <summary>
 		/// Constructor
 		/// </summary>
-		/// <param name="Span">The node to construct from</param>
-		/// <param name="Steps">Failing steps for this span</param>
-		public GetIssueSpanResponse(IIssueSpan Span, List<IIssueStep> Steps)
+		/// <param name="span">The node to construct from</param>
+		/// <param name="steps">Failing steps for this span</param>
+		public GetIssueSpanResponse(IIssueSpan span, List<IIssueStep> steps)
 		{
-			this.Id = Span.Id.ToString();
-			this.Name = Span.NodeName;
-			this.TemplateId = Span.TemplateRefId.ToString();
-			this.LastSuccess = (Span.LastSuccess != null) ? new GetIssueStepResponse(Span.LastSuccess) : null;
-			this.Steps = Steps.ConvertAll(x => new GetIssueStepResponse(x));
-			this.NextSuccess = (Span.NextSuccess != null) ? new GetIssueStepResponse(Span.NextSuccess) : null;
+			Id = span.Id.ToString();
+			Name = span.NodeName;
+			TemplateId = span.TemplateRefId.ToString();
+			LastSuccess = (span.LastSuccess != null) ? new GetIssueStepResponse(span.LastSuccess) : null;
+			Steps = steps.ConvertAll(x => new GetIssueStepResponse(x));
+			NextSuccess = (span.NextSuccess != null) ? new GetIssueStepResponse(span.NextSuccess) : null;
 		}
 	}
 
@@ -158,26 +154,26 @@ namespace Horde.Build.Api
 		/// <summary>
 		/// Constructor
 		/// </summary>
-		/// <param name="StreamId">The stream to construct from</param>
-		/// <param name="Spans">List of spans for the given stream</param>
-		/// <param name="Steps">List of steps for the given stream</param>
-		public GetIssueStreamResponse(StreamId StreamId, List<IIssueSpan> Spans, List<IIssueStep> Steps)
+		/// <param name="streamId">The stream to construct from</param>
+		/// <param name="spans">List of spans for the given stream</param>
+		/// <param name="steps">List of steps for the given stream</param>
+		public GetIssueStreamResponse(StreamId streamId, List<IIssueSpan> spans, List<IIssueStep> steps)
 		{
-			this.StreamId = StreamId.ToString();
+			StreamId = streamId.ToString();
 
-			foreach (IIssueSpan Span in Spans)
+			foreach (IIssueSpan span in spans)
 			{
-				if (Span.LastSuccess != null && (MinChange == null || Span.LastSuccess.Change < MinChange.Value))
+				if (span.LastSuccess != null && (MinChange == null || span.LastSuccess.Change < MinChange.Value))
 				{
-					MinChange = Span.LastSuccess.Change;
+					MinChange = span.LastSuccess.Change;
 				}
-				if (Span.NextSuccess != null && (MaxChange == null || Span.NextSuccess.Change > MaxChange.Value))
+				if (span.NextSuccess != null && (MaxChange == null || span.NextSuccess.Change > MaxChange.Value))
 				{
-					MaxChange = Span.NextSuccess.Change;
+					MaxChange = span.NextSuccess.Change;
 				}
 			}
 
-			this.Nodes = Spans.ConvertAll(x => new GetIssueSpanResponse(x, Steps.Where(y => y.SpanId == x.Id).ToList()));
+			Nodes = spans.ConvertAll(x => new GetIssueSpanResponse(x, steps.Where(y => y.SpanId == x.Id).ToList()));
 		}
 	}
 
@@ -230,14 +226,14 @@ namespace Horde.Build.Api
 		/// <summary>
 		/// Constructor
 		/// </summary>
-		/// <param name="BuildId">The corresponding build id</param>
-		/// <param name="Message">Message for the diagnostic</param>
-		/// <param name="Url">Link to the diagnostic</param>
-		public GetIssueDiagnosticResponse(long? BuildId, string Message, Uri Url)
+		/// <param name="buildId">The corresponding build id</param>
+		/// <param name="message">Message for the diagnostic</param>
+		/// <param name="url">Link to the diagnostic</param>
+		public GetIssueDiagnosticResponse(long? buildId, string message, Uri url)
 		{
-			this.BuildId = BuildId;
-			this.Message = Message;
-			this.Url = Url;
+			BuildId = buildId;
+			Message = message;
+			Url = url;
 		}
 	}
 
@@ -269,16 +265,16 @@ namespace Horde.Build.Api
 		/// <summary>
 		/// Constructor
 		/// </summary>
-		/// <param name="TemplateId"></param>
-		/// <param name="TemplateName"></param>
-		/// <param name="Resolved"></param>
-		/// <param name="Severity"></param>
-		public GetIssueAffectedTemplateResponse(string TemplateId, string TemplateName, bool Resolved, IssueSeverity Severity = IssueSeverity.Unspecified)
+		/// <param name="templateId"></param>
+		/// <param name="templateName"></param>
+		/// <param name="resolved"></param>
+		/// <param name="severity"></param>
+		public GetIssueAffectedTemplateResponse(string templateId, string templateName, bool resolved, IssueSeverity severity = IssueSeverity.Unspecified)
 		{
-			this.TemplateId = TemplateId;
-			this.TemplateName = TemplateName;
-			this.Resolved = Resolved;
-			this.Severity = Severity;
+			TemplateId = templateId;
+			TemplateName = templateName;
+			Resolved = resolved;
+			Severity = severity;
 		}
 	}
 
@@ -325,38 +321,38 @@ namespace Horde.Build.Api
 		/// <summary>
 		/// Constructor
 		/// </summary>
-		/// <param name="Details">Issue to construct from</param>
-		/// <param name="Stream"></param>
-		/// <param name="Spans">The spans to construct from</param>
-		public GetIssueAffectedStreamResponse(IIssueDetails Details, IStream? Stream, IEnumerable<IIssueSpan> Spans)
+		/// <param name="details">Issue to construct from</param>
+		/// <param name="stream"></param>
+		/// <param name="spans">The spans to construct from</param>
+		public GetIssueAffectedStreamResponse(IIssueDetails details, IStream? stream, IEnumerable<IIssueSpan> spans)
 		{
-			IIssueSpan FirstSpan = Spans.First();
-			this.StreamId = FirstSpan.StreamId.ToString();
-			this.StreamName = FirstSpan.StreamName;
-			this.Resolved = Spans.All(x => x.NextSuccess != null);
+			IIssueSpan firstSpan = spans.First();
+			StreamId = firstSpan.StreamId.ToString();
+			StreamName = firstSpan.StreamName;
+			Resolved = spans.All(x => x.NextSuccess != null);
 
-			this.AffectedTemplates = new List<GetIssueAffectedTemplateResponse>();
-			foreach (IGrouping<TemplateRefId, IIssueSpan> Template in Spans.GroupBy(x => x.TemplateRefId))
+			AffectedTemplates = new List<GetIssueAffectedTemplateResponse>();
+			foreach (IGrouping<TemplateRefId, IIssueSpan> template in spans.GroupBy(x => x.TemplateRefId))
 			{
-				string TemplateName = Template.Key.ToString();
-				if (Stream != null && Stream.Templates.TryGetValue(Template.Key, out TemplateRef? TemplateRef))
+				string templateName = template.Key.ToString();
+				if (stream != null && stream.Templates.TryGetValue(template.Key, out TemplateRef? templateRef))
 				{
-					TemplateName = TemplateRef.Name;
+					templateName = templateRef.Name;
 				}
 
-				HashSet<ObjectId> UnresolvedTemplateSpans = new HashSet<ObjectId>(Template.Where(x => x.NextSuccess == null).Select(x => x.Id));
+				HashSet<ObjectId> unresolvedTemplateSpans = new HashSet<ObjectId>(template.Where(x => x.NextSuccess == null).Select(x => x.Id));
 
-				IIssueStep? TemplateStep = Details.Steps.Where(x => UnresolvedTemplateSpans.Contains(x.SpanId)).OrderByDescending(x => x.StepTime).FirstOrDefault();
+				IIssueStep? templateStep = details.Steps.Where(x => unresolvedTemplateSpans.Contains(x.SpanId)).OrderByDescending(x => x.StepTime).FirstOrDefault();
 
-				this.AffectedTemplates.Add(new GetIssueAffectedTemplateResponse(Template.Key.ToString(), TemplateName, Template.All(x => x.NextSuccess != null), TemplateStep?.Severity ?? IssueSeverity.Unspecified));
+				AffectedTemplates.Add(new GetIssueAffectedTemplateResponse(template.Key.ToString(), templateName, template.All(x => x.NextSuccess != null), templateStep?.Severity ?? IssueSeverity.Unspecified));
 			}
 
-			HashSet<TemplateRefId> TemplateIdsSet = new HashSet<TemplateRefId>(Spans.Select(x => x.TemplateRefId));
-			this.TemplateIds = TemplateIdsSet.Select(x => x.ToString()).ToList();
+			HashSet<TemplateRefId> templateIdsSet = new HashSet<TemplateRefId>(spans.Select(x => x.TemplateRefId));
+			TemplateIds = templateIdsSet.Select(x => x.ToString()).ToList();
 
-			HashSet<TemplateRefId> UnresolvedTemplateIdsSet = new HashSet<TemplateRefId>(Spans.Where(x => x.NextSuccess == null).Select(x => x.TemplateRefId));
-			this.UnresolvedTemplateIds = UnresolvedTemplateIdsSet.Select(x => x.ToString()).ToList();
-			this.ResolvedTemplateIds = TemplateIdsSet.Except(UnresolvedTemplateIdsSet).Select(x => x.ToString()).ToList();
+			HashSet<TemplateRefId> unresolvedTemplateIdsSet = new HashSet<TemplateRefId>(spans.Where(x => x.NextSuccess == null).Select(x => x.TemplateRefId));
+			UnresolvedTemplateIds = unresolvedTemplateIdsSet.Select(x => x.ToString()).ToList();
+			ResolvedTemplateIds = templateIdsSet.Except(unresolvedTemplateIdsSet).Select(x => x.ToString()).ToList();
 		}
 	}
 
@@ -513,60 +509,60 @@ namespace Horde.Build.Api
 		/// <summary>
 		/// Constructs a new issue
 		/// </summary>
-		/// <param name="Details">Issue to construct from</param>
-		/// <param name="AffectedStreams">The affected streams</param>
-		/// <param name="ShowDesktopAlerts">Whether to show alerts for this issue</param>
-		public GetIssueResponse(IIssueDetails Details, List<GetIssueAffectedStreamResponse> AffectedStreams, bool ShowDesktopAlerts)
+		/// <param name="details">Issue to construct from</param>
+		/// <param name="affectedStreams">The affected streams</param>
+		/// <param name="showDesktopAlerts">Whether to show alerts for this issue</param>
+		public GetIssueResponse(IIssueDetails details, List<GetIssueAffectedStreamResponse> affectedStreams, bool showDesktopAlerts)
 		{
-			IIssue Issue = Details.Issue;
-			this.Id = Issue.Id;
-			this.CreatedAt = Issue.CreatedAt;
-			this.RetrievedAt = DateTime.UtcNow;
-			this.Summary = String.IsNullOrEmpty(Issue.UserSummary)? Issue.Summary : Issue.UserSummary;
-			this.Description = Issue.Description;
-			this.Severity = Issue.Severity;
-			this.Promoted = Issue.Promoted;
-			this.Owner = Details.Owner?.Login;
-			this.OwnerId = (Details.Owner == null)? null : Details.Owner.Id.ToString();
-			if(Details.Owner != null)
+			IIssue issue = details.Issue;
+			Id = issue.Id;
+			CreatedAt = issue.CreatedAt;
+			RetrievedAt = DateTime.UtcNow;
+			Summary = String.IsNullOrEmpty(issue.UserSummary)? issue.Summary : issue.UserSummary;
+			Description = issue.Description;
+			Severity = issue.Severity;
+			Promoted = issue.Promoted;
+			Owner = details.Owner?.Login;
+			OwnerId = details.Owner?.Id.ToString();
+			if(details.Owner != null)
 			{
-				this.OwnerInfo = new GetThinUserInfoResponse(Details.Owner);
+				OwnerInfo = new GetThinUserInfoResponse(details.Owner);
 			}
-			this.NominatedBy = Details.NominatedBy?.Login;
-			if (Details.NominatedBy != null)
+			NominatedBy = details.NominatedBy?.Login;
+			if (details.NominatedBy != null)
 			{
-				this.NominatedByInfo = new GetThinUserInfoResponse(Details.NominatedBy);
+				NominatedByInfo = new GetThinUserInfoResponse(details.NominatedBy);
 			}
-			this.AcknowledgedAt = Issue.AcknowledgedAt;
-			this.FixChange = Issue.FixChange;
-			this.ResolvedAt = Issue.ResolvedAt;
-			this.ResolvedBy = Details.ResolvedBy?.Login;
-			this.ResolvedById = (Details.ResolvedBy == null) ? null : Details.ResolvedBy.Id.ToString();
-			if (Details.ResolvedBy != null)
+			AcknowledgedAt = issue.AcknowledgedAt;
+			FixChange = issue.FixChange;
+			ResolvedAt = issue.ResolvedAt;
+			ResolvedBy = details.ResolvedBy?.Login;
+			ResolvedById = details.ResolvedBy?.Id.ToString();
+			if (details.ResolvedBy != null)
 			{
-				this.ResolvedByInfo = new GetThinUserInfoResponse(Details.ResolvedBy);
+				ResolvedByInfo = new GetThinUserInfoResponse(details.ResolvedBy);
 			}
-			this.VerifiedAt = Issue.VerifiedAt;
-			this.LastSeenAt = Issue.LastSeenAt;
-			this.Streams = Details.Spans.Select(x => x.StreamName).Distinct().ToList()!;
-			this.ResolvedStreams = new List<string>();
-			this.UnresolvedStreams = new List<string>();
-			this.AffectedStreams = AffectedStreams;
-			foreach (IGrouping<StreamId, IIssueSpan> Stream in Details.Spans.GroupBy(x => x.StreamId))
+			VerifiedAt = issue.VerifiedAt;
+			LastSeenAt = issue.LastSeenAt;
+			Streams = details.Spans.Select(x => x.StreamName).Distinct().ToList()!;
+			ResolvedStreams = new List<string>();
+			UnresolvedStreams = new List<string>();
+			AffectedStreams = affectedStreams;
+			foreach (IGrouping<StreamId, IIssueSpan> stream in details.Spans.GroupBy(x => x.StreamId))
 			{
-				if (Stream.All(x => x.NextSuccess != null))
+				if (stream.All(x => x.NextSuccess != null))
 				{
-					this.ResolvedStreams.Add(Stream.Key.ToString());
+					ResolvedStreams.Add(stream.Key.ToString());
 				}
 				else
 				{
-					this.UnresolvedStreams.Add(Stream.Key.ToString());
+					UnresolvedStreams.Add(stream.Key.ToString());
 				}
 			}
-			this.PrimarySuspects = Details.SuspectUsers.Where(x => x.Login != null).Select(x => x.Login).ToList();
-			this.PrimarySuspectIds= Details.SuspectUsers.Select(x => x.Id.ToString()).ToList();
-			this.PrimarySuspectsInfo = Details.SuspectUsers.ConvertAll(x => new GetThinUserInfoResponse(x));
-			this.ShowDesktopAlerts = ShowDesktopAlerts;
+			PrimarySuspects = details.SuspectUsers.Where(x => x.Login != null).Select(x => x.Login).ToList();
+			PrimarySuspectIds= details.SuspectUsers.Select(x => x.Id.ToString()).ToList();
+			PrimarySuspectsInfo = details.SuspectUsers.ConvertAll(x => new GetThinUserInfoResponse(x));
+			ShowDesktopAlerts = showDesktopAlerts;
 		}
 	}
 
@@ -603,19 +599,19 @@ namespace Horde.Build.Api
 		/// <summary>
 		/// Constructor
 		/// </summary>
-		/// <param name="Span"></param>
-		public FindIssueSpanResponse(IIssueSpan Span)
+		/// <param name="span"></param>
+		public FindIssueSpanResponse(IIssueSpan span)
 		{
-			this.Id = Span.Id.ToString();
-			this.TemplateId = Span.TemplateRefId.ToString();
-			this.Name = Span.NodeName;
-			if (Span.LastSuccess != null)
+			Id = span.Id.ToString();
+			TemplateId = span.TemplateRefId.ToString();
+			Name = span.NodeName;
+			if (span.LastSuccess != null)
 			{
-				this.LastSuccess = new GetIssueStepResponse(Span.LastSuccess);
+				LastSuccess = new GetIssueStepResponse(span.LastSuccess);
 			}
-			if (Span.NextSuccess != null)
+			if (span.NextSuccess != null)
 			{
-				this.NextSuccess = new GetIssueStepResponse(Span.NextSuccess);
+				NextSuccess = new GetIssueStepResponse(span.NextSuccess);
 			}
 		}
 	}
@@ -718,40 +714,40 @@ namespace Horde.Build.Api
 		/// <summary>
 		/// Constructs a new issue
 		/// </summary>
-		/// <param name="Issue">The isseu information</param>
-		/// <param name="Owner">Owner of the issue</param>
-		/// <param name="NominatedBy">User that nominated the current fixer</param>
-		/// <param name="ResolvedBy">User that resolved the issue</param>
-		/// <param name="StreamSeverity">The current severity in the stream</param>
-		/// <param name="Spans">Spans for this issue</param>
-		public FindIssueResponse(IIssue Issue, IUser? Owner, IUser? NominatedBy, IUser? ResolvedBy, IssueSeverity? StreamSeverity, List<FindIssueSpanResponse> Spans)
+		/// <param name="issue">The isseu information</param>
+		/// <param name="owner">Owner of the issue</param>
+		/// <param name="nominatedBy">User that nominated the current fixer</param>
+		/// <param name="resolvedBy">User that resolved the issue</param>
+		/// <param name="streamSeverity">The current severity in the stream</param>
+		/// <param name="spans">Spans for this issue</param>
+		public FindIssueResponse(IIssue issue, IUser? owner, IUser? nominatedBy, IUser? resolvedBy, IssueSeverity? streamSeverity, List<FindIssueSpanResponse> spans)
 		{
-			this.Id = Issue.Id;
-			this.CreatedAt = Issue.CreatedAt;
-			this.RetrievedAt = DateTime.UtcNow;
-			this.Summary = String.IsNullOrEmpty(Issue.UserSummary) ? Issue.Summary : Issue.UserSummary;
-			this.Description = Issue.Description;
-			this.Severity = Issue.Severity;
-			this.StreamSeverity = StreamSeverity;
-			this.Promoted = Issue.Promoted;
-			if (Owner != null)
+			Id = issue.Id;
+			CreatedAt = issue.CreatedAt;
+			RetrievedAt = DateTime.UtcNow;
+			Summary = String.IsNullOrEmpty(issue.UserSummary) ? issue.Summary : issue.UserSummary;
+			Description = issue.Description;
+			Severity = issue.Severity;
+			StreamSeverity = streamSeverity;
+			Promoted = issue.Promoted;
+			if (owner != null)
 			{
-				this.Owner = new GetThinUserInfoResponse(Owner);
+				Owner = new GetThinUserInfoResponse(owner);
 			}
-			if (NominatedBy != null)
+			if (nominatedBy != null)
 			{
-				this.NominatedBy = new GetThinUserInfoResponse(NominatedBy);
+				NominatedBy = new GetThinUserInfoResponse(nominatedBy);
 			}
-			this.AcknowledgedAt = Issue.AcknowledgedAt;
-			this.FixChange = Issue.FixChange;
-			this.ResolvedAt = Issue.ResolvedAt;
-			if (ResolvedBy != null)
+			AcknowledgedAt = issue.AcknowledgedAt;
+			FixChange = issue.FixChange;
+			ResolvedAt = issue.ResolvedAt;
+			if (resolvedBy != null)
 			{
-				this.ResolvedBy = new GetThinUserInfoResponse(ResolvedBy);
+				ResolvedBy = new GetThinUserInfoResponse(resolvedBy);
 			}
-			this.VerifiedAt = Issue.VerifiedAt;
-			this.LastSeenAt = Issue.LastSeenAt;
-			this.Spans = Spans;
+			VerifiedAt = issue.VerifiedAt;
+			LastSeenAt = issue.LastSeenAt;
+			Spans = spans;
 		}
 	}
 
diff --git a/Engine/Source/Programs/Horde/Horde.Build/Api/Jira.cs b/Engine/Source/Programs/Horde/Horde.Build/Api/Jira.cs
--- a/Engine/Source/Programs/Horde/Horde.Build/Api/Jira.cs
+++ b/Engine/Source/Programs/Horde/Horde.Build/Api/Jira.cs
@@ -1,11 +1,5 @@
 // Copyright Epic Games, Inc. All Rights Reserved.
 
-using System;
-using System.Collections.Generic;
-using System.ComponentModel.DataAnnotations;
-using System.Diagnostics.CodeAnalysis;
-using System.Linq;
-using Horde.Build.Models;
 using Horde.Build.Services;
 
 namespace Horde.Build.Api
@@ -59,20 +53,18 @@ namespace Horde.Build.Api
 		/// <summary>
 		/// Response constructor
 		/// </summary>
-		public GetJiraIssueResponse(JiraIssue Issue)
+		public GetJiraIssueResponse(JiraIssue issue)
 		{
-			this.Key = Issue.Key;
-			this.JiraLink = Issue.JiraLink;
-			this.StatusName = Issue.StatusName;
-			this.ResolutionName = Issue.ResolutionName;
-			this.PriorityName = Issue.PriorityName;
-			this.AssigneeName = Issue.AssigneeName;
-			this.AssigneeDisplayName = Issue.AssigneeDisplayName;
-			this.AssigneeEmailAddress = Issue.AssigneeEmailAddress;
+			Key = issue.Key;
+			JiraLink = issue.JiraLink;
+			StatusName = issue.StatusName;
+			ResolutionName = issue.ResolutionName;
+			PriorityName = issue.PriorityName;
+			AssigneeName = issue.AssigneeName;
+			AssigneeDisplayName = issue.AssigneeDisplayName;
+			AssigneeEmailAddress = issue.AssigneeEmailAddress;
 		}
-
 	}
-
 }
 
  
\ No newline at end of file
diff --git a/Engine/Source/Programs/Horde/Horde.Build/Api/Jobs.cs b/Engine/Source/Programs/Horde/Horde.Build/Api/Jobs.cs
--- a/Engine/Source/Programs/Horde/Horde.Build/Api/Jobs.cs
+++ b/Engine/Source/Programs/Horde/Horde.Build/Api/Jobs.cs
@@ -1,16 +1,13 @@
 // Copyright Epic Games, Inc. All Rights Reserved.
 
-using EpicGames.Core;
-using HordeCommon;
-using Horde.Build.Models;
-using Horde.Build.Utilities;
 using System;
 using System.Collections.Generic;
 using System.ComponentModel.DataAnnotations;
 using System.Linq;
-using System.Security.Claims;
 using System.Text.Json.Serialization;
-using System.Threading.Tasks;
+using EpicGames.Core;
+using Horde.Build.Models;
+using HordeCommon;
 
 namespace Horde.Build.Api
 {
@@ -95,10 +92,10 @@ namespace Horde.Build.Api
 		/// <summary>
 		/// Private constructor for serialization
 		/// </summary>
-		public CreateJobRequest(string StreamId, string TemplateId)
+		public CreateJobRequest(string streamId, string templateId)
 		{
-			this.StreamId = StreamId;
-			this.TemplateId = TemplateId;
+			StreamId = streamId;
+			TemplateId = templateId;
 		}
 	}
 
@@ -115,10 +112,10 @@ namespace Horde.Build.Api
 		/// <summary>
 		/// Constructor
 		/// </summary>
-		/// <param name="Id">Unique id for the new job</param>
-		public CreateJobResponse(string Id)
+		/// <param name="id">Unique id for the new job</param>
+		public CreateJobResponse(string id)
 		{
-			this.Id = Id;
+			Id = id;
 		}
 	}
 
@@ -181,12 +178,12 @@ namespace Horde.Build.Api
 		/// <summary>
 		/// Constructor
 		/// </summary>
-		/// <param name="Report"></param>
-		public GetReportResponse(IReport Report)
+		/// <param name="report"></param>
+		public GetReportResponse(IReport report)
 		{
-			Name = Report.Name;
-			Placement = Report.Placement;
-			ArtifactId = Report.ArtifactId.ToString();
+			Name = report.Name;
+			Placement = report.Placement;
+			ArtifactId = report.ArtifactId.ToString();
 		}
 	}
 
@@ -343,38 +340,38 @@ namespace Horde.Build.Api
 		/// <summary>
 		/// Constructor
 		/// </summary>
-		/// <param name="Job">Job to create a response for</param>
-		/// <param name="StartedByUserInfo">User that started this job</param>
-		/// <param name="AbortedByUserInfo">User that aborted this job</param>
-		/// <param name="AclResponse">The ACL response object</param>
-		public GetJobResponse(IJob Job, GetThinUserInfoResponse? StartedByUserInfo, GetThinUserInfoResponse? AbortedByUserInfo, GetAclResponse? AclResponse)
+		/// <param name="job">Job to create a response for</param>
+		/// <param name="startedByUserInfo">User that started this job</param>
+		/// <param name="abortedByUserInfo">User that aborted this job</param>
+		/// <param name="aclResponse">The ACL response object</param>
+		public GetJobResponse(IJob job, GetThinUserInfoResponse? startedByUserInfo, GetThinUserInfoResponse? abortedByUserInfo, GetAclResponse? aclResponse)
 		{
-			this.Id = Job.Id.ToString();
-			this.StreamId = Job.StreamId.ToString();
-			this.Name = Job.Name;
-			this.Change = Job.Change;
-			this.CodeChange = (Job.CodeChange != 0) ? (int?)Job.CodeChange : null;
-			this.PreflightChange = (Job.PreflightChange != 0) ? (int?)Job.PreflightChange : null;
-			this.ClonedPreflightChange = (Job.ClonedPreflightChange != 0) ? (int?)Job.ClonedPreflightChange : null;
-			this.TemplateId = Job.TemplateId.ToString();
-			this.TemplateHash = Job.TemplateHash?.ToString() ?? String.Empty;
-			this.GraphHash = Job.GraphHash.ToString();
-			this.StartedByUserId = Job.StartedByUserId?.ToString();
-			this.StartedByUser = StartedByUserInfo?.Login;
-			this.StartedByUserInfo = StartedByUserInfo;
-			this.AbortedByUser = AbortedByUserInfo?.Login;
-			this.AbortedByUserInfo = AbortedByUserInfo;
-			this.CreateTime = new DateTimeOffset(Job.CreateTimeUtc);
-			this.State = Job.GetState();
-			this.Priority = Job.Priority;
-			this.AutoSubmit = Job.AutoSubmit;
-			this.AutoSubmitChange = Job.AutoSubmitChange;
-			this.AutoSubmitMessage = Job.AutoSubmitMessage;
-			this.Reports = Job.Reports?.ConvertAll(x => new GetReportResponse(x));
-			this.Arguments = Job.Arguments.ToList();
-			this.UpdateTime = new DateTimeOffset(Job.UpdateTimeUtc);
-			this.UpdateIssues = Job.UpdateIssues;
-			this.Acl = AclResponse;
+			Id = job.Id.ToString();
+			StreamId = job.StreamId.ToString();
+			Name = job.Name;
+			Change = job.Change;
+			CodeChange = (job.CodeChange != 0) ? (int?)job.CodeChange : null;
+			PreflightChange = (job.PreflightChange != 0) ? (int?)job.PreflightChange : null;
+			ClonedPreflightChange = (job.ClonedPreflightChange != 0) ? (int?)job.ClonedPreflightChange : null;
+			TemplateId = job.TemplateId.ToString();
+			TemplateHash = job.TemplateHash?.ToString() ?? String.Empty;
+			GraphHash = job.GraphHash.ToString();
+			StartedByUserId = job.StartedByUserId?.ToString();
+			StartedByUser = startedByUserInfo?.Login;
+			StartedByUserInfo = startedByUserInfo;
+			AbortedByUser = abortedByUserInfo?.Login;
+			AbortedByUserInfo = abortedByUserInfo;
+			CreateTime = new DateTimeOffset(job.CreateTimeUtc);
+			State = job.GetState();
+			Priority = job.Priority;
+			AutoSubmit = job.AutoSubmit;
+			AutoSubmitChange = job.AutoSubmitChange;
+			AutoSubmitMessage = job.AutoSubmitMessage;
+			Reports = job.Reports?.ConvertAll(x => new GetReportResponse(x));
+			Arguments = job.Arguments.ToList();
+			UpdateTime = new DateTimeOffset(job.UpdateTimeUtc);
+			UpdateIssues = job.UpdateIssues;
+			Acl = aclResponse;
 		}
 	}
 
@@ -406,16 +403,16 @@ namespace Horde.Build.Api
 		/// <summary>
 		/// Constructor
 		/// </summary>
-		/// <param name="Job">The job object</param>
-		/// <param name="JobResponse">The job response</param>
-		/// <param name="Steps">Timing info for each steps</param>
-		/// <param name="Labels">Timing info for each label</param>
-		public GetJobTimingResponse(IJob? Job, GetJobResponse? JobResponse, Dictionary<string, GetStepTimingInfoResponse> Steps, List<GetLabelTimingInfoResponse> Labels)
+		/// <param name="job">The job object</param>
+		/// <param name="jobResponse">The job response</param>
+		/// <param name="steps">Timing info for each steps</param>
+		/// <param name="labels">Timing info for each label</param>
+		public GetJobTimingResponse(IJob? job, GetJobResponse? jobResponse, Dictionary<string, GetStepTimingInfoResponse> steps, List<GetLabelTimingInfoResponse> labels)
 		{
-			this.Job = Job;
-			this.JobResponse = JobResponse;
-			this.Steps = Steps;
-			this.Labels = Labels;
+			Job = job;
+			JobResponse = jobResponse;
+			Steps = steps;
+			Labels = labels;
 		}
 	}
 	
@@ -432,10 +429,10 @@ namespace Horde.Build.Api
 		/// <summary>
 		/// Constructor
 		/// </summary>
-		/// <param name="Timings">Timing info for each job</param>
-		public FindJobTimingsResponse(Dictionary<string, GetJobTimingResponse> Timings)
+		/// <param name="timings">Timing info for each job</param>
+		public FindJobTimingsResponse(Dictionary<string, GetJobTimingResponse> timings)
 		{
-			this.Timings = Timings;
+			Timings = timings;
 		}
 	}
 
@@ -579,29 +576,29 @@ namespace Horde.Build.Api
 		/// <summary>
 		/// Constructor
 		/// </summary>
-		/// <param name="Step">The step to construct from</param>
-		/// <param name="AbortedByUserInfo">User that aborted this step</param>
-		/// <param name="RetriedByUserInfo">User that retried this step</param>
-		public GetStepResponse(IJobStep Step, GetThinUserInfoResponse? AbortedByUserInfo, GetThinUserInfoResponse? RetriedByUserInfo)
+		/// <param name="step">The step to construct from</param>
+		/// <param name="abortedByUserInfo">User that aborted this step</param>
+		/// <param name="retriedByUserInfo">User that retried this step</param>
+		public GetStepResponse(IJobStep step, GetThinUserInfoResponse? abortedByUserInfo, GetThinUserInfoResponse? retriedByUserInfo)
 		{
-			this.Id = Step.Id.ToString();
-			this.NodeIdx = Step.NodeIdx;
-			this.State = Step.State;
-			this.Outcome = Step.Outcome;
-			this.AbortRequested = Step.AbortRequested;
-			this.AbortByUser = AbortedByUserInfo?.Login;
-			this.AbortedByUserInfo = AbortedByUserInfo;
-			this.RetryByUser = RetriedByUserInfo?.Login;
-			this.RetriedByUserInfo = RetriedByUserInfo;
-			this.LogId = Step.LogId?.ToString();
-			this.ReadyTime = Step.ReadyTimeUtc;
-			this.StartTime = Step.StartTimeUtc;
-			this.FinishTime = Step.FinishTimeUtc;
-			this.Reports = Step.Reports?.ConvertAll(x => new GetReportResponse(x));
-
-			if (Step.Properties != null && Step.Properties.Count > 0)
+			Id = step.Id.ToString();
+			NodeIdx = step.NodeIdx;
+			State = step.State;
+			Outcome = step.Outcome;
+			AbortRequested = step.AbortRequested;
+			AbortByUser = abortedByUserInfo?.Login;
+			AbortedByUserInfo = abortedByUserInfo;
+			RetryByUser = retriedByUserInfo?.Login;
+			RetriedByUserInfo = retriedByUserInfo;
+			LogId = step.LogId?.ToString();
+			ReadyTime = step.ReadyTimeUtc;
+			StartTime = step.StartTimeUtc;
+			FinishTime = step.FinishTimeUtc;
+			Reports = step.Reports?.ConvertAll(x => new GetReportResponse(x));
+
+			if (step.Properties != null && step.Properties.Count > 0)
 			{
-				this.Properties = Step.Properties;
+				Properties = step.Properties;
 			}
 		}
 	}
@@ -794,25 +791,25 @@ namespace Horde.Build.Api
 		/// <summary>
 		/// Converts this batch into a public response object
 		/// </summary>
-		/// <param name="Batch">The batch to construct from</param>
-		/// <param name="Steps">Steps in this batch</param>
-		/// <param name="AgentRate">Rate for this agent</param>
+		/// <param name="batch">The batch to construct from</param>
+		/// <param name="steps">Steps in this batch</param>
+		/// <param name="agentRate">Rate for this agent</param>
 		/// <returns>Response instance</returns>
-		public GetBatchResponse(IJobStepBatch Batch, List<GetStepResponse> Steps, double? AgentRate)
+		public GetBatchResponse(IJobStepBatch batch, List<GetStepResponse> steps, double? agentRate)
 		{
-			this.Id = Batch.Id.ToString();
-			this.LogId = Batch.LogId?.ToString();
-			this.GroupIdx = Batch.GroupIdx;
-			this.State = Batch.State;
-			this.Error = Batch.Error;
-			this.Steps = Steps;
-			this.AgentId = Batch.AgentId?.ToString();
-			this.AgentRate = AgentRate;
-			this.SessionId = Batch.SessionId?.ToString();
-			this.LeaseId = Batch.LeaseId?.ToString();
-			this.WeightedPriority = Batch.SchedulePriority;
-			this.StartTime = Batch.StartTimeUtc;
-			this.FinishTime = Batch.FinishTimeUtc;
+			Id = batch.Id.ToString();
+			LogId = batch.LogId?.ToString();
+			GroupIdx = batch.GroupIdx;
+			State = batch.State;
+			Error = batch.Error;
+			Steps = steps;
+			AgentId = batch.AgentId?.ToString();
+			AgentRate = agentRate;
+			SessionId = batch.SessionId?.ToString();
+			LeaseId = batch.LeaseId?.ToString();
+			WeightedPriority = batch.SchedulePriority;
+			StartTime = batch.StartTimeUtc;
+			FinishTime = batch.FinishTimeUtc;
 		}
 	}
 
@@ -881,12 +878,12 @@ namespace Horde.Build.Api
 		/// <summary>
 		/// Constructor
 		/// </summary>
-		/// <param name="State">State of the label</param>
-		/// <param name="Outcome">Outcome of the label</param>
-		public GetLabelStateResponse(LabelState State, LabelOutcome Outcome)
+		/// <param name="state">State of the label</param>
+		/// <param name="outcome">Outcome of the label</param>
+		public GetLabelStateResponse(LabelState state, LabelOutcome outcome)
 		{
-			this.State = State;
-			this.Outcome = Outcome;
+			State = state;
+			Outcome = outcome;
 		}
 	}
 
@@ -903,13 +900,13 @@ namespace Horde.Build.Api
 		/// <summary>
 		/// Constructor
 		/// </summary>
-		/// <param name="State">State of the label</param>
-		/// <param name="Outcome">Outcome of the label</param>
-		/// <param name="Nodes">List of nodes that are covered by the default label</param>
-		public GetDefaultLabelStateResponse(LabelState State, LabelOutcome Outcome, List<string> Nodes)
-			: base(State, Outcome)
+		/// <param name="state">State of the label</param>
+		/// <param name="outcome">Outcome of the label</param>
+		/// <param name="nodes">List of nodes that are covered by the default label</param>
+		public GetDefaultLabelStateResponse(LabelState state, LabelOutcome outcome, List<string> nodes)
+			: base(state, outcome)
 		{
-			this.Nodes = Nodes;
+			Nodes = nodes;
 		}
 	}
 
@@ -951,16 +948,16 @@ namespace Horde.Build.Api
 		/// <summary>
 		/// Constructor
 		/// </summary>
-		/// <param name="TimingInfo">Timing info to construct from</param>
-		public GetTimingInfoResponse(TimingInfo TimingInfo)
+		/// <param name="timingInfo">Timing info to construct from</param>
+		public GetTimingInfoResponse(TimingInfo timingInfo)
 		{
-			this.TotalWaitTime = (float?)TimingInfo.TotalWaitTime?.TotalSeconds;
-			this.TotalInitTime = (float?)TimingInfo.TotalInitTime?.TotalSeconds;
-			this.TotalTimeToComplete = (float?)TimingInfo.TotalTimeToComplete?.TotalSeconds;
+			TotalWaitTime = (float?)timingInfo.TotalWaitTime?.TotalSeconds;
+			TotalInitTime = (float?)timingInfo.TotalInitTime?.TotalSeconds;
+			TotalTimeToComplete = (float?)timingInfo.TotalTimeToComplete?.TotalSeconds;
 
-			this.AverageTotalWaitTime = (float?)TimingInfo.AverageTotalWaitTime?.TotalSeconds;
-			this.AverageTotalInitTime = (float?)TimingInfo.AverageTotalInitTime?.TotalSeconds;
-			this.AverageTotalTimeToComplete = (float?)TimingInfo.AverageTotalTimeToComplete?.TotalSeconds;
+			AverageTotalWaitTime = (float?)timingInfo.AverageTotalWaitTime?.TotalSeconds;
+			AverageTotalInitTime = (float?)timingInfo.AverageTotalInitTime?.TotalSeconds;
+			AverageTotalTimeToComplete = (float?)timingInfo.AverageTotalTimeToComplete?.TotalSeconds;
 		}
 	}
 
@@ -992,15 +989,15 @@ namespace Horde.Build.Api
 		/// <summary>
 		/// Constructor
 		/// </summary>
-		/// <param name="Name">Name of the node</param>
-		/// <param name="JobTimingInfo">Timing info to construct from</param>
-		public GetStepTimingInfoResponse(string? Name, TimingInfo JobTimingInfo)
-			: base(JobTimingInfo)
+		/// <param name="name">Name of the node</param>
+		/// <param name="jobTimingInfo">Timing info to construct from</param>
+		public GetStepTimingInfoResponse(string? name, TimingInfo jobTimingInfo)
+			: base(jobTimingInfo)
 		{
-			this.Name = Name;
-			this.AverageStepWaitTime = JobTimingInfo.StepTiming?.AverageWaitTime;
-			this.AverageStepInitTime = JobTimingInfo.StepTiming?.AverageInitTime;
-			this.AverageStepDuration = JobTimingInfo.StepTiming?.AverageDuration;
+			Name = name;
+			AverageStepWaitTime = jobTimingInfo.StepTiming?.AverageWaitTime;
+			AverageStepInitTime = jobTimingInfo.StepTiming?.AverageInitTime;
+			AverageStepDuration = jobTimingInfo.StepTiming?.AverageDuration;
 		}
 	}
 
@@ -1044,15 +1041,15 @@ namespace Horde.Build.Api
 		/// <summary>
 		/// Constructor
 		/// </summary>
-		/// <param name="Label">The label to construct from</param>
-		/// <param name="TimingInfo">Timing info to construct from</param>
-		public GetLabelTimingInfoResponse(ILabel Label, TimingInfo TimingInfo)
-			: base(TimingInfo)
+		/// <param name="label">The label to construct from</param>
+		/// <param name="timingInfo">Timing info to construct from</param>
+		public GetLabelTimingInfoResponse(ILabel label, TimingInfo timingInfo)
+			: base(timingInfo)
 		{
-			this.DashboardName = Label.DashboardName;
-			this.DashboardCategory = Label.DashboardCategory;
-			this.UgsName = Label.UgsName;
-			this.UgsProject = Label.UgsProject;
+			DashboardName = label.DashboardName;
+			DashboardCategory = label.DashboardCategory;
+			UgsName = label.UgsName;
+			UgsProject = label.UgsProject;
 		}
 	}
 
@@ -1114,19 +1111,19 @@ namespace Horde.Build.Api
 		/// <summary>
 		/// Constructor
 		/// </summary>
-		/// <param name="JobStepRef">The jobstep ref to construct from</param>
-		public GetJobStepRefResponse(IJobStepRef JobStepRef)
+		/// <param name="jobStepRef">The jobstep ref to construct from</param>
+		public GetJobStepRefResponse(IJobStepRef jobStepRef)
 		{
-			this.JobId = JobStepRef.Id.JobId.ToString();
-			this.BatchId = JobStepRef.Id.BatchId.ToString();
-			this.StepId = JobStepRef.Id.StepId.ToString();
-			this.Change = JobStepRef.Change;
-			this.LogId = JobStepRef.LogId.ToString();
-			this.PoolId = JobStepRef.PoolId?.ToString();
-			this.AgentId = JobStepRef.AgentId?.ToString();
-			this.Outcome = JobStepRef.Outcome;
-			this.StartTime = JobStepRef.StartTimeUtc;
-			this.FinishTime = JobStepRef.FinishTimeUtc;
+			JobId = jobStepRef.Id.JobId.ToString();
+			BatchId = jobStepRef.Id.BatchId.ToString();
+			StepId = jobStepRef.Id.StepId.ToString();
+			Change = jobStepRef.Change;
+			LogId = jobStepRef.LogId.ToString();
+			PoolId = jobStepRef.PoolId?.ToString();
+			AgentId = jobStepRef.AgentId?.ToString();
+			Outcome = jobStepRef.Outcome;
+			StartTime = jobStepRef.StartTimeUtc;
+			FinishTime = jobStepRef.FinishTimeUtc;
 		}
 	}
 }
diff --git a/Engine/Source/Programs/Horde/Horde.Build/Api/Leases.cs b/Engine/Source/Programs/Horde/Horde.Build/Api/Leases.cs
--- a/Engine/Source/Programs/Horde/Horde.Build/Api/Leases.cs
+++ b/Engine/Source/Programs/Horde/Horde.Build/Api/Leases.cs
@@ -1,17 +1,5 @@
 // Copyright Epic Games, Inc. All Rights Reserved.
 
-using EpicGames.Core;
-using HordeCommon;
-using Horde.Build.Models;
-using Horde.Build.Utilities;
-using System;
-using System.Collections.Generic;
-using System.ComponentModel.DataAnnotations;
-using System.Linq;
-using System.Security.Claims;
-using System.Text.Json.Serialization;
-using System.Threading.Tasks;
-
 namespace Horde.Build.Api
 {
 
@@ -24,7 +12,5 @@ namespace Horde.Build.Api
 		/// Mark this lease as aborted
 		/// </summary>
 		public bool? Aborted { get; set; }
-
 	}
-
 }
diff --git a/Engine/Source/Programs/Horde/Horde.Build/Api/LogEvents.cs b/Engine/Source/Programs/Horde/Horde.Build/Api/LogEvents.cs
--- a/Engine/Source/Programs/Horde/Horde.Build/Api/LogEvents.cs
+++ b/Engine/Source/Programs/Horde/Horde.Build/Api/LogEvents.cs
@@ -1,14 +1,10 @@
 // Copyright Epic Games, Inc. All Rights Reserved.
 
-using EpicGames.Core;
-using HordeCommon;
-using Horde.Build.Models;
-using Microsoft.Extensions.Logging;
-using System;
 using System.Collections.Generic;
-using System.Linq;
 using System.Text.Json;
-using System.Threading.Tasks;
+using EpicGames.Core;
+using Horde.Build.Models;
+using HordeCommon;
 
 namespace Horde.Build.Api
 {
@@ -50,17 +46,17 @@ namespace Horde.Build.Api
 		/// <summary>
 		/// Constructor
 		/// </summary>
-		/// <param name="Event">The event to construct from</param>
-		/// <param name="EventData">The event data</param>
-		/// <param name="IssueId">The issue for this event</param>
-		public GetLogEventResponse(ILogEvent Event, ILogEventData EventData, int? IssueId)
+		/// <param name="logEvent">The event to construct from</param>
+		/// <param name="eventData">The event data</param>
+		/// <param name="issueId">The issue for this event</param>
+		public GetLogEventResponse(ILogEvent logEvent, ILogEventData eventData, int? issueId)
 		{
-			this.Severity = Event.Severity;
-			this.LogId = Event.LogId.ToString();
-			this.LineIndex = Event.LineIndex;
-			this.LineCount = Event.LineCount;
-			this.IssueId = IssueId;
-			this.Lines = EventData.Lines.ConvertAll(x => x.Data);
+			Severity = logEvent.Severity;
+			LogId = logEvent.LogId.ToString();
+			LineIndex = logEvent.LineIndex;
+			LineCount = logEvent.LineCount;
+			IssueId = issueId;
+			Lines = eventData.Lines.ConvertAll(x => x.Data);
 		}
 	}
 }
diff --git a/Engine/Source/Programs/Horde/Horde.Build/Api/LogFiles.cs b/Engine/Source/Programs/Horde/Horde.Build/Api/LogFiles.cs
--- a/Engine/Source/Programs/Horde/Horde.Build/Api/LogFiles.cs
+++ b/Engine/Source/Programs/Horde/Horde.Build/Api/LogFiles.cs
@@ -1,14 +1,9 @@
 // Copyright Epic Games, Inc. All Rights Reserved.
 
+using System.Collections.Generic;
 using Horde.Build.Logs;
 using Horde.Build.Models;
 using Horde.Build.Services;
-using System;
-using System.Collections.Generic;
-using System.ComponentModel.DataAnnotations;
-using System.Linq;
-using System.Text.Json.Serialization;
-using System.Threading.Tasks;
 
 namespace Horde.Build.Api
 {
@@ -61,15 +56,15 @@ namespace Horde.Build.Api
 		/// <summary>
 		/// Constructor
 		/// </summary>
-		/// <param name="LogFile">The logfile to construct from</param>
-		/// <param name="Metadata">Metadata about the log file</param>
-		public GetLogFileResponse(ILogFile LogFile, LogMetadata Metadata)
+		/// <param name="logFile">The logfile to construct from</param>
+		/// <param name="metadata">Metadata about the log file</param>
+		public GetLogFileResponse(ILogFile logFile, LogMetadata metadata)
 		{
-            this.Id = LogFile.Id.ToString();
-            this.JobId = LogFile.JobId.ToString();
-			this.Type = LogFile.Type;
-            this.Length = Metadata.Length;
-            this.LineCount = Metadata.MaxLineIndex;
+            Id = logFile.Id.ToString();
+            JobId = logFile.JobId.ToString();
+			Type = logFile.Type;
+            Length = metadata.Length;
+            LineCount = metadata.MaxLineIndex;
 		}
 	}
 
diff --git a/Engine/Source/Programs/Horde/Horde.Build/Api/Nodes.cs b/Engine/Source/Programs/Horde/Horde.Build/Api/Nodes.cs
--- a/Engine/Source/Programs/Horde/Horde.Build/Api/Nodes.cs
+++ b/Engine/Source/Programs/Horde/Horde.Build/Api/Nodes.cs
@@ -1,14 +1,10 @@
 // Copyright Epic Games, Inc. All Rights Reserved.
 
-using EpicGames.Core;
-using HordeCommon;
-using Horde.Build.Models;
-using Horde.Build.Utilities;
-using System;
 using System.Collections.Generic;
-using System.ComponentModel.DataAnnotations;
 using System.Linq;
-using System.Threading.Tasks;
+using EpicGames.Core;
+using Horde.Build.Models;
+using HordeCommon;
 
 namespace Horde.Build.Api
 {
@@ -70,19 +66,19 @@ namespace Horde.Build.Api
 		/// <summary>
 		/// Constructor
 		/// </summary>
-		/// <param name="Node">The node to construct from</param>
-		/// <param name="Groups">The groups in the graph</param>
-		public GetNodeResponse(INode Node, IReadOnlyList<INodeGroup> Groups)
+		/// <param name="node">The node to construct from</param>
+		/// <param name="groups">The groups in the graph</param>
+		public GetNodeResponse(INode node, IReadOnlyList<INodeGroup> groups)
 		{
-			this.Name = Node.Name;
-			this.InputDependencies = new List<string>(Node.InputDependencies.Select(x => Groups[x.GroupIdx].Nodes[x.NodeIdx].Name)); ;
-			this.OrderDependencies = new List<string>(Node.OrderDependencies.Select(x => Groups[x.GroupIdx].Nodes[x.NodeIdx].Name)); ;
-			this.Priority = Node.Priority;
-			this.AllowRetry = Node.AllowRetry;
-			this.RunEarly = Node.RunEarly;
-			this.Warnings = Node.Warnings;
-			this.Credentials = Node.Credentials;
-			this.Properties = Node.Properties;
+			Name = node.Name;
+			InputDependencies = new List<string>(node.InputDependencies.Select(x => groups[x.GroupIdx].Nodes[x.NodeIdx].Name));
+			OrderDependencies = new List<string>(node.OrderDependencies.Select(x => groups[x.GroupIdx].Nodes[x.NodeIdx].Name));
+			Priority = node.Priority;
+			AllowRetry = node.AllowRetry;
+			RunEarly = node.RunEarly;
+			Warnings = node.Warnings;
+			Credentials = node.Credentials;
+			Properties = node.Properties;
 		}
 	}
 
@@ -104,12 +100,12 @@ namespace Horde.Build.Api
 		/// <summary>
 		/// Constructor
 		/// </summary>
-		/// <param name="Group">The group to construct from</param>
-		/// <param name="Groups">Other groups in this graph</param>
-		public GetGroupResponse(INodeGroup Group, IReadOnlyList<INodeGroup> Groups)
+		/// <param name="group">The group to construct from</param>
+		/// <param name="groups">Other groups in this graph</param>
+		public GetGroupResponse(INodeGroup group, IReadOnlyList<INodeGroup> groups)
 		{
-			this.AgentType = Group.AgentType;
-			this.Nodes = Group.Nodes.ConvertAll(x => new GetNodeResponse(x, Groups));
+			AgentType = group.AgentType;
+			Nodes = group.Nodes.ConvertAll(x => new GetNodeResponse(x, groups));
 		}
 	}
 
@@ -131,12 +127,12 @@ namespace Horde.Build.Api
 		/// <summary>
 		/// Constructor
 		/// </summary>
-		/// <param name="Aggregate">The aggregate to construct from</param>
-		/// <param name="Groups">List of groups in this graph</param>
-		public GetAggregateResponse(IAggregate Aggregate, IReadOnlyList<INodeGroup> Groups)
+		/// <param name="aggregate">The aggregate to construct from</param>
+		/// <param name="groups">List of groups in this graph</param>
+		public GetAggregateResponse(IAggregate aggregate, IReadOnlyList<INodeGroup> groups)
 		{
-			Name = Aggregate.Name;
-			Nodes = Aggregate.Nodes.ConvertAll(x => x.ToNode(Groups).Name);
+			Name = aggregate.Name;
+			Nodes = aggregate.Nodes.ConvertAll(x => x.ToNode(groups).Name);
 		}
 	}
 
@@ -188,16 +184,16 @@ namespace Horde.Build.Api
 		/// <summary>
 		/// Constructor
 		/// </summary>
-		/// <param name="Label">The label to construct from</param>
-		/// <param name="Groups">List of groups in this graph</param>
-		public GetLabelResponse(ILabel Label, IReadOnlyList<INodeGroup> Groups)
+		/// <param name="label">The label to construct from</param>
+		/// <param name="groups">List of groups in this graph</param>
+		public GetLabelResponse(ILabel label, IReadOnlyList<INodeGroup> groups)
 		{
-			DashboardName = Label.DashboardName;
-			DashboardCategory = Label.DashboardCategory;
-			UgsName = Label.UgsName;
-			UgsProject = Label.UgsProject;
-			RequiredNodes = Label.RequiredNodes.ConvertAll(x => x.ToNode(Groups).Name);
-			IncludedNodes = Label.IncludedNodes.ConvertAll(x => x.ToNode(Groups).Name);
+			DashboardName = label.DashboardName;
+			DashboardCategory = label.DashboardCategory;
+			UgsName = label.UgsName;
+			UgsProject = label.UgsProject;
+			RequiredNodes = label.RequiredNodes.ConvertAll(x => x.ToNode(groups).Name);
+			IncludedNodes = label.IncludedNodes.ConvertAll(x => x.ToNode(groups).Name);
 		}
 	}
 
@@ -229,21 +225,21 @@ namespace Horde.Build.Api
 		/// <summary>
 		/// Constructor
 		/// </summary>
-		/// <param name="Graph">The graph to construct from</param>
-		public GetGraphResponse(IGraph Graph)
+		/// <param name="graph">The graph to construct from</param>
+		public GetGraphResponse(IGraph graph)
 		{
-			Hash = Graph.Id.ToString();
-			if (Graph.Groups.Count > 0)
+			Hash = graph.Id.ToString();
+			if (graph.Groups.Count > 0)
 			{
-				Groups = Graph.Groups.ConvertAll(x => new GetGroupResponse(x, Graph.Groups));
+				Groups = graph.Groups.ConvertAll(x => new GetGroupResponse(x, graph.Groups));
 			}
-			if (Graph.Aggregates.Count > 0)
+			if (graph.Aggregates.Count > 0)
 			{
-				Aggregates = Graph.Aggregates.ConvertAll(x => new GetAggregateResponse(x, Graph.Groups));
+				Aggregates = graph.Aggregates.ConvertAll(x => new GetAggregateResponse(x, graph.Groups));
 			}
-			if (Graph.Labels.Count > 0)
+			if (graph.Labels.Count > 0)
 			{
-				Labels = Graph.Labels.ConvertAll(x => new GetLabelResponse(x, Graph.Groups));
+				Labels = graph.Labels.ConvertAll(x => new GetLabelResponse(x, graph.Groups));
 			}
 		}
 	}
diff --git a/Engine/Source/Programs/Horde/Horde.Build/Api/Notices.cs b/Engine/Source/Programs/Horde/Horde.Build/Api/Notices.cs
--- a/Engine/Source/Programs/Horde/Horde.Build/Api/Notices.cs
+++ b/Engine/Source/Programs/Horde/Horde.Build/Api/Notices.cs
@@ -31,11 +31,10 @@ namespace Horde.Build.Api
 		/// <summary>
 		/// Constructor
 		/// </summary>
-		public CreateNoticeRequest(string Message)
+		public CreateNoticeRequest(string message)
 		{
-			this.Message = Message;
+			Message = message;
 		}
-
 	}
 
 	/// <summary>
@@ -43,13 +42,11 @@ namespace Horde.Build.Api
 	/// </summary>
 	public class UpdateNoticeRequest
 	{
-
 		/// <summary>
 		/// The id of the notice to update
 		/// </summary>
 		public string Id { get; set; }
 
-
 		/// <summary>
 		/// Start time to display this message
 		/// </summary>
@@ -68,14 +65,12 @@ namespace Horde.Build.Api
 		/// <summary>
 		/// Constructor
 		/// </summary>
-		public UpdateNoticeRequest(string Id)
+		public UpdateNoticeRequest(string id)
 		{
-			this.Id = Id;
+			Id = id;
 		}
-
 	}
 
-
 	/// <summary>
 	/// Notice informtation
 	/// </summary>
@@ -118,13 +113,13 @@ namespace Horde.Build.Api
 		/// <summary>
 		/// Constructor
 		/// </summary>
-		public GetNoticeResponse(INotice Notice, GetThinUserInfoResponse? CreateByUser = null)
+		public GetNoticeResponse(INotice notice, GetThinUserInfoResponse? createByUser = null)
 		{
-			this.Id = Notice.Id.ToString();
-			this.StartTime = Notice.StartTime;
-			this.FinishTime = Notice.FinishTime;
-			this.Message = Notice.Message;
-			this.CreatedByUser = CreateByUser;
+			Id = notice.Id.ToString();
+			StartTime = notice.StartTime;
+			FinishTime = notice.FinishTime;
+			Message = notice.Message;
+			CreatedByUser = createByUser;
 		}
 	}
 }
diff --git a/Engine/Source/Programs/Horde/Horde.Build/Api/Notifications.cs b/Engine/Source/Programs/Horde/Horde.Build/Api/Notifications.cs
--- a/Engine/Source/Programs/Horde/Horde.Build/Api/Notifications.cs
+++ b/Engine/Source/Programs/Horde/Horde.Build/Api/Notifications.cs
@@ -1,6 +1,5 @@
 // Copyright Epic Games, Inc. All Rights Reserved.
 
-using System.ComponentModel.DataAnnotations;
 using Horde.Build.Models;
 
 namespace Horde.Build.Api
@@ -39,18 +38,18 @@ namespace Horde.Build.Api
 		/// <summary>
 		/// Constructor
 		/// </summary>
-		/// <param name="Subscription">The subscription interface</param>
-		internal GetNotificationResponse(INotificationSubscription? Subscription)
+		/// <param name="subscription">The subscription interface</param>
+		internal GetNotificationResponse(INotificationSubscription? subscription)
 		{
-			if (Subscription == null)
+			if (subscription == null)
 			{
 				Email = false;
 				Slack = false;
 			}
 			else
 			{
-				Email = Subscription.Email;
-				Slack = Subscription.Slack;
+				Email = subscription.Email;
+				Slack = subscription.Slack;
 			}
 		}
 	}
diff --git a/Engine/Source/Programs/Horde/Horde.Build/Api/Pools.cs b/Engine/Source/Programs/Horde/Horde.Build/Api/Pools.cs
--- a/Engine/Source/Programs/Horde/Horde.Build/Api/Pools.cs
+++ b/Engine/Source/Programs/Horde/Horde.Build/Api/Pools.cs
@@ -1,12 +1,11 @@
 // Copyright Epic Games, Inc. All Rights Reserved.
 
-using EpicGames.Horde.Common;
-using Horde.Build.Models;
-using System;
 using System.Collections.Generic;
 using System.ComponentModel.DataAnnotations;
 using System.Linq;
+using EpicGames.Horde.Common;
 using Horde.Build.Fleet.Autoscale;
+using Horde.Build.Models;
 
 namespace Horde.Build.Api
 {
@@ -16,8 +15,10 @@ namespace Horde.Build.Api
 		/// <summary>
 		/// Construct a public REST API representation from the internal one
 		/// </summary>
-		/// <param name="Settings"></param>
-		public LeaseUtilizationSettings(Horde.Build.Fleet.Autoscale.LeaseUtilizationSettings Settings)
+		/// <param name="settings"></param>
+#pragma warning disable IDE0060 // Remove unused parameter
+		public LeaseUtilizationSettings(Horde.Build.Fleet.Autoscale.LeaseUtilizationSettings settings)
+#pragma warning restore IDE0060 // Remove unused parameter
 		{
 		}
 		
@@ -51,11 +52,11 @@ namespace Horde.Build.Api
 		/// <summary>
 		/// Construct a public REST API representation from the internal one
 		/// </summary>
-		/// <param name="Settings"></param>
-		public JobQueueSettings(Horde.Build.Fleet.Autoscale.JobQueueSettings Settings)
+		/// <param name="settings"></param>
+		public JobQueueSettings(Horde.Build.Fleet.Autoscale.JobQueueSettings settings)
 		{
-			ScaleInFactor = Settings.ScaleInFactor;
-			ScaleOutFactor = Settings.ScaleOutFactor;
+			ScaleInFactor = settings.ScaleInFactor;
+			ScaleOutFactor = settings.ScaleOutFactor;
 		}
 		
 		/// <summary>
@@ -143,10 +144,10 @@ namespace Horde.Build.Api
 		/// <summary>
 		/// Constructor
 		/// </summary>
-		/// <param name="Id">Unique id for the new pool</param>
-		public CreatePoolResponse(string Id)
+		/// <param name="id">Unique id for the new pool</param>
+		public CreatePoolResponse(string id)
 		{
-			this.Id = Id;
+			Id = id;
 		}
 	}
 
@@ -306,22 +307,22 @@ namespace Horde.Build.Api
 		/// <summary>
 		/// Constructor
 		/// </summary>
-		/// <param name="Pool">The pool to construct from</param>
-		public GetPoolResponse(IPool Pool)
+		/// <param name="pool">The pool to construct from</param>
+		public GetPoolResponse(IPool pool)
 		{
-			this.Id = Pool.Id.ToString();
-			this.Name = Pool.Name;
-			this.Condition = Pool.Condition;
-			this.EnableAutoscaling = Pool.EnableAutoscaling;
-			this.ScaleOutCooldown = Pool.ScaleOutCooldown == null ? null : (int)Pool.ScaleOutCooldown.Value.TotalSeconds;
-			this.ScaleInCooldown = Pool.ScaleInCooldown == null ? null : (int)Pool.ScaleInCooldown.Value.TotalSeconds;
-			this.SizeStrategy = Pool.SizeStrategy;
-			this.LeaseUtilizationSettings = Pool.LeaseUtilizationSettings == null ? null : new LeaseUtilizationSettings(Pool.LeaseUtilizationSettings);
-			this.JobQueueSettings = Pool.JobQueueSettings == null ? null : new JobQueueSettings(Pool.JobQueueSettings);
-			this.MinAgents = Pool.MinAgents;
-			this.NumReserveAgents = Pool.NumReserveAgents;
-			this.Workspaces = Pool.Workspaces.Select(x => new GetAgentWorkspaceResponse(x)).ToList();
-			this.Properties = Pool.Properties;
+			Id = pool.Id.ToString();
+			Name = pool.Name;
+			Condition = pool.Condition;
+			EnableAutoscaling = pool.EnableAutoscaling;
+			ScaleOutCooldown = pool.ScaleOutCooldown == null ? null : (int)pool.ScaleOutCooldown.Value.TotalSeconds;
+			ScaleInCooldown = pool.ScaleInCooldown == null ? null : (int)pool.ScaleInCooldown.Value.TotalSeconds;
+			SizeStrategy = pool.SizeStrategy;
+			LeaseUtilizationSettings = pool.LeaseUtilizationSettings == null ? null : new LeaseUtilizationSettings(pool.LeaseUtilizationSettings);
+			JobQueueSettings = pool.JobQueueSettings == null ? null : new JobQueueSettings(pool.JobQueueSettings);
+			MinAgents = pool.MinAgents;
+			NumReserveAgents = pool.NumReserveAgents;
+			Workspaces = pool.Workspaces.Select(x => new GetAgentWorkspaceResponse(x)).ToList();
+			Properties = pool.Properties;
 		}
 	}
 }
diff --git a/Engine/Source/Programs/Horde/Horde.Build/Api/Projects.cs b/Engine/Source/Programs/Horde/Horde.Build/Api/Projects.cs
--- a/Engine/Source/Programs/Horde/Horde.Build/Api/Projects.cs
+++ b/Engine/Source/Programs/Horde/Horde.Build/Api/Projects.cs
@@ -1,11 +1,8 @@
 // Copyright Epic Games, Inc. All Rights Reserved.
 
-using Horde.Build.Models;
-using System;
 using System.Collections.Generic;
 using System.ComponentModel.DataAnnotations;
-using System.Linq;
-using System.Threading.Tasks;
+using Horde.Build.Models;
 
 namespace Horde.Build.Api
 {
@@ -43,9 +40,9 @@ namespace Horde.Build.Api
 		/// <summary>
 		/// Constructor
 		/// </summary>
-		public CreateProjectCategoryRequest(string Name)
+		public CreateProjectCategoryRequest(string name)
 		{
-			this.Name = Name;
+			Name = name;
 		}
 	}
 
@@ -67,12 +64,12 @@ namespace Horde.Build.Api
 		/// <summary>
 		/// Constructor
 		/// </summary>
-		/// <param name="Id">The unique stream id</param>
-		/// <param name="Name">The stream name</param>
-		public GetProjectStreamResponse(string Id, string Name)
+		/// <param name="id">The unique stream id</param>
+		/// <param name="name">The stream name</param>
+		public GetProjectStreamResponse(string id, string name)
 		{
-			this.Id = Id;
-			this.Name = Name;
+			Id = id;
+			Name = name;
 		}
 	}
 
@@ -114,14 +111,14 @@ namespace Horde.Build.Api
 		/// <summary>
 		/// Constructor
 		/// </summary>
-		/// <param name="StreamCategory">The category to construct from</param>
-		public GetProjectCategoryResponse(StreamCategory StreamCategory)
+		/// <param name="streamCategory">The category to construct from</param>
+		public GetProjectCategoryResponse(StreamCategory streamCategory)
 		{
-			this.Name = StreamCategory.Name;
-			this.Row = StreamCategory.Row;
-			this.ShowOnNavMenu = StreamCategory.ShowOnNavMenu;
-			this.IncludePatterns = StreamCategory.IncludePatterns;
-			this.ExcludePatterns = StreamCategory.ExcludePatterns;
+			Name = streamCategory.Name;
+			Row = streamCategory.Row;
+			ShowOnNavMenu = streamCategory.ShowOnNavMenu;
+			IncludePatterns = streamCategory.IncludePatterns;
+			ExcludePatterns = streamCategory.ExcludePatterns;
 		}
 	}
 
@@ -163,20 +160,20 @@ namespace Horde.Build.Api
 		/// <summary>
 		/// Constructor
 		/// </summary>
-		/// <param name="Id">Unique id of the project</param>
-		/// <param name="Name">Name of the project</param>
-		/// <param name="Order">Order to show this project on the dashboard</param>
-		/// <param name="Streams">List of streams to display</param>
-		/// <param name="Categories">List of stream categories to display</param>
-		/// <param name="Acl">Custom permissions for this object</param>
-		public GetProjectResponse(string Id, string Name, int Order, List<GetProjectStreamResponse>? Streams, List<GetProjectCategoryResponse>? Categories, GetAclResponse? Acl)
+		/// <param name="id">Unique id of the project</param>
+		/// <param name="name">Name of the project</param>
+		/// <param name="order">Order to show this project on the dashboard</param>
+		/// <param name="streams">List of streams to display</param>
+		/// <param name="categories">List of stream categories to display</param>
+		/// <param name="acl">Custom permissions for this object</param>
+		public GetProjectResponse(string id, string name, int order, List<GetProjectStreamResponse>? streams, List<GetProjectCategoryResponse>? categories, GetAclResponse? acl)
 		{
-			this.Id = Id;
-			this.Name = Name;
-			this.Order = Order;
-			this.Streams = Streams;
-			this.Categories = Categories;
-			this.Acl = Acl;
+			Id = id;
+			Name = name;
+			Order = order;
+			Streams = streams;
+			Categories = categories;
+			Acl = acl;
 		}
 	}
 }
diff --git a/Engine/Source/Programs/Horde/Horde.Build/Api/Reports.cs b/Engine/Source/Programs/Horde/Horde.Build/Api/Reports.cs
--- a/Engine/Source/Programs/Horde/Horde.Build/Api/Reports.cs
+++ b/Engine/Source/Programs/Horde/Horde.Build/Api/Reports.cs
@@ -1,11 +1,9 @@
 // Copyright Epic Games, Inc. All Rights Reserved.
 
-using EpicGames.Core;
-using Horde.Build.Models;
 using System;
 using System.Collections.Generic;
-using System.Globalization;
-using System.Linq;
+using EpicGames.Core;
+using Horde.Build.Models;
 
 namespace Horde.Build.Api
 {
@@ -28,11 +26,11 @@ namespace Horde.Build.Api
 		/// <summary>
 		/// Constructor
 		/// </summary>
-		/// <param name="Stream">The stream</param>
-		public UtilizationTelemetryStream(IStreamUtilizationTelemetry Stream)
+		/// <param name="stream">The stream</param>
+		public UtilizationTelemetryStream(IStreamUtilizationTelemetry stream)
 		{
-			this.StreamId = Stream.StreamId.ToString();
-			this.Time = Stream.Time;
+			StreamId = stream.StreamId.ToString();
+			Time = stream.Time;
 		}
 	}
 
@@ -74,16 +72,16 @@ namespace Horde.Build.Api
 		/// <summary>
 		/// Constructor
 		/// </summary>
-		/// <param name="Pool"></param>
-		public UtilizationTelemetryPool(IPoolUtilizationTelemetry Pool)
+		/// <param name="pool"></param>
+		public UtilizationTelemetryPool(IPoolUtilizationTelemetry pool)
 		{
-			this.PoolId = Pool.PoolId.ToString();
-			this.NumAgents = Pool.NumAgents;
-			this.AdminTime = Pool.AdminTime;
-			this.HibernatingTime = Pool.HibernatingTime;
-			this.OtherTime = Pool.OtherTime;
+			PoolId = pool.PoolId.ToString();
+			NumAgents = pool.NumAgents;
+			AdminTime = pool.AdminTime;
+			HibernatingTime = pool.HibernatingTime;
+			OtherTime = pool.OtherTime;
 
-			this.Streams = Pool.Streams.ConvertAll(Stream => new UtilizationTelemetryStream(Stream));
+			Streams = pool.Streams.ConvertAll(stream => new UtilizationTelemetryStream(stream));
 		}
 	}
 
@@ -125,17 +123,15 @@ namespace Horde.Build.Api
 		/// <summary>
 		/// Constructor
 		/// </summary>
-		public UtilizationTelemetryResponse(IUtilizationTelemetry Telemetry)
+		public UtilizationTelemetryResponse(IUtilizationTelemetry telemetry)
 		{
-			this.StartTime = Telemetry.StartTime;
-			this.FinishTime = Telemetry.FinishTime;
-			this.AdminTime = Telemetry.AdminTime;
-			this.HibernatingTime = Telemetry.HibernatingTime;
-			this.NumAgents = Telemetry.NumAgents;
+			StartTime = telemetry.StartTime;
+			FinishTime = telemetry.FinishTime;
+			AdminTime = telemetry.AdminTime;
+			HibernatingTime = telemetry.HibernatingTime;
+			NumAgents = telemetry.NumAgents;
 
-			this.Pools = Telemetry.Pools.ConvertAll(Pool => new UtilizationTelemetryPool(Pool));
+			Pools = telemetry.Pools.ConvertAll(pool => new UtilizationTelemetryPool(pool));
 		}
 	}
-
-	
 }
diff --git a/Engine/Source/Programs/Horde/Horde.Build/Api/Schedules.cs b/Engine/Source/Programs/Horde/Horde.Build/Api/Schedules.cs
--- a/Engine/Source/Programs/Horde/Horde.Build/Api/Schedules.cs
+++ b/Engine/Source/Programs/Horde/Horde.Build/Api/Schedules.cs
@@ -1,16 +1,10 @@
 // Copyright Epic Games, Inc. All Rights Reserved.
 
-using Horde.Build.Models;
-using Horde.Build.Utilities;
 using System;
 using System.Collections.Generic;
 using System.ComponentModel.DataAnnotations;
-using System.Linq;
-using System.Net;
-using System.Security.Claims;
-using System.Text.Json;
-using System.Text.Json.Serialization;
-using System.Threading.Tasks;
+using Horde.Build.Models;
+using Horde.Build.Utilities;
 
 namespace Horde.Build.Api
 {
@@ -136,11 +130,11 @@ namespace Horde.Build.Api
 		/// <summary>
 		/// Constructs a model object
 		/// </summary>
-		/// <param name="CurrentTimeUtc">The current time</param>
+		/// <param name="currentTimeUtc">The current time</param>
 		/// <returns>New model object</returns>
-		public Schedule ToModel(DateTime CurrentTimeUtc)
+		public Schedule ToModel(DateTime currentTimeUtc)
 		{
-			return new Schedule(CurrentTimeUtc, Enabled, MaxActive, MaxChanges, RequireSubmittedChange, Gate?.ToModel(), Filter, Files, TemplateParameters, Patterns.ConvertAll(x => x.ToModel()));
+			return new Schedule(currentTimeUtc, Enabled, MaxActive, MaxChanges, RequireSubmittedChange, Gate?.ToModel(), Filter, Files, TemplateParameters, Patterns.ConvertAll(x => x.ToModel()));
 		}
 	}
 
@@ -172,13 +166,13 @@ namespace Horde.Build.Api
 		/// <summary>
 		/// Constructor
 		/// </summary>
-		/// <param name="Pattern">The pattern to construct from</param>
-		public GetSchedulePatternResponse(SchedulePattern Pattern)
+		/// <param name="pattern">The pattern to construct from</param>
+		public GetSchedulePatternResponse(SchedulePattern pattern)
 		{
-			this.DaysOfWeek = Pattern.DaysOfWeek?.ConvertAll(x => x.ToString());
-			this.MinTime = Pattern.MinTime;
-			this.MaxTime = Pattern.MaxTime;
-			this.Interval = Pattern.Interval;
+			DaysOfWeek = pattern.DaysOfWeek?.ConvertAll(x => x.ToString());
+			MinTime = pattern.MinTime;
+			MaxTime = pattern.MaxTime;
+			Interval = pattern.Interval;
 		}
 	}
 
@@ -200,11 +194,11 @@ namespace Horde.Build.Api
 		/// <summary>
 		/// Constructor
 		/// </summary>
-		/// <param name="Gate">Gate to construct from</param>
-		public GetScheduleGateResponse(ScheduleGate Gate)
+		/// <param name="gate">Gate to construct from</param>
+		public GetScheduleGateResponse(ScheduleGate gate)
 		{
-			TemplateId = Gate.TemplateRefId.ToString();
-			Target = Gate.Target;
+			TemplateId = gate.TemplateRefId.ToString();
+			Target = gate.Target;
 		}
 	}
 
@@ -271,23 +265,23 @@ namespace Horde.Build.Api
 		/// <summary>
 		/// Constructor
 		/// </summary>
-		/// <param name="Schedule">The schedule to construct from</param>
-		public GetScheduleResponse(Schedule Schedule)
+		/// <param name="schedule">The schedule to construct from</param>
+		public GetScheduleResponse(Schedule schedule)
 		{
-			this.Enabled = Schedule.Enabled;
-			this.MaxActive = Schedule.MaxActive;
-			this.MaxChanges = Schedule.MaxChanges;
-			this.RequireSubmittedChange = Schedule.RequireSubmittedChange;
-			if (Schedule.Gate != null)
+			Enabled = schedule.Enabled;
+			MaxActive = schedule.MaxActive;
+			MaxChanges = schedule.MaxChanges;
+			RequireSubmittedChange = schedule.RequireSubmittedChange;
+			if (schedule.Gate != null)
 			{
-				this.Gate = new GetScheduleGateResponse(Schedule.Gate);
+				Gate = new GetScheduleGateResponse(schedule.Gate);
 			}
-			this.Filter = Schedule.Filter;
-			this.TemplateParameters = Schedule.TemplateParameters;
-			this.Patterns = Schedule.Patterns.ConvertAll(x => new GetSchedulePatternResponse(x));
-			this.LastTriggerChange = Schedule.LastTriggerChange;
-			this.LastTriggerTime = Schedule.GetLastTriggerTimeUtc();
-			this.ActiveJobs = Schedule.ActiveJobs.ConvertAll(x => x.ToString());
+			Filter = schedule.Filter;
+			TemplateParameters = schedule.TemplateParameters;
+			Patterns = schedule.Patterns.ConvertAll(x => new GetSchedulePatternResponse(x));
+			LastTriggerChange = schedule.LastTriggerChange;
+			LastTriggerTime = schedule.GetLastTriggerTimeUtc();
+			ActiveJobs = schedule.ActiveJobs.ConvertAll(x => x.ToString());
 		}
 	}
 
@@ -304,10 +298,10 @@ namespace Horde.Build.Api
 		/// <summary>
 		/// Constructor
 		/// </summary>
-		/// <param name="Times">List of trigger times</param>
-		public GetScheduleForecastResponse(List<DateTime> Times)
+		/// <param name="times">List of trigger times</param>
+		public GetScheduleForecastResponse(List<DateTime> times)
 		{
-			this.Times = Times;
+			Times = times;
 		}
 	}
 }
diff --git a/Engine/Source/Programs/Horde/Horde.Build/Api/Server.cs b/Engine/Source/Programs/Horde/Horde.Build/Api/Server.cs
--- a/Engine/Source/Programs/Horde/Horde.Build/Api/Server.cs
+++ b/Engine/Source/Programs/Horde/Horde.Build/Api/Server.cs
@@ -1,9 +1,9 @@
 // Copyright Epic Games, Inc. All Rights Reserved.
 
 using System;
-using System.Runtime.InteropServices;
-using System.Reflection;
 using System.Diagnostics;
+using System.Reflection;
+using System.Runtime.InteropServices;
 
 namespace Horde.Build.Api
 {
@@ -20,7 +20,7 @@ namespace Horde.Build.Api
         /// <summary>
         /// The operating system server is hosted on
         /// </summary>
-        public string OSDescription { get; set; }
+        public string OsDescription { get; set; }
 
 		/// <summary>
 		/// Whether this is an installed Horde build
@@ -30,13 +30,13 @@ namespace Horde.Build.Api
 		/// <summary>
 		/// Constructor
 		/// </summary>
-		public GetServerInfoResponse( bool SingleInstance)
+		public GetServerInfoResponse( bool singleInstance)
         {
 
-            FileVersionInfo VersionInfo = FileVersionInfo.GetVersionInfo(Assembly.GetExecutingAssembly().Location);		
-			ServerVersion = VersionInfo.ProductVersion ?? String.Empty;			
-			OSDescription = RuntimeInformation.OSDescription;
-			this.SingleInstance = SingleInstance;
+            FileVersionInfo versionInfo = FileVersionInfo.GetVersionInfo(Assembly.GetExecutingAssembly().Location);		
+			ServerVersion = versionInfo.ProductVersion ?? String.Empty;			
+			OsDescription = RuntimeInformation.OSDescription;
+			SingleInstance = singleInstance;
 		}
 	}
 }
diff --git a/Engine/Source/Programs/Horde/Horde.Build/Api/StreamTabs.cs b/Engine/Source/Programs/Horde/Horde.Build/Api/StreamTabs.cs
--- a/Engine/Source/Programs/Horde/Horde.Build/Api/StreamTabs.cs
+++ b/Engine/Source/Programs/Horde/Horde.Build/Api/StreamTabs.cs
@@ -1,15 +1,9 @@
 // Copyright Epic Games, Inc. All Rights Reserved.
 
-using Amazon.S3.Model.Internal.MarshallTransformations;
-using EpicGames.Core;
-using Horde.Build.Models;
-using Horde.Build.Utilities;
-using System;
 using System.Collections.Generic;
 using System.ComponentModel.DataAnnotations;
-using System.Linq;
-using System.Text.Json.Serialization;
-using System.Threading.Tasks;
+using EpicGames.Core;
+using Horde.Build.Models;
 
 namespace Horde.Build.Api
 {
@@ -132,10 +126,10 @@ namespace Horde.Build.Api
 		/// <summary>
 		/// Constructor
 		/// </summary>
-		/// <param name="Title">Title of this page</param>
-		protected GetStreamTabResponse(string Title)//, StreamPageType Type)
+		/// <param name="title">Title of this page</param>
+		protected GetStreamTabResponse(string title)//, StreamPageType Type)
 		{
-			this.Title = Title;
+			Title = title;
 		}
 	}
 
@@ -158,12 +152,12 @@ namespace Horde.Build.Api
 		/// <summary>
 		/// Constructor
 		/// </summary>
-		/// <param name="Heading">Heading for this column</param>
-		/// <param name="RelativeWidth">Relative width of this column</param>
-		protected GetJobsTabColumnResponse(string Heading, int? RelativeWidth)
+		/// <param name="heading">Heading for this column</param>
+		/// <param name="relativeWidth">Relative width of this column</param>
+		protected GetJobsTabColumnResponse(string heading, int? relativeWidth)
 		{
-			this.Heading = Heading;
-			this.RelativeWidth = RelativeWidth;
+			Heading = heading;
+			RelativeWidth = relativeWidth;
 		}
 	}
 
@@ -181,13 +175,13 @@ namespace Horde.Build.Api
 		/// <summary>
 		/// Constructor
 		/// </summary>
-		/// <param name="Heading">Heading for this column</param>
-		/// <param name="Category">Category of aggregates to display in this column. If null, includes any aggregate not matched by another column.</param>
-		/// <param name="RelativeWidth">Relative width of this column</param>
-		public GetJobsTabLabelColumnResponse(string Heading, string? Category, int? RelativeWidth)
-			: base(Heading, RelativeWidth)
+		/// <param name="heading">Heading for this column</param>
+		/// <param name="category">Category of aggregates to display in this column. If null, includes any aggregate not matched by another column.</param>
+		/// <param name="relativeWidth">Relative width of this column</param>
+		public GetJobsTabLabelColumnResponse(string heading, string? category, int? relativeWidth)
+			: base(heading, relativeWidth)
 		{
-			this.Category = Category;
+			Category = category;
 		}
 	}
 
@@ -205,13 +199,13 @@ namespace Horde.Build.Api
 		/// <summary>
 		/// Constructor
 		/// </summary>
-		/// <param name="Heading">Heading for this column</param>
-		/// <param name="Parameter">Name of the parameter to display</param>
-		/// <param name="RelativeWidth">Relative width of this column</param>
-		public GetJobsTabParameterColumnResponse(string Heading, string? Parameter, int? RelativeWidth)
-			: base(Heading, RelativeWidth)
+		/// <param name="heading">Heading for this column</param>
+		/// <param name="parameter">Name of the parameter to display</param>
+		/// <param name="relativeWidth">Relative width of this column</param>
+		public GetJobsTabParameterColumnResponse(string heading, string? parameter, int? relativeWidth)
+			: base(heading, relativeWidth)
 		{
-			this.Parameter = Parameter;
+			Parameter = parameter;
 		}
 	}
 
@@ -244,18 +238,18 @@ namespace Horde.Build.Api
 		/// <summary>
 		/// Constructor
 		/// </summary>
-		/// <param name="Title">Title for this page</param>
+		/// <param name="title">Title for this page</param>
 		/// <param name="bShowNames">Whether to show names on the page</param>
-		/// <param name="Templates">Templates to include on this page</param>
-		/// <param name="JobNames">List of job names to include on this page</param>
-		/// <param name="Columns">List of columns to display</param>
-		public GetJobsTabResponse(string Title, bool bShowNames, List<string>? Templates, List<string>? JobNames, List<GetJobsTabColumnResponse>? Columns)
-			: base(Title)
+		/// <param name="templates">Templates to include on this page</param>
+		/// <param name="jobNames">List of job names to include on this page</param>
+		/// <param name="columns">List of columns to display</param>
+		public GetJobsTabResponse(string title, bool bShowNames, List<string>? templates, List<string>? jobNames, List<GetJobsTabColumnResponse>? columns)
+			: base(title)
 		{
-			this.ShowNames = bShowNames;
-			this.Templates = Templates;
-			this.JobNames = JobNames;
-			this.Columns = Columns;
+			ShowNames = bShowNames;
+			Templates = templates;
+			JobNames = jobNames;
+			Columns = columns;
 		}
 	}
 }
diff --git a/Engine/Source/Programs/Horde/Horde.Build/Api/Streams.cs b/Engine/Source/Programs/Horde/Horde.Build/Api/Streams.cs
--- a/Engine/Source/Programs/Horde/Horde.Build/Api/Streams.cs
+++ b/Engine/Source/Programs/Horde/Horde.Build/Api/Streams.cs
@@ -1,13 +1,11 @@
 // Copyright Epic Games, Inc. All Rights Reserved.
 
-using HordeCommon;
-using Horde.Build.Models;
-using Horde.Build.Utilities;
 using System;
 using System.Collections.Generic;
 using System.ComponentModel.DataAnnotations;
-using System.Linq;
-using System.Threading.Tasks;
+using Horde.Build.Models;
+using Horde.Build.Utilities;
+using HordeCommon;
 
 namespace Horde.Build.Api
 {
@@ -39,14 +37,14 @@ namespace Horde.Build.Api
 		/// <returns></returns>
 		public ChangeQuery ToModel()
 		{
-			ChangeQuery Query = new ChangeQuery();
+			ChangeQuery query = new ChangeQuery();
 			if (TemplateId != null)
 			{
-				Query.TemplateRefId = new StringId<TemplateRef>(TemplateId);
+				query.TemplateRefId = new StringId<TemplateRef>(TemplateId);
 			}
-			Query.Target = Target;
-			Query.Outcomes = Outcomes;
-			return Query;
+			query.Target = Target;
+			query.Outcomes = Outcomes;
+			return query;
 		}
 	}
 
@@ -263,16 +261,16 @@ namespace Horde.Build.Api
 		/// <summary>
 		/// Constructor
 		/// </summary>
-		/// <param name="Pool">Pool of agents to use for this agent type</param>
-		/// <param name="Workspace">Name of the workspace to sync</param>
-		/// <param name="TempStorageDir">Path to the temp storage directory</param>
-		/// <param name="Environment">Environment variables to be set when executing this job</param>
-		public GetAgentTypeResponse(string Pool, string? Workspace, string? TempStorageDir, Dictionary<string, string>? Environment)
+		/// <param name="pool">Pool of agents to use for this agent type</param>
+		/// <param name="workspace">Name of the workspace to sync</param>
+		/// <param name="tempStorageDir">Path to the temp storage directory</param>
+		/// <param name="environment">Environment variables to be set when executing this job</param>
+		public GetAgentTypeResponse(string pool, string? workspace, string? tempStorageDir, Dictionary<string, string>? environment)
 		{
-			this.Pool = Pool;
-			this.Workspace = Workspace;
-			this.TempStorageDir = TempStorageDir;
-			this.Environment = Environment;
+			Pool = pool;
+			Workspace = workspace;
+			TempStorageDir = tempStorageDir;
+			Environment = environment;
 		}
 	}
 
@@ -331,24 +329,24 @@ namespace Horde.Build.Api
 		/// <summary>
 		/// Constructor
 		/// </summary>
-		/// <param name="Cluster">The server cluster</param>
-		/// <param name="ServerAndPort">The perforce server</param>
-		/// <param name="UserName">The perforce user name</param>
-		/// <param name="Identifier">Identifier to distinguish this workspace from other workspaces. Defaults to the workspace type name.</param>
-		/// <param name="Stream">Override for the stream to sync</param>
-		/// <param name="View">Custom view for the workspace</param>
+		/// <param name="cluster">The server cluster</param>
+		/// <param name="serverAndPort">The perforce server</param>
+		/// <param name="userName">The perforce user name</param>
+		/// <param name="identifier">Identifier to distinguish this workspace from other workspaces. Defaults to the workspace type name.</param>
+		/// <param name="stream">Override for the stream to sync</param>
+		/// <param name="view">Custom view for the workspace</param>
 		/// <param name="bIncremental">Whether to use an incrementally synced workspace</param>
 		/// <param name="bUseAutoSdk">Whether to use the AutoSDK</param>
-		public GetWorkspaceTypeResponse(string? Cluster, string? ServerAndPort, string? UserName, string? Identifier, string? Stream, List<string>? View, bool bIncremental, bool bUseAutoSdk)
+		public GetWorkspaceTypeResponse(string? cluster, string? serverAndPort, string? userName, string? identifier, string? stream, List<string>? view, bool bIncremental, bool bUseAutoSdk)
 		{
-			this.Cluster = Cluster;
-			this.ServerAndPort = ServerAndPort;
-			this.UserName = UserName;
-			this.Identifier = Identifier;
-			this.Stream = Stream;
-			this.View = View;
-			this.Incremental = bIncremental;
-			this.UseAutoSdk = bUseAutoSdk;
+			Cluster = cluster;
+			ServerAndPort = serverAndPort;
+			UserName = userName;
+			Identifier = identifier;
+			Stream = stream;
+			View = view;
+			Incremental = bIncremental;
+			UseAutoSdk = bUseAutoSdk;
 		}
 	}
 
@@ -370,11 +368,11 @@ namespace Horde.Build.Api
 		/// <summary>
 		/// Constructor
 		/// </summary>
-		/// <param name="Trigger">The trigger definition</param>
-		public GetChainedJobTemplateResponse(ChainedJobTemplate Trigger)
+		/// <param name="trigger">The trigger definition</param>
+		public GetChainedJobTemplateResponse(ChainedJobTemplate trigger)
 		{
-			this.Trigger = Trigger.Trigger;
-			this.TemplateId = Trigger.TemplateRefId.ToString();
+			Trigger = trigger.Trigger;
+			TemplateId = trigger.TemplateRefId.ToString();
 		}
 	}
 
@@ -436,22 +434,22 @@ namespace Horde.Build.Api
 		/// <summary>
 		/// Constructor
 		/// </summary>
-		/// <param name="Id">The template ref id</param>
-		/// <param name="TemplateRef">The template ref</param>
-		/// <param name="Template">The template instance</param>
+		/// <param name="id">The template ref id</param>
+		/// <param name="templateRef">The template ref</param>
+		/// <param name="template">The template instance</param>
 		/// <param name="bIncludeAcl">Whether to include the ACL in the response</param>
-		public GetTemplateRefResponse(TemplateRefId Id, TemplateRef TemplateRef, ITemplate Template, bool bIncludeAcl)
-			: base(Template)
+		public GetTemplateRefResponse(TemplateRefId id, TemplateRef templateRef, ITemplate template, bool bIncludeAcl)
+			: base(template)
 		{
-			this.Id = Id.ToString();
-			Hash = TemplateRef.Hash.ToString();
-			ShowUgsBadges = TemplateRef.ShowUgsBadges;
-			ShowUgsAlerts = TemplateRef.ShowUgsAlerts;
-			NotificationChannel = TemplateRef.NotificationChannel;
-			NotificationChannelFilter = TemplateRef.NotificationChannelFilter;
-			Schedule = (TemplateRef.Schedule != null) ? new GetScheduleResponse(TemplateRef.Schedule) : null;
-			ChainedJobs = (TemplateRef.ChainedJobs != null && TemplateRef.ChainedJobs.Count > 0) ? TemplateRef.ChainedJobs.ConvertAll(x => new GetChainedJobTemplateResponse(x)) : null;
-			Acl = (bIncludeAcl && TemplateRef.Acl != null)? new GetAclResponse(TemplateRef.Acl) : null;
+			Id = id.ToString();
+			Hash = templateRef.Hash.ToString();
+			ShowUgsBadges = templateRef.ShowUgsBadges;
+			ShowUgsAlerts = templateRef.ShowUgsAlerts;
+			NotificationChannel = templateRef.NotificationChannel;
+			NotificationChannelFilter = templateRef.NotificationChannelFilter;
+			Schedule = (templateRef.Schedule != null) ? new GetScheduleResponse(templateRef.Schedule) : null;
+			ChainedJobs = (templateRef.ChainedJobs != null && templateRef.ChainedJobs.Count > 0) ? templateRef.ChainedJobs.ConvertAll(x => new GetChainedJobTemplateResponse(x)) : null;
+			Acl = (bIncludeAcl && templateRef.Acl != null)? new GetAclResponse(templateRef.Acl) : null;
 		}
 	}
 
@@ -553,43 +551,43 @@ namespace Horde.Build.Api
 		/// <summary>
 		/// Constructor
 		/// </summary>
-		/// <param name="Id">Unique id of the stream</param>
-		/// <param name="ProjectId">Unique id of the project containing the stream</param>
-		/// <param name="Name">Name of the stream</param>
-		/// <param name="ConfigPath">Path to the config file for this stream</param>
-		/// <param name="ConfigRevision">The config file path on the server</param>
-		/// <param name="Order">Order to display this stream</param>
-		/// <param name="NotificationChannel"></param>
-		/// <param name="NotificationChannelFilter"></param>
-		/// <param name="TriageChannel"></param>
-		/// <param name="DefaultPreflight">The default template to use for preflights</param>
-		/// <param name="Tabs">List of tabs to display for this stream</param>
-		/// <param name="AgentTypes">Map of agent type name to description</param>
-		/// <param name="WorkspaceTypes">Map of workspace name to description</param>
-		/// <param name="Templates">Templates for this stream</param>
-		/// <param name="Acl">Permissions for this object</param>
-		/// <param name="PausedUntil">Stream paused for new builds until this date</param>
-		/// <param name="PauseComment">Reason for stream being paused</param>
-		public GetStreamResponse(string Id, string ProjectId, string Name, string ConfigPath, string ConfigRevision, int Order, string? NotificationChannel, string? NotificationChannelFilter, string? TriageChannel, DefaultPreflightRequest? DefaultPreflight, List<GetStreamTabResponse> Tabs, Dictionary<string, GetAgentTypeResponse> AgentTypes, Dictionary<string, GetWorkspaceTypeResponse>? WorkspaceTypes, List<GetTemplateRefResponse> Templates, GetAclResponse? Acl, DateTime? PausedUntil, string? PauseComment)
+		/// <param name="id">Unique id of the stream</param>
+		/// <param name="projectId">Unique id of the project containing the stream</param>
+		/// <param name="name">Name of the stream</param>
+		/// <param name="configPath">Path to the config file for this stream</param>
+		/// <param name="configRevision">The config file path on the server</param>
+		/// <param name="order">Order to display this stream</param>
+		/// <param name="notificationChannel"></param>
+		/// <param name="notificationChannelFilter"></param>
+		/// <param name="triageChannel"></param>
+		/// <param name="defaultPreflight">The default template to use for preflights</param>
+		/// <param name="tabs">List of tabs to display for this stream</param>
+		/// <param name="agentTypes">Map of agent type name to description</param>
+		/// <param name="workspaceTypes">Map of workspace name to description</param>
+		/// <param name="templates">Templates for this stream</param>
+		/// <param name="acl">Permissions for this object</param>
+		/// <param name="pausedUntil">Stream paused for new builds until this date</param>
+		/// <param name="pauseComment">Reason for stream being paused</param>
+		public GetStreamResponse(string id, string projectId, string name, string configPath, string configRevision, int order, string? notificationChannel, string? notificationChannelFilter, string? triageChannel, DefaultPreflightRequest? defaultPreflight, List<GetStreamTabResponse> tabs, Dictionary<string, GetAgentTypeResponse> agentTypes, Dictionary<string, GetWorkspaceTypeResponse>? workspaceTypes, List<GetTemplateRefResponse> templates, GetAclResponse? acl, DateTime? pausedUntil, string? pauseComment)
 		{
-			this.Id = Id;
-			this.ProjectId = ProjectId;
-			this.Name = Name;
-			this.ConfigPath = ConfigPath;
-			this.ConfigRevision = ConfigRevision;
-			this.Order = Order;
-			this.NotificationChannel = NotificationChannel;
-			this.NotificationChannelFilter = NotificationChannelFilter;
-			this.TriageChannel = TriageChannel;
-			this.DefaultPreflightTemplate = DefaultPreflight?.TemplateId;
-			this.DefaultPreflight = DefaultPreflight;
-			this.Tabs = Tabs;
-			this.AgentTypes = AgentTypes;
-			this.WorkspaceTypes = WorkspaceTypes;
-			this.Templates = Templates;
-			this.Acl = Acl;
-			this.PausedUntil = PausedUntil;
-			this.PauseComment = PauseComment;
+			Id = id;
+			ProjectId = projectId;
+			Name = name;
+			ConfigPath = configPath;
+			ConfigRevision = configRevision;
+			Order = order;
+			NotificationChannel = notificationChannel;
+			NotificationChannelFilter = notificationChannelFilter;
+			TriageChannel = triageChannel;
+			DefaultPreflightTemplate = defaultPreflight?.TemplateId;
+			DefaultPreflight = defaultPreflight;
+			Tabs = tabs;
+			AgentTypes = agentTypes;
+			WorkspaceTypes = workspaceTypes;
+			Templates = templates;
+			Acl = acl;
+			PausedUntil = pausedUntil;
+			PauseComment = pauseComment;
 		}
 	}
 }
diff --git a/Engine/Source/Programs/Horde/Horde.Build/Api/Subscriptions.cs b/Engine/Source/Programs/Horde/Horde.Build/Api/Subscriptions.cs
--- a/Engine/Source/Programs/Horde/Horde.Build/Api/Subscriptions.cs
+++ b/Engine/Source/Programs/Horde/Horde.Build/Api/Subscriptions.cs
@@ -1,11 +1,7 @@
 // Copyright Epic Games, Inc. All Rights Reserved.
 
-using Horde.Build.Models;
-using System;
-using System.Collections.Generic;
 using System.ComponentModel.DataAnnotations;
-using System.Linq;
-using System.Threading.Tasks;
+using Horde.Build.Models;
 
 namespace Horde.Build.Api
 {
@@ -45,10 +41,10 @@ namespace Horde.Build.Api
 		/// <summary>
 		/// Constructor
 		/// </summary>
-		/// <param name="Subscription">The subscription to construct from</param>
-		public CreateSubscriptionResponse(ISubscription Subscription)
+		/// <param name="subscription">The subscription to construct from</param>
+		public CreateSubscriptionResponse(ISubscription subscription)
 		{
-			this.Id = Subscription.Id;
+			Id = subscription.Id;
 		}
 	}
 
@@ -80,13 +76,13 @@ namespace Horde.Build.Api
 		/// <summary>
 		/// Constructor
 		/// </summary>
-		/// <param name="Subscription">Subscription to construct from</param>
-		public GetSubscriptionResponse(ISubscription Subscription)
+		/// <param name="subscription">Subscription to construct from</param>
+		public GetSubscriptionResponse(ISubscription subscription)
 		{
-			this.Id = Subscription.Id;
-			this.Event = Subscription.Event.ToRecord();
-			this.UserId = Subscription.UserId.ToString();
-			this.NotificationType = Subscription.NotificationType;
+			Id = subscription.Id;
+			Event = subscription.Event.ToRecord();
+			UserId = subscription.UserId.ToString();
+			NotificationType = subscription.NotificationType;
 		}
 	}
 }
diff --git a/Engine/Source/Programs/Horde/Horde.Build/Api/Templates.cs b/Engine/Source/Programs/Horde/Horde.Build/Api/Templates.cs
--- a/Engine/Source/Programs/Horde/Horde.Build/Api/Templates.cs
+++ b/Engine/Source/Programs/Horde/Horde.Build/Api/Templates.cs
@@ -1,15 +1,11 @@
 // Copyright Epic Games, Inc. All Rights Reserved.
 
-using EpicGames.Core;
-using HordeCommon;
-using Horde.Build.Models;
-using Horde.Build.Utilities;
 using System;
 using System.Collections.Generic;
 using System.ComponentModel.DataAnnotations;
-using System.Linq;
-using System.Text.Json.Serialization;
-using System.Threading.Tasks;
+using EpicGames.Core;
+using Horde.Build.Models;
+using HordeCommon;
 
 namespace Horde.Build.Api
 {
@@ -75,14 +71,14 @@ namespace Horde.Build.Api
 		/// <summary>
 		/// Constructor
 		/// </summary>
-		/// <param name="Label">Name of the group</param>
-		/// <param name="Style">How to display this group</param>
-		/// <param name="Children">List of child parameters</param>
-		public GroupParameterData(string Label, GroupParameterStyle Style, List<ParameterData> Children)
+		/// <param name="label">Name of the group</param>
+		/// <param name="style">How to display this group</param>
+		/// <param name="children">List of child parameters</param>
+		public GroupParameterData(string label, GroupParameterStyle style, List<ParameterData> children)
 		{
-			this.Label = Label;
-			this.Style = Style;
-			this.Children = Children;
+			Label = label;
+			Style = style;
+			Children = children;
 		}
 
 		/// <summary>
@@ -149,22 +145,22 @@ namespace Horde.Build.Api
 		/// <summary>
 		/// Constructor
 		/// </summary>
-		/// <param name="Label">Label to show next to the parameter</param>
-		/// <param name="Argument">Argument to pass this value with</param>
-		/// <param name="Default">Default value for this parameter</param>
-		/// <param name="Hint">Hint text to display for this parameter</param>
-		/// <param name="Validation">Regex used to validate entries</param>
-		/// <param name="ValidationError">Message displayed to explain validation issues</param>
-		/// <param name="ToolTip">Tool tip text to display</param>
-		public TextParameterData(string Label, string Argument, string Default, string? Hint, string? Validation, string? ValidationError, string? ToolTip)
+		/// <param name="label">Label to show next to the parameter</param>
+		/// <param name="argument">Argument to pass this value with</param>
+		/// <param name="defaultValue">Default value for this parameter</param>
+		/// <param name="hint">Hint text to display for this parameter</param>
+		/// <param name="validation">Regex used to validate entries</param>
+		/// <param name="validationError">Message displayed to explain validation issues</param>
+		/// <param name="toolTip">Tool tip text to display</param>
+		public TextParameterData(string label, string argument, string defaultValue, string? hint, string? validation, string? validationError, string? toolTip)
 		{
-			this.Label = Label;
-			this.Argument = Argument;
-			this.Default = Default;
-			this.Hint = Hint;
-			this.Validation = Validation;
-			this.ValidationError = ValidationError;
-			this.ToolTip = ToolTip;
+			Label = label;
+			Argument = argument;
+			Default = defaultValue;
+			Hint = hint;
+			Validation = validation;
+			ValidationError = validationError;
+			ToolTip = toolTip;
 		}
 
 		/// <summary>
@@ -239,18 +235,18 @@ namespace Horde.Build.Api
 		/// <summary>
 		/// Constructor
 		/// </summary>
-		/// <param name="Group">The group to put this parameter in</param>
-		/// <param name="Text">Text to display for this option</param>
-		/// <param name="ArgumentIfEnabled">Argument to pass for this item if it's enabled</param>
-		/// <param name="ArgumentIfDisabled">Argument to pass for this item if it's enabled</param>
-		/// <param name="Default">Whether this item is selected by default</param>
-		public ListParameterItemData(string? Group, string Text, string? ArgumentIfEnabled, string? ArgumentIfDisabled, bool Default)
+		/// <param name="group">The group to put this parameter in</param>
+		/// <param name="text">Text to display for this option</param>
+		/// <param name="argumentIfEnabled">Argument to pass for this item if it's enabled</param>
+		/// <param name="argumentIfDisabled">Argument to pass for this item if it's enabled</param>
+		/// <param name="defaultValue">Whether this item is selected by default</param>
+		public ListParameterItemData(string? group, string text, string? argumentIfEnabled, string? argumentIfDisabled, bool defaultValue)
 		{
-			this.Group = Group;
-			this.Text = Text;
-			this.ArgumentIfEnabled = ArgumentIfEnabled;
-			this.ArgumentIfDisabled = ArgumentIfDisabled;
-			this.Default = Default;
+			Group = group;
+			Text = text;
+			ArgumentIfEnabled = argumentIfEnabled;
+			ArgumentIfDisabled = argumentIfDisabled;
+			Default = defaultValue;
 		}
 
 		/// <summary>
@@ -301,16 +297,16 @@ namespace Horde.Build.Api
 		/// <summary>
 		/// List of possible values
 		/// </summary>
-		/// <param name="Label">Label to show next to this parameter</param>
-		/// <param name="Style">Type of picker to show</param>
-		/// <param name="Items">Entries for this list</param>
-		/// <param name="ToolTip">Tool tip text to display</param>
-		public ListParameterData(string Label, ListParameterStyle Style, List<ListParameterItemData> Items, string? ToolTip)
+		/// <param name="label">Label to show next to this parameter</param>
+		/// <param name="style">Type of picker to show</param>
+		/// <param name="items">Entries for this list</param>
+		/// <param name="toolTip">Tool tip text to display</param>
+		public ListParameterData(string label, ListParameterStyle style, List<ListParameterItemData> items, string? toolTip)
 		{
-			this.Label = Label;
-			this.Style = Style;
-			this.Items = Items;
-			this.ToolTip = ToolTip;
+			Label = label;
+			Style = style;
+			Items = items;
+			ToolTip = toolTip;
 		}
 
 		/// <summary>
@@ -365,18 +361,18 @@ namespace Horde.Build.Api
 		/// <summary>
 		/// Constructor
 		/// </summary>
-		/// <param name="Label">Label to show next to this parameter</param>
-		/// <param name="ArgumentIfEnabled">Value if enabled</param>
-		/// <param name="ArgumentIfDisabled">Value if disabled</param>
-		/// <param name="Default">Whether this option is enabled by default</param>
-		/// <param name="ToolTip">The tool tip text to display</param>
-		public BoolParameterData(string Label, string? ArgumentIfEnabled, string? ArgumentIfDisabled, bool Default, string? ToolTip)
+		/// <param name="label">Label to show next to this parameter</param>
+		/// <param name="argumentIfEnabled">Value if enabled</param>
+		/// <param name="argumentIfDisabled">Value if disabled</param>
+		/// <param name="defaultValue">Whether this option is enabled by default</param>
+		/// <param name="toolTip">The tool tip text to display</param>
+		public BoolParameterData(string label, string? argumentIfEnabled, string? argumentIfDisabled, bool defaultValue, string? toolTip)
 		{
-			this.Label = Label;
-			this.ArgumentIfEnabled = ArgumentIfEnabled;
-			this.ArgumentIfDisabled = ArgumentIfDisabled;
-			this.Default = Default;
-			this.ToolTip = ToolTip;
+			Label = label;
+			ArgumentIfEnabled = argumentIfEnabled;
+			ArgumentIfDisabled = argumentIfDisabled;
+			Default = defaultValue;
+			ToolTip = toolTip;
 		}
 
 		/// <summary>
@@ -496,26 +492,26 @@ namespace Horde.Build.Api
 		/// </summary>
 		protected GetTemplateResponseBase()
 		{
-			this.Name = null!;
-			this.AllowPreflights = true;
-			this.Arguments = new List<string>();
-			this.Parameters = new List<ParameterData>();
+			Name = null!;
+			AllowPreflights = true;
+			Arguments = new List<string>();
+			Parameters = new List<ParameterData>();
 		}
 
 		/// <summary>
 		/// Constructor
 		/// </summary>
-		/// <param name="Template">The template to construct from</param>
-		public GetTemplateResponseBase(ITemplate Template)
+		/// <param name="template">The template to construct from</param>
+		public GetTemplateResponseBase(ITemplate template)
 		{
-			this.Name = Template.Name;
-			this.Priority = Template.Priority;
-			this.AllowPreflights = Template.AllowPreflights;
-			this.UpdateIssues = Template.UpdateIssues;
-			this.InitialAgentType = Template.InitialAgentType;
-			this.SubmitNewChange = Template.SubmitNewChange;
-			this.Arguments = new List<string>(Template.Arguments);
-			this.Parameters = Template.Parameters.ConvertAll(x => x.ToData());
+			Name = template.Name;
+			Priority = template.Priority;
+			AllowPreflights = template.AllowPreflights;
+			UpdateIssues = template.UpdateIssues;
+			InitialAgentType = template.InitialAgentType;
+			SubmitNewChange = template.SubmitNewChange;
+			Arguments = new List<string>(template.Arguments);
+			Parameters = template.Parameters.ConvertAll(x => x.ToData());
 		}
 	}
 
@@ -535,17 +531,17 @@ namespace Horde.Build.Api
 		protected GetTemplateResponse()
 			: base()
 		{
-			this.Id = null!;
+			Id = null!;
 		}
 
 		/// <summary>
 		/// Constructor
 		/// </summary>
-		/// <param name="Template">The template to construct from</param>
-		public GetTemplateResponse(ITemplate Template)
-			: base(Template)
+		/// <param name="template">The template to construct from</param>
+		public GetTemplateResponse(ITemplate template)
+			: base(template)
 		{
-			this.Id = Template.Id.ToString();
+			Id = template.Id.ToString();
 		}
 	}
 }
diff --git a/Engine/Source/Programs/Horde/Horde.Build/Api/TestData.cs b/Engine/Source/Programs/Horde/Horde.Build/Api/TestData.cs
--- a/Engine/Source/Programs/Horde/Horde.Build/Api/TestData.cs
+++ b/Engine/Source/Programs/Horde/Horde.Build/Api/TestData.cs
@@ -1,12 +1,10 @@
 // Copyright Epic Games, Inc. All Rights Reserved.
 
-using Horde.Build.Models;
-using MongoDB.Bson.Serialization;
 using System;
 using System.Collections.Generic;
 using System.ComponentModel.DataAnnotations;
-using System.Linq;
-using System.Threading.Tasks;
+using Horde.Build.Models;
+using MongoDB.Bson.Serialization;
 
 namespace Horde.Build.Api
 {
@@ -52,10 +50,10 @@ namespace Horde.Build.Api
 		/// <summary>
 		/// Constructor
 		/// </summary>
-		/// <param name="Id">Id of the new document</param>
-		public CreateTestDataResponse(string Id)
+		/// <param name="id">Id of the new document</param>
+		public CreateTestDataResponse(string id)
 		{
-			this.Id = Id;
+			Id = id;
 		}
 	}
 
@@ -107,17 +105,17 @@ namespace Horde.Build.Api
 		/// <summary>
 		/// Constructor
 		/// </summary>
-		/// <param name="TestData">Test data to construct from</param>
-		internal GetTestDataResponse(ITestData TestData)
+		/// <param name="testData">Test data to construct from</param>
+		internal GetTestDataResponse(ITestData testData)
 		{
-			this.Id = TestData.Id.ToString();
-			this.StreamId = TestData.StreamId.ToString();
-			this.TemplateRefId = TestData.TemplateRefId.ToString();
-			this.JobId = TestData.JobId.ToString();
-			this.StepId = TestData.StepId.ToString();
-			this.Change = TestData.Change;
-			this.Key = TestData.Key;
-			this.Data = BsonSerializer.Deserialize<Dictionary<string, object>>(TestData.Data);
+			Id = testData.Id.ToString();
+			StreamId = testData.StreamId.ToString();
+			TemplateRefId = testData.TemplateRefId.ToString();
+			JobId = testData.JobId.ToString();
+			StepId = testData.StepId.ToString();
+			Change = testData.Change;
+			Key = testData.Key;
+			Data = BsonSerializer.Deserialize<Dictionary<string, object>>(testData.Data);
 		}
 	}
 }
