Horde: Fix coding conventions within Horde.Build.

diff --git a/Engine/Source/Programs/Horde/Horde.Build/Api/UgsMetadata.cs b/Engine/Source/Programs/Horde/Horde.Build/Api/UgsMetadata.cs
--- a/Engine/Source/Programs/Horde/Horde.Build/Api/UgsMetadata.cs
+++ b/Engine/Source/Programs/Horde/Horde.Build/Api/UgsMetadata.cs
@@ -1,15 +1,11 @@
 // Copyright Epic Games, Inc. All Rights Reserved.
 
-using EpicGames.Core;
-using Horde.Build.Models;
-using Horde.Build.Services;
-using Horde.Build.Utilities;
 using System;
 using System.Collections.Generic;
 using System.ComponentModel.DataAnnotations;
 using System.Linq;
-using System.Text.Json.Serialization;
-using System.Threading.Tasks;
+using EpicGames.Core;
+using Horde.Build.Models;
 
 namespace Horde.Build.Api
 {
@@ -193,15 +189,15 @@ namespace Horde.Build.Api
 		/// <summary>
 		/// Constructor
 		/// </summary>
-		/// <param name="UserData">The user data object</param>
-		internal GetUgsUserResponse(IUgsUserData UserData)
+		/// <param name="userData">The user data object</param>
+		internal GetUgsUserResponse(IUgsUserData userData)
 		{
-			this.User = UserData.User;
-			this.SyncTime = UserData.SyncTime;
-			this.Vote = (UserData.Vote == UgsUserVote.None) ? (UgsUserVote?)null : UserData.Vote;
-			this.Comment = UserData.Comment;
-			this.Investigating = UserData.Investigating;
-			this.Starred = UserData.Starred;
+			User = userData.User;
+			SyncTime = userData.SyncTime;
+			Vote = (userData.Vote == UgsUserVote.None) ? (UgsUserVote?)null : userData.Vote;
+			Comment = userData.Comment;
+			Investigating = userData.Investigating;
+			Starred = userData.Starred;
 		}
 	}
 
@@ -228,12 +224,12 @@ namespace Horde.Build.Api
 		/// <summary>
 		/// Constructor
 		/// </summary>
-		/// <param name="BadgeData">The badge data object</param>
-		internal GetUgsBadgeResponse(IUgsBadgeData BadgeData)
+		/// <param name="badgeData">The badge data object</param>
+		internal GetUgsBadgeResponse(IUgsBadgeData badgeData)
 		{
-			this.Name = BadgeData.Name;
-			this.Url = BadgeData.Url;
-			this.State = BadgeData.State;
+			Name = badgeData.Name;
+			Url = badgeData.Url;
+			State = badgeData.State;
 		}
 	}
 
@@ -265,13 +261,13 @@ namespace Horde.Build.Api
 		/// <summary>
 		/// Constructor
 		/// </summary>
-		/// <param name="Metadata">Metadata object</param>
-		internal GetUgsMetadataResponse(IUgsMetadata Metadata)
+		/// <param name="metadata">Metadata object</param>
+		internal GetUgsMetadataResponse(IUgsMetadata metadata)
 		{
-			Change = Metadata.Change;
-			Project = Metadata.Project;
-			Users = Metadata.Users?.ConvertAll(x => new GetUgsUserResponse(x));
-			Badges = Metadata.Badges?.ConvertAll(x => new GetUgsBadgeResponse(x));
+			Change = metadata.Change;
+			Project = metadata.Project;
+			Users = metadata.Users?.ConvertAll(x => new GetUgsUserResponse(x));
+			Badges = metadata.Badges?.ConvertAll(x => new GetUgsBadgeResponse(x));
 		}
 	}
 
@@ -370,14 +366,14 @@ namespace Horde.Build.Api
 		/// <summary>
 		/// 
 		/// </summary>
-		/// <param name="Stream"></param>
-		/// <param name="Change"></param>
-		/// <param name="Outcome"></param>
-		public GetUgsIssueBuildResponse(string Stream, int Change, IssueBuildOutcome Outcome)
+		/// <param name="stream"></param>
+		/// <param name="change"></param>
+		/// <param name="outcome"></param>
+		public GetUgsIssueBuildResponse(string stream, int change, IssueBuildOutcome outcome)
 		{
-			this.Stream = Stream;
-			this.Change = Change;
-			this.Outcome = (int)Outcome;
+			Stream = stream;
+			Change = change;
+			Outcome = (int)outcome;
 		}
 	}
 
@@ -404,14 +400,14 @@ namespace Horde.Build.Api
 		/// <summary>
 		/// Constructor
 		/// </summary>
-		/// <param name="BuildId">The corresponding build id</param>
-		/// <param name="Message">Message for the diagnostic</param>
-		/// <param name="Url">Link to the diagnostic</param>
-		public GetUgsIssueDiagnosticResponse(long? BuildId, string Message, Uri Url)
+		/// <param name="buildId">The corresponding build id</param>
+		/// <param name="message">Message for the diagnostic</param>
+		/// <param name="url">Link to the diagnostic</param>
+		public GetUgsIssueDiagnosticResponse(long? buildId, string message, Uri url)
 		{
-			this.BuildId = BuildId;
-			this.Message = Message;
-			this.Url = Url;
+			BuildId = buildId;
+			Message = message;
+			Url = url;
 		}
 	}
 
@@ -478,12 +474,12 @@ namespace Horde.Build.Api
 		/// <summary>
 		/// Whether to notify the user about this issue
 		/// </summary>
-		public bool bNotify { get; set; }
+		public bool BNotify { get; set; }
 
 		/// <summary>
 		/// Whether this issue just contains warnings
 		/// </summary>
-		public bool bWarning { get; set; }
+		public bool BWarning { get; set; }
 
 		/// <summary>
 		/// Link to the last build
@@ -498,27 +494,27 @@ namespace Horde.Build.Api
 		/// <summary>
 		/// Constructs a new issue
 		/// </summary>
-		/// <param name="Details">Issue to construct from</param>
-		/// <param name="Owner"></param>
-		/// <param name="NominatedBy"></param>
+		/// <param name="details">Issue to construct from</param>
+		/// <param name="owner"></param>
+		/// <param name="nominatedBy"></param>
 		/// <param name="bNotify">Whether to notify the user about this issue</param>
-		/// <param name="BuildUrl">Link to the last build for this issue</param>
-		public GetUgsIssueResponse(IIssueDetails Details, IUser? Owner, IUser? NominatedBy, bool bNotify, Uri? BuildUrl)
+		/// <param name="buildUrl">Link to the last build for this issue</param>
+		public GetUgsIssueResponse(IIssueDetails details, IUser? owner, IUser? nominatedBy, bool bNotify, Uri? buildUrl)
 		{
-			IIssue Issue = Details.Issue;
-			this.Id = Issue.Id;
-			this.CreatedAt = Issue.CreatedAt;
-			this.RetrievedAt = DateTime.UtcNow;
-			this.Summary = Issue.Summary;
-			this.Owner = Owner?.Login;
-			this.NominatedBy = NominatedBy?.Login;
-			this.AcknowledgedAt = Issue.AcknowledgedAt;
-			this.FixChange = Issue.FixChange;
-			this.ResolvedAt = Issue.ResolvedAt;
-			this.bNotify = bNotify;
-			this.bWarning = Issue.Severity == IssueSeverity.Warning;
-			this.BuildUrl = BuildUrl;
-			this.Streams = Details.Spans.Select(x => x.StreamName ?? x.StreamId.ToString()).Distinct().ToList();
+			IIssue issue = details.Issue;
+			Id = issue.Id;
+			CreatedAt = issue.CreatedAt;
+			RetrievedAt = DateTime.UtcNow;
+			Summary = issue.Summary;
+			Owner = owner?.Login;
+			NominatedBy = nominatedBy?.Login;
+			AcknowledgedAt = issue.AcknowledgedAt;
+			FixChange = issue.FixChange;
+			ResolvedAt = issue.ResolvedAt;
+			BNotify = bNotify;
+			BWarning = issue.Severity == IssueSeverity.Warning;
+			BuildUrl = buildUrl;
+			Streams = details.Spans.Select(x => x.StreamName ?? x.StreamId.ToString()).Distinct().ToList();
 		}
 	}
 
diff --git a/Engine/Source/Programs/Horde/Horde.Build/Api/Users.cs b/Engine/Source/Programs/Horde/Horde.Build/Api/Users.cs
--- a/Engine/Source/Programs/Horde/Horde.Build/Api/Users.cs
+++ b/Engine/Source/Programs/Horde/Horde.Build/Api/Users.cs
@@ -1,14 +1,12 @@
 // Copyright Epic Games, Inc. All Rights Reserved.
 
-using EpicGames.Core;
-using Horde.Build.Models;
-using Horde.Build.Services;
-using MongoDB.Bson;
 using System;
 using System.Collections.Generic;
 using System.Linq;
 using System.Text.Json;
-using System.Threading.Tasks;
+using EpicGames.Core;
+using Horde.Build.Models;
+using MongoDB.Bson;
 
 namespace Horde.Build.Api
 {
@@ -75,25 +73,25 @@ namespace Horde.Build.Api
 		/// <summary>
 		/// Constructor
 		/// </summary>
-		public GetUserResponse(IUser User, IAvatar? Avatar, IUserClaims? Claims, IUserSettings? Settings)
+		public GetUserResponse(IUser user, IAvatar? avatar, IUserClaims? claims, IUserSettings? settings)
 		{
-			this.Id = User.Id.ToString();
-			this.Name = User.Name;
-			this.Email = User.Email;
-
-			this.Image24 = Avatar?.Image24;
-			this.Image32 = Avatar?.Image32;
-			this.Image48 = Avatar?.Image48;
-			this.Image72 = Avatar?.Image72;
+			Id = user.Id.ToString();
+			Name = user.Name;
+			Email = user.Email;
+
+			Image24 = avatar?.Image24;
+			Image32 = avatar?.Image32;
+			Image48 = avatar?.Image48;
+			Image72 = avatar?.Image72;
 						
-			this.Claims = Claims == null ? null : Claims.Claims.Select(x => new UserClaim(x)).ToList();
+			Claims = claims?.Claims.Select(x => new UserClaim(x)).ToList();
 
-			if (Settings != null)
+			if (settings != null)
 			{
-				this.EnableExperimentalFeatures = Settings.EnableExperimentalFeatures;
+				EnableExperimentalFeatures = settings.EnableExperimentalFeatures;
 				
-				this.DashboardSettings = BsonTypeMapper.MapToDotNetValue(Settings.DashboardSettings);
-				this.PinnedJobIds = Settings.PinnedJobIds.ConvertAll(x => x.ToString());
+				DashboardSettings = BsonTypeMapper.MapToDotNetValue(settings.DashboardSettings);
+				PinnedJobIds = settings.PinnedJobIds.ConvertAll(x => x.ToString());
 			}
 		}
 	}
@@ -126,22 +124,22 @@ namespace Horde.Build.Api
 		/// <summary>
 		/// Constructor
 		/// </summary>
-		/// <param name="User"></param>
-		public GetThinUserInfoResponse(IUser? User)
+		/// <param name="user"></param>
+		public GetThinUserInfoResponse(IUser? user)
 		{
-			if (User == null)
+			if (user == null)
 			{
-				this.Id = String.Empty;
-				this.Name = "(Unknown)";
-				this.Email = null;
-				this.Login = null;
+				Id = String.Empty;
+				Name = "(Unknown)";
+				Email = null;
+				Login = null;
 			}
 			else
 			{
-				this.Id = User.Id.ToString();
-				this.Name = User.Name;
-				this.Email = User.Email;
-				this.Login = User.Login;
+				Id = user.Id.ToString();
+				Name = user.Name;
+				Email = user.Email;
+				Login = user.Login;
 			}
 		}
 	}
@@ -171,5 +169,4 @@ namespace Horde.Build.Api
 		/// </summary>
 		public List<string>? RemovePinnedJobIds { get; set; }
 	}
-
 }
diff --git a/Engine/Source/Programs/Horde/Horde.Build/Auditing/IAuditLog.cs b/Engine/Source/Programs/Horde/Horde.Build/Auditing/IAuditLog.cs
--- a/Engine/Source/Programs/Horde/Horde.Build/Auditing/IAuditLog.cs
+++ b/Engine/Source/Programs/Horde/Horde.Build/Auditing/IAuditLog.cs
@@ -1,26 +1,11 @@
 // Copyright Epic Games, Inc. All Rights Reserved.
 
-using Google.Protobuf.WellKnownTypes;
-using Horde.Build.Api;
-using HordeCommon;
-using Horde.Build.Models;
-using HordeCommon.Rpc.Tasks;
-using Horde.Build.Services;
-using Horde.Build.Utilities;
-using MongoDB.Bson;
-using MongoDB.Bson.Serialization.Attributes;
-using MongoDB.Driver;
 using System;
 using System.Collections.Generic;
-using System.Linq;
-using System.Linq.Expressions;
+using System.IO.Pipelines;
+using System.Text;
 using System.Threading.Tasks;
 using Microsoft.Extensions.Logging;
-using EpicGames.Core;
-using System.Buffers;
-using System.Text.Json;
-using System.Text;
-using System.IO.Pipelines;
 
 namespace Horde.Build.Collections
 {
@@ -53,20 +38,20 @@ namespace Horde.Build.Collections
 		/// <summary>
 		/// Finds messages matching certain criteria
 		/// </summary>
-		/// <param name="MinTime"></param>
-		/// <param name="MaxTime"></param>
-		/// <param name="Index"></param>
-		/// <param name="Count"></param>
+		/// <param name="minTime"></param>
+		/// <param name="maxTime"></param>
+		/// <param name="index"></param>
+		/// <param name="count"></param>
 		/// <returns></returns>
-		IAsyncEnumerable<IAuditLogMessage> FindAsync(DateTime? MinTime = null, DateTime? MaxTime = null, int? Index = null, int? Count = null);
+		IAsyncEnumerable<IAuditLogMessage> FindAsync(DateTime? minTime = null, DateTime? maxTime = null, int? index = null, int? count = null);
 
 		/// <summary>
 		/// Deletes messages between a given time range for a particular object
 		/// </summary>
-		/// <param name="MinTime">Minimum time to remove</param>
-		/// <param name="MaxTime">Maximum time to remove</param>
+		/// <param name="minTime">Minimum time to remove</param>
+		/// <param name="maxTime">Maximum time to remove</param>
 		/// <returns>Async task</returns>
-		Task<long> DeleteAsync(DateTime? MinTime = null, DateTime? MaxTime = null);
+		Task<long> DeleteAsync(DateTime? minTime = null, DateTime? maxTime = null);
 	}
 
 	/// <summary>
@@ -100,9 +85,9 @@ namespace Horde.Build.Collections
 		/// <summary>
 		/// Get the channel for a particular subject
 		/// </summary>
-		/// <param name="Subject"></param>
+		/// <param name="subject"></param>
 		/// <returns></returns>
-		IAuditLogChannel<TSubject> this[TSubject Subject] { get; }
+		IAuditLogChannel<TSubject> this[TSubject subject] { get; }
 	}
 
 	/// <summary>
@@ -114,10 +99,10 @@ namespace Horde.Build.Collections
 		/// <summary>
 		/// Create a new audit log instance, with the given database name
 		/// </summary>
-		/// <param name="CollectionName"></param>
-		/// <param name="SubjectProperty"></param>
+		/// <param name="collectionName"></param>
+		/// <param name="subjectProperty"></param>
 		/// <returns></returns>
-		IAuditLog<TSubject> Create(string CollectionName, string SubjectProperty);
+		IAuditLog<TSubject> Create(string collectionName, string subjectProperty);
 	}
 
 	/// <summary>
@@ -128,28 +113,28 @@ namespace Horde.Build.Collections
 		/// <summary>
 		/// Retrieve historical information about a specific agent
 		/// </summary>
-		/// <param name="Channel">Channel to query</param>
-		/// <param name="BodyWriter">Writer for Json data</param>
-		/// <param name="MinTime">Minimum time for records to return</param>
-		/// <param name="MaxTime">Maximum time for records to return</param>
-		/// <param name="Index">Offset of the first result</param>
-		/// <param name="Count">Number of records to return</param>
+		/// <param name="channel">Channel to query</param>
+		/// <param name="bodyWriter">Writer for Json data</param>
+		/// <param name="minTime">Minimum time for records to return</param>
+		/// <param name="maxTime">Maximum time for records to return</param>
+		/// <param name="index">Offset of the first result</param>
+		/// <param name="count">Number of records to return</param>
 		/// <returns>Information about the requested agent</returns>
-		public static async Task FindAsync<T>(this IAuditLogChannel<T> Channel, PipeWriter BodyWriter, DateTime? MinTime = null, DateTime? MaxTime = null, int Index = 0, int Count = 50)
+		public static async Task FindAsync<T>(this IAuditLogChannel<T> channel, PipeWriter bodyWriter, DateTime? minTime = null, DateTime? maxTime = null, int index = 0, int count = 50)
 		{
-			string Prefix = "{\n\t\"entries\":\n\t[";
-			await BodyWriter.WriteAsync(Encoding.UTF8.GetBytes(Prefix));
+			string prefix = "{\n\t\"entries\":\n\t[";
+			await bodyWriter.WriteAsync(Encoding.UTF8.GetBytes(prefix));
 
-			string Separator = "";
-			await foreach (IAuditLogMessage<T> Message in Channel.FindAsync(MinTime, MaxTime, Index, Count))
+			string separator = "";
+			await foreach (IAuditLogMessage<T> message in channel.FindAsync(minTime, maxTime, index, count))
 			{
-				string Line = $"{Separator}\n\t\t{Message.Data}";
-				await BodyWriter.WriteAsync(Encoding.UTF8.GetBytes(Line));
-				Separator = ",";
+				string line = $"{separator}\n\t\t{message.Data}";
+				await bodyWriter.WriteAsync(Encoding.UTF8.GetBytes(line));
+				separator = ",";
 			}
 
-			string Suffix = "\n\t]\n}";
-			await BodyWriter.WriteAsync(Encoding.UTF8.GetBytes(Suffix));
+			string suffix = "\n\t]\n}";
+			await bodyWriter.WriteAsync(Encoding.UTF8.GetBytes(suffix));
 		}
 	}
 }
diff --git a/Engine/Source/Programs/Horde/Horde.Build/Auditing/Impl/AuditLog.cs b/Engine/Source/Programs/Horde/Horde.Build/Auditing/Impl/AuditLog.cs
--- a/Engine/Source/Programs/Horde/Horde.Build/Auditing/Impl/AuditLog.cs
+++ b/Engine/Source/Programs/Horde/Horde.Build/Auditing/Impl/AuditLog.cs
@@ -1,24 +1,16 @@
 // Copyright Epic Games, Inc. All Rights Reserved.
 
+using System;
+using System.Collections.Generic;
+using System.Threading.Channels;
+using System.Threading.Tasks;
 using EpicGames.Core;
-using Google.Protobuf.WellKnownTypes;
-using HordeCommon;
-using HordeCommon.Rpc.Tasks;
-using Horde.Build.Models;
 using Horde.Build.Services;
 using Horde.Build.Utilities;
 using Microsoft.Extensions.Logging;
 using MongoDB.Bson;
 using MongoDB.Bson.Serialization.Attributes;
 using MongoDB.Driver;
-using System;
-using System.Buffers;
-using System.Collections.Generic;
-using System.Linq;
-using System.Linq.Expressions;
-using System.Text.Json;
-using System.Threading.Channels;
-using System.Threading.Tasks;
 
 namespace Horde.Build.Collections.Impl
 {
@@ -46,13 +38,13 @@ namespace Horde.Build.Collections.Impl
 				Data = String.Empty;
 			}
 
-			public AuditLogMessage(TSubject Subject, DateTime TimeUtc, LogLevel Level, string Data)
+			public AuditLogMessage(TSubject subject, DateTime timeUtc, LogLevel level, string data)
 			{
-				this.Id = ObjectId.GenerateNewId();
-				this.Subject = Subject;
-				this.TimeUtc = TimeUtc;
-				this.Level = Level;
-				this.Data = Data;
+				Id = ObjectId.GenerateNewId();
+				Subject = subject;
+				TimeUtc = timeUtc;
+				Level = level;
+				Data = data;
 			}
 		}
 
@@ -66,65 +58,65 @@ namespace Horde.Build.Collections.Impl
 			public readonly AuditLog<TSubject> Outer;
 			public TSubject Subject { get; }
 
-			public AuditLogChannel(AuditLog<TSubject> Outer, TSubject Subject)
+			public AuditLogChannel(AuditLog<TSubject> outer, TSubject subject)
 			{
-				this.Outer = Outer;
-				this.Subject = Subject;
+				Outer = outer;
+				Subject = subject;
 			}
 
 			public IDisposable BeginScope<TState>(TState state) => new Scope();
 
 			public bool IsEnabled(LogLevel logLevel) => true;
 
-			public void Log<TState>(LogLevel LogLevel, EventId EventId, TState State, Exception? Exception, Func<TState, Exception?, string> Formatter)
+			public void Log<TState>(LogLevel logLevel, EventId eventId, TState state, Exception? exception, Func<TState, Exception?, string> formatter)
 			{
-				DateTime Time = DateTime.UtcNow;
+				DateTime time = DateTime.UtcNow;
 
-				LogEvent Event = LogEvent.FromState(LogLevel, EventId, State, Exception, Formatter);
-				string Data = Event.ToJson();
-				AuditLogMessage Message = new AuditLogMessage(Subject, Time, LogLevel, Data);
-				Outer.MessageChannel.Writer.TryWrite(Message);
+				LogEvent logEvent = LogEvent.FromState(logLevel, eventId, state, exception, formatter);
+				string data = logEvent.ToJson();
+				AuditLogMessage message = new AuditLogMessage(Subject, time, logLevel, data);
+				Outer._messageChannel.Writer.TryWrite(message);
 
 #pragma warning disable CA2254 // Template should be a static expression
-				using (IDisposable _ = Outer.Logger.BeginScope($"Subject: {{{Outer.SubjectProperty}}}", Subject))
+				using (IDisposable _ = Outer._logger.BeginScope($"Subject: {{{Outer._subjectProperty}}}", Subject))
 				{
-					Outer.Logger.Log(LogLevel, EventId, State, Exception, Formatter);
+					Outer._logger.Log(logLevel, eventId, state, exception, formatter);
 				}
 #pragma warning restore CA2254 // Template should be a static expression
 			}
 
-			public IAsyncEnumerable<IAuditLogMessage> FindAsync(DateTime? MinTime, DateTime? MaxTime, int? Index, int? Count) => Outer.FindAsync(Subject, MinTime, MaxTime, Index, Count);
+			public IAsyncEnumerable<IAuditLogMessage> FindAsync(DateTime? minTime, DateTime? maxTime, int? index, int? count) => Outer.FindAsync(Subject, minTime, maxTime, index, count);
 
-			public Task<long> DeleteAsync(DateTime? MinTime, DateTime? MaxTime) => Outer.DeleteAsync(Subject, MinTime, MaxTime);
+			public Task<long> DeleteAsync(DateTime? minTime, DateTime? maxTime) => Outer.DeleteAsync(Subject, minTime, maxTime);
 		}
 
-		IMongoCollection<AuditLogMessage> Messages;
-		Channel<AuditLogMessage> MessageChannel;
-		string SubjectProperty;
-		ILogger Logger;
-		Task BackgroundTask;
+		readonly IMongoCollection<AuditLogMessage> _messages;
+		readonly Channel<AuditLogMessage> _messageChannel;
+		readonly string _subjectProperty;
+		readonly ILogger _logger;
+		readonly Task _backgroundTask;
 
-		public IAuditLogChannel<TSubject> this[TSubject Subject] => new AuditLogChannel(this, Subject);
+		public IAuditLogChannel<TSubject> this[TSubject subject] => new AuditLogChannel(this, subject);
 
-		public AuditLog(DatabaseService DatabaseService, string CollectionName, string SubjectProperty, ILogger Logger)
+		public AuditLog(DatabaseService databaseService, string collectionName, string subjectProperty, ILogger logger)
 		{
-			this.Messages = DatabaseService.GetCollection<AuditLogMessage>(CollectionName);
-			this.MessageChannel = Channel.CreateUnbounded<AuditLogMessage>();
-			this.SubjectProperty = SubjectProperty;
-			this.Logger = Logger;
+			_messages = databaseService.GetCollection<AuditLogMessage>(collectionName);
+			_messageChannel = Channel.CreateUnbounded<AuditLogMessage>();
+			_subjectProperty = subjectProperty;
+			_logger = logger;
 
-			if (!DatabaseService.ReadOnlyMode)
+			if (!databaseService.ReadOnlyMode)
 			{
-				Messages.Indexes.CreateOne(new CreateIndexModel<AuditLogMessage>(Builders<AuditLogMessage>.IndexKeys.Ascending(x => x.Id).Descending(x => x.TimeUtc)));
+				_messages.Indexes.CreateOne(new CreateIndexModel<AuditLogMessage>(Builders<AuditLogMessage>.IndexKeys.Ascending(x => x.Id).Descending(x => x.TimeUtc)));
 			}
 
-			BackgroundTask = Task.Run(() => WriteMessagesAsync());
+			_backgroundTask = Task.Run(() => WriteMessagesAsync());
 		}
 
 		public async ValueTask DisposeAsync()
 		{
-			MessageChannel.Writer.TryComplete();
-			await BackgroundTask;
+			_messageChannel.Writer.TryComplete();
+			await _backgroundTask;
 		}
 
 		public void Dispose()
@@ -134,75 +126,75 @@ namespace Horde.Build.Collections.Impl
 
 		async Task WriteMessagesAsync()
 		{
-			while (await MessageChannel.Reader.WaitToReadAsync())
+			while (await _messageChannel.Reader.WaitToReadAsync())
 			{
-				List<AuditLogMessage> NewMessages = new List<AuditLogMessage>();
-				while (MessageChannel.Reader.TryRead(out AuditLogMessage? NewMessage))
+				List<AuditLogMessage> newMessages = new List<AuditLogMessage>();
+				while (_messageChannel.Reader.TryRead(out AuditLogMessage? newMessage))
 				{
-					NewMessages.Add(NewMessage);
+					newMessages.Add(newMessage);
 				}
-				if (NewMessages.Count > 0)
+				if (newMessages.Count > 0)
 				{
-					await Messages.InsertManyAsync(NewMessages);
+					await _messages.InsertManyAsync(newMessages);
 				}
 			}
 		}
 
-		async IAsyncEnumerable<IAuditLogMessage<TSubject>> FindAsync(TSubject Subject, DateTime? MinTime = null, DateTime? MaxTime = null, int? Index = null, int? Count = null)
+		async IAsyncEnumerable<IAuditLogMessage<TSubject>> FindAsync(TSubject subject, DateTime? minTime = null, DateTime? maxTime = null, int? index = null, int? count = null)
 		{
-			FilterDefinition<AuditLogMessage> Filter = Builders<AuditLogMessage>.Filter.Eq(x => x.Subject, Subject);
-			if (MinTime != null)
+			FilterDefinition<AuditLogMessage> filter = Builders<AuditLogMessage>.Filter.Eq(x => x.Subject, subject);
+			if (minTime != null)
 			{
-				Filter &= Builders<AuditLogMessage>.Filter.Gte(x => x.TimeUtc, MinTime.Value);
+				filter &= Builders<AuditLogMessage>.Filter.Gte(x => x.TimeUtc, minTime.Value);
 			}
-			if (MaxTime != null)
+			if (maxTime != null)
 			{
-				Filter &= Builders<AuditLogMessage>.Filter.Lte(x => x.TimeUtc, MaxTime.Value);
+				filter &= Builders<AuditLogMessage>.Filter.Lte(x => x.TimeUtc, maxTime.Value);
 			}
 
-			using (IAsyncCursor<AuditLogMessage> Cursor = await Messages.Find(Filter).SortByDescending(x => x.TimeUtc).Range(Index, Count).ToCursorAsync())
+			using (IAsyncCursor<AuditLogMessage> cursor = await _messages.Find(filter).SortByDescending(x => x.TimeUtc).Range(index, count).ToCursorAsync())
 			{
-				while (await Cursor.MoveNextAsync())
+				while (await cursor.MoveNextAsync())
 				{
-					foreach (AuditLogMessage Message in Cursor.Current)
+					foreach (AuditLogMessage message in cursor.Current)
 					{
-						yield return Message;
+						yield return message;
 					}
 				}
 			}
 		}
 
-		async Task<long> DeleteAsync(TSubject Subject, DateTime? MinTime = null, DateTime? MaxTime = null)
+		async Task<long> DeleteAsync(TSubject subject, DateTime? minTime = null, DateTime? maxTime = null)
 		{
-			FilterDefinition<AuditLogMessage> Filter = Builders<AuditLogMessage>.Filter.Eq(x => x.Subject, Subject);
-			if (MinTime != null)
+			FilterDefinition<AuditLogMessage> filter = Builders<AuditLogMessage>.Filter.Eq(x => x.Subject, subject);
+			if (minTime != null)
 			{
-				Filter &= Builders<AuditLogMessage>.Filter.Gte(x => x.TimeUtc, MinTime.Value);
+				filter &= Builders<AuditLogMessage>.Filter.Gte(x => x.TimeUtc, minTime.Value);
 			}
-			if (MaxTime != null)
+			if (maxTime != null)
 			{
-				Filter &= Builders<AuditLogMessage>.Filter.Lte(x => x.TimeUtc, MaxTime.Value);
+				filter &= Builders<AuditLogMessage>.Filter.Lte(x => x.TimeUtc, maxTime.Value);
 			}
 
-			DeleteResult Result = await Messages.DeleteManyAsync(Filter);
-			return Result.DeletedCount;
+			DeleteResult result = await _messages.DeleteManyAsync(filter);
+			return result.DeletedCount;
 		}
 	}
 
 	class AuditLogFactory<TSubject> : IAuditLogFactory<TSubject>
 	{
-		DatabaseService DatabaseService;
-		ILogger<AuditLog<TSubject>> Logger;
+		readonly DatabaseService _databaseService;
+		readonly ILogger<AuditLog<TSubject>> _logger;
 
-		public AuditLogFactory(DatabaseService DatabaseService, ILogger<AuditLog<TSubject>> Logger)
+		public AuditLogFactory(DatabaseService databaseService, ILogger<AuditLog<TSubject>> logger)
 		{
-			this.DatabaseService = DatabaseService;
-			this.Logger = Logger;
+			_databaseService = databaseService;
+			_logger = logger;
 		}
 
-		public IAuditLog<TSubject> Create(string CollectionName, string SubjectProperty)
+		public IAuditLog<TSubject> Create(string collectionName, string subjectProperty)
 		{
-			return new AuditLog<TSubject>(DatabaseService, CollectionName, SubjectProperty, Logger);
+			return new AuditLog<TSubject>(_databaseService, collectionName, subjectProperty, _logger);
 		}
 	}
 }
diff --git a/Engine/Source/Programs/Horde/Horde.Build/Authentication/AnonymousHandler.cs b/Engine/Source/Programs/Horde/Horde.Build/Authentication/AnonymousHandler.cs
--- a/Engine/Source/Programs/Horde/Horde.Build/Authentication/AnonymousHandler.cs
+++ b/Engine/Source/Programs/Horde/Horde.Build/Authentication/AnonymousHandler.cs
@@ -1,18 +1,16 @@
 // Copyright Epic Games, Inc. All Rights Reserved.
 
+using System;
+using System.Collections.Generic;
+using System.Security.Claims;
+using System.Text.Encodings.Web;
+using System.Threading.Tasks;
 using Horde.Build.Collections;
 using Horde.Build.Models;
-using Horde.Build.Services;
 using Horde.Build.Utilities;
 using Microsoft.AspNetCore.Authentication;
 using Microsoft.Extensions.Logging;
 using Microsoft.Extensions.Options;
-using System;
-using System.Collections.Generic;
-using System.Linq;
-using System.Security.Claims;
-using System.Text.Encodings.Web;
-using System.Threading.Tasks;
 
 namespace Horde.Build.Authentication
 {
@@ -25,43 +23,42 @@ namespace Horde.Build.Authentication
 	class AnonymousAuthenticationHandler : AuthenticationHandler<AnonymousAuthenticationOptions>
 	{
 		public const string AuthenticationScheme = "Anonymous";
+		readonly IUserCollection _userCollection;
 
-		IUserCollection UserCollection;
-
-		public AnonymousAuthenticationHandler(IUserCollection UserCollection, IOptionsMonitor<AnonymousAuthenticationOptions> Options,
-			ILoggerFactory Logger, UrlEncoder Encoder, ISystemClock Clock)
-			: base(Options, Logger, Encoder, Clock)
+		public AnonymousAuthenticationHandler(IUserCollection userCollection, IOptionsMonitor<AnonymousAuthenticationOptions> options,
+			ILoggerFactory logger, UrlEncoder encoder, ISystemClock clock)
+			: base(options, logger, encoder, clock)
 		{
-			this.UserCollection = UserCollection;
+			_userCollection = userCollection;
 		}
 
 		protected override async Task<AuthenticateResult> HandleAuthenticateAsync()
 		{
-			List<Claim> Claims = new List<Claim>();
-			Claims.Add(new Claim(ClaimTypes.Name, AuthenticationScheme));
+			List<Claim> claims = new List<Claim>();
+			claims.Add(new Claim(ClaimTypes.Name, AuthenticationScheme));
 
 			if (Options.AdminClaimType != null && Options.AdminClaimValue != null)
 			{
-				Claims.Add(new Claim(Options.AdminClaimType, Options.AdminClaimValue));
+				claims.Add(new Claim(Options.AdminClaimType, Options.AdminClaimValue));
 			}
 
-			IUser User = await UserCollection.FindOrAddUserByLoginAsync("anonymous", "Anonymous", "[REDACTED_EMAIL]");
+			IUser user = await _userCollection.FindOrAddUserByLoginAsync("anonymous", "Anonymous", "[REDACTED_EMAIL]");
 
-			ClaimsIdentity Identity = new ClaimsIdentity(Claims, Scheme.Name);
-			Identity.AddClaim(new Claim(HordeClaimTypes.UserId, User.Id.ToString()));
+			ClaimsIdentity identity = new ClaimsIdentity(claims, Scheme.Name);
+			identity.AddClaim(new Claim(HordeClaimTypes.UserId, user.Id.ToString()));
 
-			ClaimsPrincipal Principal = new ClaimsPrincipal(Identity);
-			AuthenticationTicket Ticket = new AuthenticationTicket(Principal, Scheme.Name);
+			ClaimsPrincipal principal = new ClaimsPrincipal(identity);
+			AuthenticationTicket ticket = new AuthenticationTicket(principal, Scheme.Name);
 
-			return AuthenticateResult.Success(Ticket);
+			return AuthenticateResult.Success(ticket);
 		}
 	}
 
 	static class AnonymousExtensions
 	{
-		public static AuthenticationBuilder AddAnonymous(this AuthenticationBuilder Builder, Action<AnonymousAuthenticationOptions> Configure)
+		public static AuthenticationBuilder AddAnonymous(this AuthenticationBuilder builder, Action<AnonymousAuthenticationOptions> configure)
 		{
-			return Builder.AddScheme<AnonymousAuthenticationOptions, AnonymousAuthenticationHandler>(AnonymousAuthenticationHandler.AuthenticationScheme, Configure);
+			return builder.AddScheme<AnonymousAuthenticationOptions, AnonymousAuthenticationHandler>(AnonymousAuthenticationHandler.AuthenticationScheme, configure);
 		}
 	}
 }
diff --git a/Engine/Source/Programs/Horde/Horde.Build/Authentication/HordeJwtBearerHandler.cs b/Engine/Source/Programs/Horde/Horde.Build/Authentication/HordeJwtBearerHandler.cs
--- a/Engine/Source/Programs/Horde/Horde.Build/Authentication/HordeJwtBearerHandler.cs
+++ b/Engine/Source/Programs/Horde/Horde.Build/Authentication/HordeJwtBearerHandler.cs
@@ -1,17 +1,15 @@
 // Copyright Epic Games, Inc. All Rights Reserved.
 
+using System;
+using System.IdentityModel.Tokens.Jwt;
+using System.Text.Encodings.Web;
+using System.Threading.Tasks;
 using Horde.Build.Services;
 using Horde.Build.Utilities;
 using Microsoft.AspNetCore.Authentication;
 using Microsoft.AspNetCore.Authentication.JwtBearer;
 using Microsoft.Extensions.Logging;
 using Microsoft.Extensions.Options;
-using System;
-using System.Collections.Generic;
-using System.IdentityModel.Tokens.Jwt;
-using System.Linq;
-using System.Text.Encodings.Web;
-using System.Threading.Tasks;
 
 namespace Horde.Build.Authentication
 {
@@ -24,48 +22,47 @@ namespace Horde.Build.Authentication
 		/// Default name of the authentication scheme
 		/// </summary>
 		public const string AuthenticationScheme = "ServerJwt";
+		readonly DatabaseService _databaseService;
 
-		DatabaseService DatabaseService;
-
-		public HordeJwtBearerHandler(ILoggerFactory Logger, UrlEncoder Encoder, ISystemClock Clock, DatabaseService DatabaseService, IOptionsMonitorCache<JwtBearerOptions> OptionsCache)
-			: base(GetOptionsMonitor(DatabaseService, OptionsCache), Logger, Encoder, Clock)
+		public HordeJwtBearerHandler(ILoggerFactory logger, UrlEncoder encoder, ISystemClock clock, DatabaseService databaseService, IOptionsMonitorCache<JwtBearerOptions> optionsCache)
+			: base(GetOptionsMonitor(databaseService, optionsCache), logger, encoder, clock)
 		{
-			this.DatabaseService = DatabaseService;
+			_databaseService = databaseService;
 		}
 
-		private static IOptionsMonitor<JwtBearerOptions> GetOptionsMonitor(DatabaseService DatabaseService, IOptionsMonitorCache<JwtBearerOptions> OptionsCache)
+		private static IOptionsMonitor<JwtBearerOptions> GetOptionsMonitor(DatabaseService databaseService, IOptionsMonitorCache<JwtBearerOptions> optionsCache)
 		{
-			ConfigureNamedOptions<JwtBearerOptions> NamedOptions = new ConfigureNamedOptions<JwtBearerOptions>(AuthenticationScheme, Options => Configure(Options, DatabaseService));
-			OptionsFactory<JwtBearerOptions> OptionsFactory = new OptionsFactory<JwtBearerOptions>(new[] { NamedOptions }, Array.Empty<IPostConfigureOptions<JwtBearerOptions>>());
-			return new OptionsMonitor<JwtBearerOptions>(OptionsFactory, Array.Empty<IOptionsChangeTokenSource<JwtBearerOptions>>(), OptionsCache);
+			ConfigureNamedOptions<JwtBearerOptions> namedOptions = new ConfigureNamedOptions<JwtBearerOptions>(AuthenticationScheme, options => Configure(options, databaseService));
+			OptionsFactory<JwtBearerOptions> optionsFactory = new OptionsFactory<JwtBearerOptions>(new[] { namedOptions }, Array.Empty<IPostConfigureOptions<JwtBearerOptions>>());
+			return new OptionsMonitor<JwtBearerOptions>(optionsFactory, Array.Empty<IOptionsChangeTokenSource<JwtBearerOptions>>(), optionsCache);
 		}
 
-		private static void Configure(JwtBearerOptions Options, DatabaseService DatabaseService)
+		private static void Configure(JwtBearerOptions options, DatabaseService databaseService)
 		{
-			Options.TokenValidationParameters.ValidateAudience = false;
+			options.TokenValidationParameters.ValidateAudience = false;
 
-			Options.TokenValidationParameters.RequireExpirationTime = false;
-			Options.TokenValidationParameters.ValidateLifetime = true;
+			options.TokenValidationParameters.RequireExpirationTime = false;
+			options.TokenValidationParameters.ValidateLifetime = true;
 
-			Options.TokenValidationParameters.ValidIssuer = DatabaseService.JwtIssuer;
-			Options.TokenValidationParameters.ValidateIssuer = true;
+			options.TokenValidationParameters.ValidIssuer = databaseService.JwtIssuer;
+			options.TokenValidationParameters.ValidateIssuer = true;
 
-			Options.TokenValidationParameters.ValidateIssuerSigningKey = true;
-			Options.TokenValidationParameters.IssuerSigningKey = DatabaseService.JwtSigningKey;
+			options.TokenValidationParameters.ValidateIssuerSigningKey = true;
+			options.TokenValidationParameters.IssuerSigningKey = databaseService.JwtSigningKey;
 		}
 
 		protected override Task<AuthenticateResult> HandleAuthenticateAsync()
 		{
 			// Silent fail if this JWT is not issued by the server
-			string? Token;
-			if (!JwtUtils.TryGetBearerToken(Request, "Bearer ", out Token))
+			string? token;
+			if (!JwtUtils.TryGetBearerToken(Request, "Bearer ", out token))
 			{
 				return Task.FromResult(AuthenticateResult.NoResult());
 			}
 
 			// Validate that it's from the correct issuer
-			JwtSecurityToken? JwtToken;
-			if (!JwtUtils.TryParseJwt(Token, out JwtToken) || !String.Equals(JwtToken.Issuer, DatabaseService.JwtIssuer, StringComparison.Ordinal))
+			JwtSecurityToken? jwtToken;
+			if (!JwtUtils.TryParseJwt(token, out jwtToken) || !String.Equals(jwtToken.Issuer, _databaseService.JwtIssuer, StringComparison.Ordinal))
 			{
 				return Task.FromResult(AuthenticateResult.NoResult());
 			}
diff --git a/Engine/Source/Programs/Horde/Horde.Build/Authentication/HordeOpenIdConnectHandler.cs b/Engine/Source/Programs/Horde/Horde.Build/Authentication/HordeOpenIdConnectHandler.cs
--- a/Engine/Source/Programs/Horde/Horde.Build/Authentication/HordeOpenIdConnectHandler.cs
+++ b/Engine/Source/Programs/Horde/Horde.Build/Authentication/HordeOpenIdConnectHandler.cs
@@ -1,5 +1,11 @@
 // Copyright Epic Games, Inc. All Rights Reserved.
 
+using System;
+using System.Linq;
+using System.Security.Claims;
+using System.Text.Encodings.Web;
+using System.Text.Json;
+using System.Threading.Tasks;
 using Horde.Build.Collections;
 using Horde.Build.Models;
 using Horde.Build.Utilities;
@@ -11,23 +17,17 @@ using Microsoft.Extensions.DependencyInjection.Extensions;
 using Microsoft.Extensions.Logging;
 using Microsoft.Extensions.Options;
 using Microsoft.IdentityModel.Protocols.OpenIdConnect;
-using System;
-using System.Linq;
-using System.Security.Claims;
-using System.Text.Encodings.Web;
-using System.Text.Json;
-using System.Threading.Tasks;
 
 namespace Horde.Build.Authentication
 {
 	class HordeOpenIdConnectHandler : OpenIdConnectHandler
 	{
-		IUserCollection UserCollection;
+		readonly IUserCollection _userCollection;
 
-		public HordeOpenIdConnectHandler(IOptionsMonitor<OpenIdConnectOptions> Options, ILoggerFactory Logger, HtmlEncoder HtmlEncoder, UrlEncoder Encoder, ISystemClock Clock, IUserCollection UserCollection)
-			: base(Options, Logger, HtmlEncoder, Encoder, Clock)
+		public HordeOpenIdConnectHandler(IOptionsMonitor<OpenIdConnectOptions> options, ILoggerFactory logger, HtmlEncoder htmlEncoder, UrlEncoder encoder, ISystemClock clock, IUserCollection userCollection)
+			: base(options, logger, htmlEncoder, encoder, clock)
 		{
-			this.UserCollection = UserCollection;
+			_userCollection = userCollection;
 		}
 
 		/// <summary>
@@ -35,44 +35,44 @@ namespace Horde.Build.Authentication
 		///
 		/// If the claim is already set, it is skipped.
 		/// </summary>
-		/// <param name="UserInfo">User info from the OIDC /userinfo endpoint </param>
-		/// <param name="Identity">Identity which claims to modify</param>
-		/// <param name="ClaimName">Name of claim</param>
-		/// <param name="UserInfoFields">List of field names to try when mapping</param>
+		/// <param name="userInfo">User info from the OIDC /userinfo endpoint </param>
+		/// <param name="identity">Identity which claims to modify</param>
+		/// <param name="claimName">Name of claim</param>
+		/// <param name="userInfoFields">List of field names to try when mapping</param>
 		/// <exception cref="Exception">Thrown if no field name matches</exception>
-		private static void MapUserInfoFieldToClaim(JsonElement UserInfo, ClaimsIdentity Identity, string ClaimName, string[] UserInfoFields)
+		private static void MapUserInfoFieldToClaim(JsonElement userInfo, ClaimsIdentity identity, string claimName, string[] userInfoFields)
 		{
-			if (Identity.FindFirst(ClaimName) != null)
+			if (identity.FindFirst(claimName) != null)
 			{
 				return;
 			}
 
-			foreach (string UserDataField in UserInfoFields)
+			foreach (string userDataField in userInfoFields)
 			{
-				if (UserInfo.TryGetProperty(UserDataField, out JsonElement FieldElement))
+				if (userInfo.TryGetProperty(userDataField, out JsonElement fieldElement))
 				{
-					Identity.AddClaim(new Claim(ClaimName, FieldElement.ToString()));
+					identity.AddClaim(new Claim(claimName, fieldElement.ToString()));
 					return;
 				}
 			}
 
-			string Message = $"Unable to map a field from user info to claim '{ClaimName}' using list [{string.Join(", ", UserInfoFields)}].";
-			Message += " UserInfo: " + UserInfo;
-			throw new Exception(Message);
+			string message = $"Unable to map a field from user info to claim '{claimName}' using list [{String.Join(", ", userInfoFields)}].";
+			message += " UserInfo: " + userInfo;
+			throw new Exception(message);
 		}
 
-		public static void AddUserInfoClaims(ServerSettings Settings, JsonElement UserInfo, ClaimsIdentity Identity)
+		public static void AddUserInfoClaims(ServerSettings settings, JsonElement userInfo, ClaimsIdentity identity)
 		{
-			MapUserInfoFieldToClaim(UserInfo, Identity, ClaimTypes.Name, Settings.OidcClaimNameMapping);
-			MapUserInfoFieldToClaim(UserInfo, Identity, ClaimTypes.Email, Settings.OidcClaimEmailMapping);
-			MapUserInfoFieldToClaim(UserInfo, Identity, HordeClaimTypes.User, Settings.OidcClaimHordeUserMapping);
-			MapUserInfoFieldToClaim(UserInfo, Identity, HordeClaimTypes.PerforceUser, Settings.OidcClaimHordePerforceUserMapping);
+			MapUserInfoFieldToClaim(userInfo, identity, ClaimTypes.Name, settings.OidcClaimNameMapping);
+			MapUserInfoFieldToClaim(userInfo, identity, ClaimTypes.Email, settings.OidcClaimEmailMapping);
+			MapUserInfoFieldToClaim(userInfo, identity, HordeClaimTypes.User, settings.OidcClaimHordeUserMapping);
+			MapUserInfoFieldToClaim(userInfo, identity, HordeClaimTypes.PerforceUser, settings.OidcClaimHordePerforceUserMapping);
 
-			if (UserInfo.TryGetProperty("groups", out JsonElement GroupsElement) && GroupsElement.ValueKind == JsonValueKind.Array)
+			if (userInfo.TryGetProperty("groups", out JsonElement groupsElement) && groupsElement.ValueKind == JsonValueKind.Array)
 			{
-				for (int Idx = 0; Idx < GroupsElement.GetArrayLength(); Idx++)
+				for (int idx = 0; idx < groupsElement.GetArrayLength(); idx++)
 				{
-					Identity.AddClaim(new Claim(ClaimTypes.Role, GroupsElement[Idx].ToString()!));
+					identity.AddClaim(new Claim(ClaimTypes.Role, groupsElement[idx].ToString()!));
 				}
 			}
 		}
@@ -80,28 +80,28 @@ namespace Horde.Build.Authentication
 		protected override async Task<HandleRequestResult> HandleRemoteAuthenticateAsync()
 		{
 			// Authenticate with the OIDC provider
-			HandleRequestResult Result = await base.HandleRemoteAuthenticateAsync();
-			if (!Result.Succeeded)
+			HandleRequestResult result = await base.HandleRemoteAuthenticateAsync();
+			if (!result.Succeeded)
 			{
-				return Result;
+				return result;
 			}
 
-			ClaimsIdentity? Identity = (ClaimsIdentity?)Result.Principal?.Identity;
-			if (Identity == null)
+			ClaimsIdentity? identity = (ClaimsIdentity?)result.Principal?.Identity;
+			if (identity == null)
 			{
 				return HandleRequestResult.Fail("No identity specified");
 			}
 
-			string Login = Identity.FindFirst(ClaimTypes.Name)!.Value;
-			string? Name = Identity.FindFirst("name")?.Value;
-			string? Email = Identity.FindFirst(ClaimTypes.Email)?.Value;
+			string login = identity.FindFirst(ClaimTypes.Name)!.Value;
+			string? name = identity.FindFirst("name")?.Value;
+			string? email = identity.FindFirst(ClaimTypes.Email)?.Value;
 
-			IUser User = await UserCollection.FindOrAddUserByLoginAsync(Login, Name, Email);
-			Identity.AddClaim(new Claim(HordeClaimTypes.UserId, User.Id.ToString()));
+			IUser user = await _userCollection.FindOrAddUserByLoginAsync(login, name, email);
+			identity.AddClaim(new Claim(HordeClaimTypes.UserId, user.Id.ToString()));
 
-			await UserCollection.UpdateClaimsAsync(User.Id, Identity.Claims.Select(x => new UserClaim(x.Type, x.Value)));
+			await _userCollection.UpdateClaimsAsync(user.Id, identity.Claims.Select(x => new UserClaim(x.Type, x.Value)));
 
-			return Result;
+			return result;
 		}
 	}
 
@@ -109,26 +109,26 @@ namespace Horde.Build.Authentication
 	{
 		class MapRolesClaimAction : ClaimAction
 		{
-			private ServerSettings Settings;
-			public MapRolesClaimAction(ServerSettings Settings) : base(ClaimTypes.Role, ClaimTypes.Role) { this.Settings = Settings; }
-			public override void Run(JsonElement UserData, ClaimsIdentity Identity, string Issuer) { HordeOpenIdConnectHandler.AddUserInfoClaims(Settings, UserData, Identity); }
+			private readonly ServerSettings _settings;
+			public MapRolesClaimAction(ServerSettings settings) : base(ClaimTypes.Role, ClaimTypes.Role) { _settings = settings; }
+			public override void Run(JsonElement userData, ClaimsIdentity identity, string issuer) { HordeOpenIdConnectHandler.AddUserInfoClaims(_settings, userData, identity); }
 		}
 
-		static void ApplyDefaultOptions(ServerSettings Settings, OpenIdConnectOptions Options, Action<OpenIdConnectOptions> Handler)
+		static void ApplyDefaultOptions(ServerSettings settings, OpenIdConnectOptions options, Action<OpenIdConnectOptions> handler)
 		{
-			Options.ResponseType = OpenIdConnectResponseType.Code;
-			Options.GetClaimsFromUserInfoEndpoint = true;
-			Options.SaveTokens = true;
-			Options.TokenValidationParameters.NameClaimType = "name";
-			Options.ClaimActions.Add(new MapRolesClaimAction(Settings));
+			options.ResponseType = OpenIdConnectResponseType.Code;
+			options.GetClaimsFromUserInfoEndpoint = true;
+			options.SaveTokens = true;
+			options.TokenValidationParameters.NameClaimType = "name";
+			options.ClaimActions.Add(new MapRolesClaimAction(settings));
 
-			Handler(Options);
+			handler(options);
 		}
 
-		public static void AddHordeOpenId(this AuthenticationBuilder Builder, ServerSettings Settings, string AuthenticationScheme, string DisplayName, Action<OpenIdConnectOptions> Handler)
+		public static void AddHordeOpenId(this AuthenticationBuilder builder, ServerSettings settings, string authenticationScheme, string displayName, Action<OpenIdConnectOptions> handler)
 		{
-			Builder.Services.TryAddEnumerable(ServiceDescriptor.Singleton<IPostConfigureOptions<OpenIdConnectOptions>, OpenIdConnectPostConfigureOptions>());
-			Builder.AddRemoteScheme<OpenIdConnectOptions, HordeOpenIdConnectHandler>(AuthenticationScheme, DisplayName, Options => ApplyDefaultOptions(Settings, Options, Handler));
+			builder.Services.TryAddEnumerable(ServiceDescriptor.Singleton<IPostConfigureOptions<OpenIdConnectOptions>, OpenIdConnectPostConfigureOptions>());
+			builder.AddRemoteScheme<OpenIdConnectOptions, HordeOpenIdConnectHandler>(authenticationScheme, displayName, options => ApplyDefaultOptions(settings, options, handler));
 		}
 	}
 }
diff --git a/Engine/Source/Programs/Horde/Horde.Build/Authentication/OktaHandler.cs b/Engine/Source/Programs/Horde/Horde.Build/Authentication/OktaHandler.cs
--- a/Engine/Source/Programs/Horde/Horde.Build/Authentication/OktaHandler.cs
+++ b/Engine/Source/Programs/Horde/Horde.Build/Authentication/OktaHandler.cs
@@ -1,5 +1,11 @@
 // Copyright Epic Games, Inc. All Rights Reserved.
 
+using System;
+using System.Linq;
+using System.Security.Claims;
+using System.Text.Encodings.Web;
+using System.Text.Json;
+using System.Threading.Tasks;
 using Horde.Build.Collections;
 using Horde.Build.Models;
 using Horde.Build.Utilities;
@@ -11,13 +17,6 @@ using Microsoft.Extensions.DependencyInjection.Extensions;
 using Microsoft.Extensions.Logging;
 using Microsoft.Extensions.Options;
 using Microsoft.IdentityModel.Protocols.OpenIdConnect;
-using System;
-using System.Collections.Generic;
-using System.Linq;
-using System.Security.Claims;
-using System.Text.Encodings.Web;
-using System.Text.Json;
-using System.Threading.Tasks;
 
 namespace Horde.Build.Authentication
 {
@@ -28,74 +27,74 @@ namespace Horde.Build.Authentication
 
 	class OktaHandler : OpenIdConnectHandler
 	{
-		IUserCollection UserCollection;
+		readonly IUserCollection _userCollection;
 
-		public OktaHandler(IOptionsMonitor<OpenIdConnectOptions> options, ILoggerFactory logger, HtmlEncoder htmlEncoder, UrlEncoder encoder, ISystemClock clock, IUserCollection UserCollection)
+		public OktaHandler(IOptionsMonitor<OpenIdConnectOptions> options, ILoggerFactory logger, HtmlEncoder htmlEncoder, UrlEncoder encoder, ISystemClock clock, IUserCollection userCollection)
 			: base(options, logger, htmlEncoder, encoder, clock)
 		{
-			this.UserCollection = UserCollection;
+			_userCollection = userCollection;
 		}
 
-		public static void AddUserInfoClaims(JsonElement UserInfo, ClaimsIdentity Identity)
+		public static void AddUserInfoClaims(JsonElement userInfo, ClaimsIdentity identity)
 		{
-			JsonElement UserElement;
-			if (UserInfo.TryGetProperty("preferred_username", out UserElement))
+			JsonElement userElement;
+			if (userInfo.TryGetProperty("preferred_username", out userElement))
 			{
-				if (Identity.FindFirst(ClaimTypes.Name) == null)
+				if (identity.FindFirst(ClaimTypes.Name) == null)
 				{
-					Identity.AddClaim(new Claim(ClaimTypes.Name, UserElement.ToString()!));
+					identity.AddClaim(new Claim(ClaimTypes.Name, userElement.ToString()!));
 				}
-				if (Identity.FindFirst(HordeClaimTypes.User) == null)
+				if (identity.FindFirst(HordeClaimTypes.User) == null)
 				{
-					Identity.AddClaim(new Claim(HordeClaimTypes.User, UserElement.ToString()!));
+					identity.AddClaim(new Claim(HordeClaimTypes.User, userElement.ToString()!));
 				}
-				if (Identity.FindFirst(HordeClaimTypes.PerforceUser) == null)
+				if (identity.FindFirst(HordeClaimTypes.PerforceUser) == null)
 				{
-					Identity.AddClaim(new Claim(HordeClaimTypes.PerforceUser, UserElement.ToString()!));
+					identity.AddClaim(new Claim(HordeClaimTypes.PerforceUser, userElement.ToString()!));
 				}
 			}
 
-			JsonElement GroupsElement;
-			if (UserInfo.TryGetProperty("groups", out GroupsElement) && GroupsElement.ValueKind == JsonValueKind.Array)
+			JsonElement groupsElement;
+			if (userInfo.TryGetProperty("groups", out groupsElement) && groupsElement.ValueKind == JsonValueKind.Array)
 			{
-				for (int Idx = 0; Idx < GroupsElement.GetArrayLength(); Idx++)
+				for (int idx = 0; idx < groupsElement.GetArrayLength(); idx++)
 				{
-					Identity.AddClaim(new Claim(ClaimTypes.Role, GroupsElement[Idx].ToString()!));
+					identity.AddClaim(new Claim(ClaimTypes.Role, groupsElement[idx].ToString()!));
 				}
 			}
 
-			JsonElement EmailElement;
-			if (UserInfo.TryGetProperty("email", out EmailElement) && Identity.FindFirst(ClaimTypes.Email) == null)
+			JsonElement emailElement;
+			if (userInfo.TryGetProperty("email", out emailElement) && identity.FindFirst(ClaimTypes.Email) == null)
 			{
-				Identity.AddClaim(new Claim(ClaimTypes.Email, EmailElement.ToString()!));
+				identity.AddClaim(new Claim(ClaimTypes.Email, emailElement.ToString()!));
 			}
 		}
 
 		protected override async Task<HandleRequestResult> HandleRemoteAuthenticateAsync()
 		{
 			// Authenticate with the OIDC provider
-			HandleRequestResult Result = await base.HandleRemoteAuthenticateAsync();
-			if (!Result.Succeeded)
+			HandleRequestResult result = await base.HandleRemoteAuthenticateAsync();
+			if (!result.Succeeded)
 			{
-				return Result;
+				return result;
 			}
 
-			ClaimsIdentity? Identity = (ClaimsIdentity?)Result.Principal?.Identity;
-			if (Identity == null)
+			ClaimsIdentity? identity = (ClaimsIdentity?)result.Principal?.Identity;
+			if (identity == null)
 			{
 				return HandleRequestResult.Fail("No identity specified");
 			}
 
-			string Login = Identity.FindFirst(HordeClaimTypes.User)!.Value;
-			string? Name = Identity.FindFirst(ClaimTypes.Name)?.Value;
-			string? Email = Identity.FindFirst(ClaimTypes.Email)?.Value;
+			string login = identity.FindFirst(HordeClaimTypes.User)!.Value;
+			string? name = identity.FindFirst(ClaimTypes.Name)?.Value;
+			string? email = identity.FindFirst(ClaimTypes.Email)?.Value;
 
-			IUser User = await UserCollection.FindOrAddUserByLoginAsync(Login, Name, Email);
-			Identity.AddClaim(new Claim(HordeClaimTypes.UserId, User.Id.ToString()));
+			IUser user = await _userCollection.FindOrAddUserByLoginAsync(login, name, email);
+			identity.AddClaim(new Claim(HordeClaimTypes.UserId, user.Id.ToString()));
 
-			await UserCollection.UpdateClaimsAsync(User.Id, Identity.Claims.Select(x => new UserClaim(x.Type, x.Value)));
+			await _userCollection.UpdateClaimsAsync(user.Id, identity.Claims.Select(x => new UserClaim(x.Type, x.Value)));
 
-			return Result;
+			return result;
 		}
 	}
 
@@ -108,30 +107,30 @@ namespace Horde.Build.Authentication
 			{
 			}
 
-			public override void Run(JsonElement UserData, ClaimsIdentity Identity, string Issuer)
+			public override void Run(JsonElement userData, ClaimsIdentity identity, string issuer)
 			{
-				OktaHandler.AddUserInfoClaims(UserData, Identity);
+				OktaHandler.AddUserInfoClaims(userData, identity);
 			}
 		}
 
-		static void ApplyDefaultOktaOptions(OpenIdConnectOptions Options, Action<OpenIdConnectOptions> Handler)
+		static void ApplyDefaultOktaOptions(OpenIdConnectOptions options, Action<OpenIdConnectOptions> handler)
 		{
-			Options.Scope.Add("profile");
-			Options.Scope.Add("groups");
-			Options.Scope.Add("email");
-			Options.ResponseType = OpenIdConnectResponseType.Code;
-			Options.GetClaimsFromUserInfoEndpoint = true;
-			Options.SaveTokens = true;
-			Options.TokenValidationParameters.NameClaimType = "name";
-			Options.ClaimActions.Add(new MapRolesClaimAction());
-
-			Handler(Options);
+			options.Scope.Add("profile");
+			options.Scope.Add("groups");
+			options.Scope.Add("email");
+			options.ResponseType = OpenIdConnectResponseType.Code;
+			options.GetClaimsFromUserInfoEndpoint = true;
+			options.SaveTokens = true;
+			options.TokenValidationParameters.NameClaimType = "name";
+			options.ClaimActions.Add(new MapRolesClaimAction());
+
+			handler(options);
 		}
 
-		public static void AddOkta(this AuthenticationBuilder Builder, string AuthenticationScheme, string DisplayName, Action<OpenIdConnectOptions> Handler)
+		public static void AddOkta(this AuthenticationBuilder builder, string authenticationScheme, string displayName, Action<OpenIdConnectOptions> handler)
 		{
-			Builder.Services.TryAddEnumerable(ServiceDescriptor.Singleton<IPostConfigureOptions<OpenIdConnectOptions>, OpenIdConnectPostConfigureOptions>());
-			Builder.AddRemoteScheme<OpenIdConnectOptions, OktaHandler>(AuthenticationScheme, DisplayName, Options => ApplyDefaultOktaOptions(Options, Handler));
+			builder.Services.TryAddEnumerable(ServiceDescriptor.Singleton<IPostConfigureOptions<OpenIdConnectOptions>, OpenIdConnectPostConfigureOptions>());
+			builder.AddRemoteScheme<OpenIdConnectOptions, OktaHandler>(authenticationScheme, displayName, options => ApplyDefaultOktaOptions(options, handler));
 		}
 	}
 }
diff --git a/Engine/Source/Programs/Horde/Horde.Build/Authentication/ServiceAccountAuthHandler.cs b/Engine/Source/Programs/Horde/Horde.Build/Authentication/ServiceAccountAuthHandler.cs
--- a/Engine/Source/Programs/Horde/Horde.Build/Authentication/ServiceAccountAuthHandler.cs
+++ b/Engine/Source/Programs/Horde/Horde.Build/Authentication/ServiceAccountAuthHandler.cs
@@ -1,8 +1,5 @@
 ﻿// Copyright Epic Games, Inc. All Rights Reserved.
 
-using Microsoft.AspNetCore.Authentication;
-using Microsoft.Extensions.Logging;
-using Microsoft.Extensions.Options;
 using System;
 using System.Collections.Generic;
 using System.Linq;
@@ -11,6 +8,9 @@ using System.Text.Encodings.Web;
 using System.Threading.Tasks;
 using Horde.Build.Collections;
 using Horde.Build.Models;
+using Microsoft.AspNetCore.Authentication;
+using Microsoft.Extensions.Logging;
+using Microsoft.Extensions.Options;
 using Microsoft.Net.Http.Headers;
 
 namespace Horde.Build.Authentication
@@ -24,57 +24,57 @@ namespace Horde.Build.Authentication
 		public const string AuthenticationScheme = "ServiceAccount";
 		public const string Prefix = "ServiceAccount";
 
-		private readonly IServiceAccountCollection ServiceAccounts;
+		private readonly IServiceAccountCollection _serviceAccounts;
 
-		public ServiceAccountAuthHandler(IOptionsMonitor<ServiceAccountAuthOptions> Options,
-			ILoggerFactory Logger, UrlEncoder Encoder, ISystemClock Clock, IServiceAccountCollection ServiceAccounts)
-			: base(Options, Logger, Encoder, Clock)
+		public ServiceAccountAuthHandler(IOptionsMonitor<ServiceAccountAuthOptions> options,
+			ILoggerFactory logger, UrlEncoder encoder, ISystemClock clock, IServiceAccountCollection serviceAccounts)
+			: base(options, logger, encoder, clock)
 		{
-			this.ServiceAccounts = ServiceAccounts;
+			_serviceAccounts = serviceAccounts;
 		}
 		
 		protected override async Task<AuthenticateResult> HandleAuthenticateAsync()
 		{
-			if (!Context.Request.Headers.TryGetValue(HeaderNames.Authorization, out var HeaderValue))
+			if (!Context.Request.Headers.TryGetValue(HeaderNames.Authorization, out Microsoft.Extensions.Primitives.StringValues headerValue))
 			{
 				return AuthenticateResult.NoResult();
 			}
 
-			if (HeaderValue.Count < 1)
+			if (headerValue.Count < 1)
 			{
 				return AuthenticateResult.NoResult();
 			}
 
-			if (!HeaderValue[0].StartsWith(Prefix, StringComparison.Ordinal))
+			if (!headerValue[0].StartsWith(Prefix, StringComparison.Ordinal))
 			{
 				return AuthenticateResult.NoResult();
 			}
 			
-			string Token = HeaderValue[0].Replace(Prefix, "", StringComparison.Ordinal).Trim();
-			IServiceAccount? ServiceAccount = await ServiceAccounts.GetBySecretTokenAsync(Token);
+			string token = headerValue[0].Replace(Prefix, "", StringComparison.Ordinal).Trim();
+			IServiceAccount? serviceAccount = await _serviceAccounts.GetBySecretTokenAsync(token);
 
-			if (ServiceAccount == null)
+			if (serviceAccount == null)
 			{
-				return AuthenticateResult.Fail($"Service account for token {Token} not found");
+				return AuthenticateResult.Fail($"Service account for token {token} not found");
 			}
 
-			List<Claim> Claims = new List<Claim>(10);
-			Claims.Add(new Claim(ClaimTypes.Name, AuthenticationScheme));
-			Claims.AddRange(ServiceAccount.GetClaims().Select(ClaimPair => new Claim(ClaimPair.Type, ClaimPair.Value)));
+			List<Claim> claims = new List<Claim>(10);
+			claims.Add(new Claim(ClaimTypes.Name, AuthenticationScheme));
+			claims.AddRange(serviceAccount.GetClaims().Select(claimPair => new Claim(claimPair.Type, claimPair.Value)));
 
-			ClaimsIdentity Identity = new ClaimsIdentity(Claims, Scheme.Name);
-			ClaimsPrincipal Principal = new ClaimsPrincipal(Identity);
-			AuthenticationTicket Ticket = new AuthenticationTicket(Principal, Scheme.Name);
+			ClaimsIdentity identity = new ClaimsIdentity(claims, Scheme.Name);
+			ClaimsPrincipal principal = new ClaimsPrincipal(identity);
+			AuthenticationTicket ticket = new AuthenticationTicket(principal, Scheme.Name);
 
-			return AuthenticateResult.Success(Ticket);
+			return AuthenticateResult.Success(ticket);
 		}
 	}
 
 	static class ServiceAccountExtensions
 	{
-		public static AuthenticationBuilder AddServiceAccount(this AuthenticationBuilder Builder, Action<ServiceAccountAuthOptions> Config)
+		public static AuthenticationBuilder AddServiceAccount(this AuthenticationBuilder builder, Action<ServiceAccountAuthOptions> config)
 		{
-			return Builder.AddScheme<ServiceAccountAuthOptions, ServiceAccountAuthHandler>(ServiceAccountAuthHandler.AuthenticationScheme, Config);
+			return builder.AddScheme<ServiceAccountAuthOptions, ServiceAccountAuthHandler>(ServiceAccountAuthHandler.AuthenticationScheme, config);
 		}
 	}
 }
diff --git a/Engine/Source/Programs/Horde/Horde.Build/Commands/Bundles/CreateCommand.cs b/Engine/Source/Programs/Horde/Horde.Build/Commands/Bundles/CreateCommand.cs
--- a/Engine/Source/Programs/Horde/Horde.Build/Commands/Bundles/CreateCommand.cs
+++ b/Engine/Source/Programs/Horde/Horde.Build/Commands/Bundles/CreateCommand.cs
@@ -1,5 +1,6 @@
 // Copyright Epic Games, Inc. All Rights Reserved.
 
+using System.Threading.Tasks;
 using EpicGames.Core;
 using EpicGames.Horde.Bundles;
 using EpicGames.Horde.Bundles.Nodes;
@@ -7,26 +8,19 @@ using EpicGames.Horde.Storage;
 using EpicGames.Horde.Storage.Impl;
 using EpicGames.Serialization;
 using Microsoft.Extensions.Caching.Memory;
-using Microsoft.Extensions.Configuration;
-using Microsoft.Extensions.DependencyInjection;
 using Microsoft.Extensions.Logging;
-using System;
-using System.Collections.Generic;
-using System.IO;
-using System.Linq;
-using System.Threading.Tasks;
 
 namespace Horde.Build.Commands.Bundles
 {
 	abstract class BundleCommandBase : Command
 	{
 		[CommandLine("-StorageDir=", Description = "Overrides the default storage server with a local directory")]
-		public DirectoryReference? StorageDir = null;
+		public DirectoryReference? _storageDir = null;
 
-		protected IStorageClient CreateStorageClient(ILogger Logger)
+		protected IStorageClient CreateStorageClient(ILogger logger)
 		{
-			StorageDir ??= DirectoryReference.Combine(Program.DataDir, "Storage");
-			return new FileStorageClient(StorageDir, Logger);
+			_storageDir ??= DirectoryReference.Combine(Program.DataDir, "Storage");
+			return new FileStorageClient(_storageDir, logger);
 		}
 	}
 
@@ -45,15 +39,15 @@ namespace Horde.Build.Commands.Bundles
 		[CommandLine("-InputDir=", Required = true)]
 		public DirectoryReference InputDir { get; set; } = null!;
 
-		public override async Task<int> ExecuteAsync(ILogger Logger)
+		public override async Task<int> ExecuteAsync(ILogger logger)
 		{
-			IStorageClient StorageClient = base.CreateStorageClient(Logger);
-			using (MemoryCache Cache = new MemoryCache(new MemoryCacheOptions { SizeLimit = 50 * 1024 * 1024 }))
+			IStorageClient storageClient = base.CreateStorageClient(logger);
+			using (MemoryCache cache = new MemoryCache(new MemoryCacheOptions { SizeLimit = 50 * 1024 * 1024 }))
 			{
-				using (Bundle<DirectoryNode> NewBundle = Bundle.Create<DirectoryNode>(StorageClient, NamespaceId, new BundleOptions(), Cache))
+				using (Bundle<DirectoryNode> newBundle = Bundle.Create<DirectoryNode>(storageClient, NamespaceId, new BundleOptions(), cache))
 				{
-					await NewBundle.Root.CopyFromDirectoryAsync(InputDir.ToDirectoryInfo(), new ChunkingOptions(), Logger);
-					await NewBundle.WriteAsync(BucketId, RefId, CbObject.Empty, false);
+					await newBundle.Root.CopyFromDirectoryAsync(InputDir.ToDirectoryInfo(), new ChunkingOptions(), logger);
+					await newBundle.WriteAsync(BucketId, RefId, CbObject.Empty, false);
 				}
 			}
 
diff --git a/Engine/Source/Programs/Horde/Horde.Build/Commands/Bundles/DeltaCommand.cs b/Engine/Source/Programs/Horde/Horde.Build/Commands/Bundles/DeltaCommand.cs
--- a/Engine/Source/Programs/Horde/Horde.Build/Commands/Bundles/DeltaCommand.cs
+++ b/Engine/Source/Programs/Horde/Horde.Build/Commands/Bundles/DeltaCommand.cs
@@ -1,19 +1,10 @@
 // Copyright Epic Games, Inc. All Rights Reserved.
 
+using System.Threading.Tasks;
 using EpicGames.Core;
 using EpicGames.Horde.Bundles;
 using EpicGames.Horde.Storage;
-using EpicGames.Horde.Storage.Impl;
-using EpicGames.Serialization;
-using Microsoft.Extensions.Caching.Memory;
-using Microsoft.Extensions.Configuration;
 using Microsoft.Extensions.Logging;
-using System;
-using System.Collections.Generic;
-using System.Diagnostics;
-using System.Linq;
-using System.Text;
-using System.Threading.Tasks;
 
 namespace Horde.Build.Commands.Bundles
 {
@@ -32,10 +23,10 @@ namespace Horde.Build.Commands.Bundles
 		[CommandLine("-Ref2=", Required = true)]
 		public RefId RefId2 { get; set; }
 
-		public override async Task<int> ExecuteAsync(ILogger Logger)
+		public override async Task<int> ExecuteAsync(ILogger logger)
 		{
-			IStorageClient StorageClient = CreateStorageClient(Logger);
-			await BundleTools.FindDeltaAsync(StorageClient, NamespaceId, BucketId, RefId1, RefId2, Logger);
+			IStorageClient storageClient = CreateStorageClient(logger);
+			await BundleTools.FindDeltaAsync(storageClient, NamespaceId, BucketId, RefId1, RefId2, logger);
 			return 0;
 		}
 	}
diff --git a/Engine/Source/Programs/Horde/Horde.Build/Commands/Bundles/DumpCommand.cs b/Engine/Source/Programs/Horde/Horde.Build/Commands/Bundles/DumpCommand.cs
--- a/Engine/Source/Programs/Horde/Horde.Build/Commands/Bundles/DumpCommand.cs
+++ b/Engine/Source/Programs/Horde/Horde.Build/Commands/Bundles/DumpCommand.cs
@@ -1,19 +1,10 @@
 // Copyright Epic Games, Inc. All Rights Reserved.
 
+using System.Threading.Tasks;
 using EpicGames.Core;
 using EpicGames.Horde.Bundles;
 using EpicGames.Horde.Storage;
-using EpicGames.Horde.Storage.Impl;
-using EpicGames.Serialization;
-using Microsoft.Extensions.Caching.Memory;
-using Microsoft.Extensions.Configuration;
-using Microsoft.Extensions.DependencyInjection;
 using Microsoft.Extensions.Logging;
-using System;
-using System.Collections.Generic;
-using System.IO;
-using System.Linq;
-using System.Threading.Tasks;
 
 namespace Horde.Build.Commands.Bundles
 {
@@ -32,12 +23,12 @@ namespace Horde.Build.Commands.Bundles
 		[CommandLine("-OutputDir=")]
 		public DirectoryReference? OutputDir { get; set; }
 
-		public override async Task<int> ExecuteAsync(ILogger Logger)
+		public override async Task<int> ExecuteAsync(ILogger logger)
 		{
 			OutputDir = DirectoryReference.Combine(Program.DataDir, "Bundles");
 
-			IStorageClient StorageClient = CreateStorageClient(Logger);
-			await BundleTools.WriteSummaryAsync(StorageClient, NamespaceId, BucketId, RefId, OutputDir);
+			IStorageClient storageClient = CreateStorageClient(logger);
+			await BundleTools.WriteSummaryAsync(storageClient, NamespaceId, BucketId, RefId, OutputDir);
 
 			return 0;
 		}
diff --git a/Engine/Source/Programs/Horde/Horde.Build/Commands/Bundles/ExtractCommand.cs b/Engine/Source/Programs/Horde/Horde.Build/Commands/Bundles/ExtractCommand.cs
--- a/Engine/Source/Programs/Horde/Horde.Build/Commands/Bundles/ExtractCommand.cs
+++ b/Engine/Source/Programs/Horde/Horde.Build/Commands/Bundles/ExtractCommand.cs
@@ -1,20 +1,12 @@
 // Copyright Epic Games, Inc. All Rights Reserved.
 
+using System.Threading.Tasks;
 using EpicGames.Core;
 using EpicGames.Horde.Bundles;
 using EpicGames.Horde.Bundles.Nodes;
 using EpicGames.Horde.Storage;
-using EpicGames.Horde.Storage.Impl;
-using EpicGames.Serialization;
 using Microsoft.Extensions.Caching.Memory;
-using Microsoft.Extensions.Configuration;
-using Microsoft.Extensions.DependencyInjection;
 using Microsoft.Extensions.Logging;
-using System;
-using System.Collections.Generic;
-using System.IO;
-using System.Linq;
-using System.Threading.Tasks;
 
 namespace Horde.Build.Commands.Bundles
 {
@@ -33,13 +25,13 @@ namespace Horde.Build.Commands.Bundles
 		[CommandLine("-OutputDir=", Required = true)]
 		public DirectoryReference OutputDir { get; set; } = null!;
 
-		public override async Task<int> ExecuteAsync(ILogger Logger)
+		public override async Task<int> ExecuteAsync(ILogger logger)
 		{
-			IStorageClient StorageClient = CreateStorageClient(Logger);
-			using (IMemoryCache Cache = new MemoryCache(new MemoryCacheOptions()))
+			IStorageClient storageClient = CreateStorageClient(logger);
+			using (IMemoryCache cache = new MemoryCache(new MemoryCacheOptions()))
 			{
-				using Bundle<DirectoryNode> NewBundle = await Bundle.ReadAsync<DirectoryNode>(StorageClient, NamespaceId, BucketId, RefId, new BundleOptions(), Cache);
-				await NewBundle.Root.CopyToDirectoryAsync(NewBundle, OutputDir.ToDirectoryInfo(), Logger);
+				using Bundle<DirectoryNode> newBundle = await Bundle.ReadAsync<DirectoryNode>(storageClient, NamespaceId, BucketId, RefId, new BundleOptions(), cache);
+				await newBundle.Root.CopyToDirectoryAsync(newBundle, OutputDir.ToDirectoryInfo(), logger);
 			}
 			return 0;
 		}
diff --git a/Engine/Source/Programs/Horde/Horde.Build/Commands/Config/SchemasCommand.cs b/Engine/Source/Programs/Horde/Horde.Build/Commands/Config/SchemasCommand.cs
--- a/Engine/Source/Programs/Horde/Horde.Build/Commands/Config/SchemasCommand.cs
+++ b/Engine/Source/Programs/Horde/Horde.Build/Commands/Config/SchemasCommand.cs
@@ -1,11 +1,11 @@
 // Copyright Epic Games, Inc. All Rights Reserved.
 
+using System;
+using System.Threading.Tasks;
 using EpicGames.Core;
 using Horde.Build.Controllers;
 using Horde.Build.Utilities;
 using Microsoft.Extensions.Logging;
-using System;
-using System.Threading.Tasks;
 
 namespace Horde.Build.Commands.Config
 {
@@ -13,17 +13,17 @@ namespace Horde.Build.Commands.Config
 	class SchemasCommand : Command
 	{
 		[CommandLine]
-		DirectoryReference? OutputDir = null!;
+		DirectoryReference? _outputDir = null!;
 
-		public override Task<int> ExecuteAsync(ILogger Logger)
+		public override Task<int> ExecuteAsync(ILogger logger)
 		{
-			OutputDir ??= DirectoryReference.Combine(Program.AppDir, "Schemas");
+			_outputDir ??= DirectoryReference.Combine(Program.AppDir, "Schemas");
 
-			DirectoryReference.CreateDirectory(OutputDir);
-			foreach (Type SchemaType in SchemaController.ConfigSchemas)
+			DirectoryReference.CreateDirectory(_outputDir);
+			foreach (Type schemaType in SchemaController.ConfigSchemas)
 			{
-				FileReference OutputFile = FileReference.Combine(OutputDir, $"{SchemaType.Name}.json");
-				Schemas.CreateSchema(SchemaType).Write(OutputFile);
+				FileReference outputFile = FileReference.Combine(_outputDir, $"{schemaType.Name}.json");
+				Schemas.CreateSchema(schemaType).Write(outputFile);
 			}
 
 			return Task.FromResult(0);
diff --git a/Engine/Source/Programs/Horde/Horde.Build/Commands/Perforce/ReplicateCommand.cs b/Engine/Source/Programs/Horde/Horde.Build/Commands/Perforce/ReplicateCommand.cs
--- a/Engine/Source/Programs/Horde/Horde.Build/Commands/Perforce/ReplicateCommand.cs
+++ b/Engine/Source/Programs/Horde/Horde.Build/Commands/Perforce/ReplicateCommand.cs
@@ -1,22 +1,20 @@
 // Copyright Epic Games, Inc. All Rights Reserved.
 
+using System;
+using System.Collections.Generic;
+using System.Linq;
+using System.Threading.Tasks;
 using EpicGames.Core;
 using EpicGames.Horde.Storage;
 using EpicGames.Horde.Storage.Impl;
-using EpicGames.Serialization;
 using Horde.Build.Collections;
 using Horde.Build.Commits;
 using Horde.Build.Commits.Impl;
-using Horde.Build.Controllers;
 using Horde.Build.Models;
 using Horde.Build.Utilities;
 using Microsoft.Extensions.Configuration;
 using Microsoft.Extensions.DependencyInjection;
 using Microsoft.Extensions.Logging;
-using System;
-using System.Collections.Generic;
-using System.Linq;
-using System.Threading.Tasks;
 
 namespace Horde.Build.Commands
 {
@@ -55,54 +53,54 @@ namespace Horde.Build.Commands
 		[CommandLine]
 		public DirectoryReference? OutputDir { get; set; }
 
-		IConfiguration Configuration;
-		ILoggerProvider LoggerProvider;
+		readonly IConfiguration _configuration;
+		readonly ILoggerProvider _loggerProvider;
 
-		public ReplicateCommand(IConfiguration Configuration, ILoggerProvider LoggerProvider)
+		public ReplicateCommand(IConfiguration configuration, ILoggerProvider loggerProvider)
 		{
-			this.Configuration = Configuration;
-			this.LoggerProvider = LoggerProvider;
+			_configuration = configuration;
+			_loggerProvider = loggerProvider;
 		}
 
-		public override async Task<int> ExecuteAsync(ILogger Logger)
+		public override async Task<int> ExecuteAsync(ILogger logger)
 		{
-			IServiceCollection Services = new ServiceCollection();
-			Services.AddLogging(Builder => Builder.AddProvider(LoggerProvider));
+			IServiceCollection services = new ServiceCollection();
+			services.AddLogging(builder => builder.AddProvider(_loggerProvider));
 
-			Startup.AddServices(Services, Configuration);
+			Startup.AddServices(services, _configuration);
 
 			OutputDir ??= DirectoryReference.Combine(Program.DataDir, "Storage");
-			Logger.LogInformation("Writing output to {OutputDir}", OutputDir);
-			Services.AddSingleton<IStorageClient, FileStorageClient>(SP => new FileStorageClient(OutputDir, Logger));
+			logger.LogInformation("Writing output to {OutputDir}", OutputDir);
+			services.AddSingleton<IStorageClient, FileStorageClient>(sp => new FileStorageClient(OutputDir, logger));
 
-			IServiceProvider ServiceProvider = Services.BuildServiceProvider();
-			CommitService CommitService = ServiceProvider.GetRequiredService<CommitService>();
-			ICommitCollection CommitCollection = ServiceProvider.GetRequiredService<ICommitCollection>();
-			IStreamCollection StreamCollection = ServiceProvider.GetRequiredService<IStreamCollection>();
-			IStorageClient StorageClient = ServiceProvider.GetRequiredService<IStorageClient>();
+			IServiceProvider serviceProvider = services.BuildServiceProvider();
+			CommitService commitService = serviceProvider.GetRequiredService<CommitService>();
+			ICommitCollection commitCollection = serviceProvider.GetRequiredService<ICommitCollection>();
+			IStreamCollection streamCollection = serviceProvider.GetRequiredService<IStreamCollection>();
+			IStorageClient storageClient = serviceProvider.GetRequiredService<IStorageClient>();
 
-			IStream? Stream = await StreamCollection.GetAsync(new StreamId(StreamId));
-			if (Stream == null)
+			IStream? stream = await streamCollection.GetAsync(new StreamId(StreamId));
+			if (stream == null)
 			{
-				throw new FatalErrorException($"Stream '{Stream}' not found");
+				throw new FatalErrorException($"Stream '{stream}' not found");
 			}
 
-			Dictionary<IStream, int> StreamToFirstChange = new Dictionary<IStream, int>();
-			StreamToFirstChange[Stream] = Change;
+			Dictionary<IStream, int> streamToFirstChange = new Dictionary<IStream, int>();
+			streamToFirstChange[stream] = Change;
 
-			await foreach (NewCommit NewCommit in CommitService.FindCommitsForClusterAsync(Stream.ClusterName, StreamToFirstChange).Take(Count))
+			await foreach (NewCommit newCommit in commitService.FindCommitsForClusterAsync(stream.ClusterName, streamToFirstChange).Take(Count))
 			{
-				string BriefSummary = NewCommit.Description.Replace('\n', ' ');
-				Logger.LogInformation("Commit {Change} by {AuthorId}: {Summary}", NewCommit.Change, NewCommit.AuthorId, BriefSummary.Substring(0, Math.Min(50, BriefSummary.Length)));
-				Logger.LogInformation(" - Base path: {BasePath}", NewCommit.BasePath);
+				string briefSummary = newCommit.Description.Replace('\n', ' ');
+				logger.LogInformation("Commit {Change} by {AuthorId}: {Summary}", newCommit.Change, newCommit.AuthorId, briefSummary.Substring(0, Math.Min(50, briefSummary.Length)));
+				logger.LogInformation(" - Base path: {BasePath}", newCommit.BasePath);
 
 				if (Content)
 				{
 					if (Clean || BaseChange != null)
 					{
-						await CommitService.WriteCommitTreeAsync(Stream, NewCommit.Change, BaseChange, Filter, Metadata);
+						await commitService.WriteCommitTreeAsync(stream, newCommit.Change, BaseChange, Filter, Metadata);
 					}
-					BaseChange = NewCommit.Change;
+					BaseChange = newCommit.Change;
 				}
 			}
 
diff --git a/Engine/Source/Programs/Horde/Horde.Build/Commands/ServerCommand.cs b/Engine/Source/Programs/Horde/Horde.Build/Commands/ServerCommand.cs
--- a/Engine/Source/Programs/Horde/Horde.Build/Commands/ServerCommand.cs
+++ b/Engine/Source/Programs/Horde/Horde.Build/Commands/ServerCommand.cs
@@ -1,19 +1,18 @@
 // Copyright Epic Games, Inc. All Rights Reserved.
 
+using System;
+using System.Collections.Generic;
+using System.IO;
+using System.Linq;
+using System.Security.Cryptography.X509Certificates;
+using System.Threading.Tasks;
 using EpicGames.Core;
-using Horde.Build;
 using Microsoft.AspNetCore.Hosting;
 using Microsoft.AspNetCore.Server.Kestrel.Core;
 using Microsoft.Extensions.Configuration;
 using Microsoft.Extensions.Hosting;
 using Microsoft.Extensions.Logging;
 using Serilog;
-using System;
-using System.Collections.Generic;
-using System.IO;
-using System.Linq;
-using System.Security.Cryptography.X509Certificates;
-using System.Threading.Tasks;
 
 namespace Horde.Build.Commands
 {
@@ -22,74 +21,80 @@ namespace Horde.Build.Commands
 	[Command("server", "Runs the Horde Build server (default)")]
 	class ServerCommand : Command
 	{
-		ServerSettings HordeSettings;
-		IConfiguration Config;
-		string[] Args = Array.Empty<string>();
+		readonly ServerSettings _hordeSettings;
+		readonly IConfiguration _config;
+		string[] _args = Array.Empty<string>();
 
-		public ServerCommand(ServerSettings Settings, IConfiguration Config)
+		public ServerCommand(ServerSettings settings, IConfiguration config)
 		{
-			this.HordeSettings = Settings;
-			this.Config = Config;
+			_hordeSettings = settings;
+			_config = config;
 		}
 
-		public override void Configure(CommandLineArguments Arguments, ILogger Logger)
+		public override void Configure(CommandLineArguments arguments, ILogger logger)
 		{
-			base.Configure(Arguments, Logger);
-			this.Args = Arguments.GetRawArray();
+			base.Configure(arguments, logger);
+			_args = arguments.GetRawArray();
 		}
 
-		public override async Task<int> ExecuteAsync(ILogger Logger)
+		public override async Task<int> ExecuteAsync(ILogger logger)
 		{
-			using (X509Certificate2? GrpcCertificate = ReadGrpcCertificate(HordeSettings))
+			using (X509Certificate2? grpcCertificate = ReadGrpcCertificate(_hordeSettings))
 			{
-				List<IHost> Hosts = new List<IHost>();
-				Hosts.Add(CreateHostBuilderWithCert(Args, Config, HordeSettings, GrpcCertificate).Build());
+				List<IHost> hosts = new List<IHost>();
+				hosts.Add(CreateHostBuilderWithCert(_args, _config, _hordeSettings, grpcCertificate).Build());
 #if WITH_HORDE_STORAGE
-				IHostBuilder StorageHostBuilder = Horde.Storage.Program.CreateHostBuilder(Args);
-				StorageHostBuilder.ConfigureWebHostDefaults(Builder =>
+				IHostBuilder storageHostBuilder = Horde.Storage.Program.CreateHostBuilder(_args);
+				storageHostBuilder.ConfigureWebHostDefaults(builder =>
 				{
-					Builder.ConfigureKestrel(Options =>
+					builder.ConfigureKestrel(options =>
 					{
-						Options.ListenAnyIP(57000);
-						Options.ListenAnyIP(57001, Configure => Configure.UseHttps());
+						options.ListenAnyIP(57000);
+						options.ListenAnyIP(57001, configure => configure.UseHttps());
 					});
 				});
-				Hosts.Add(StorageHostBuilder.Build());
+				hosts.Add(storageHostBuilder.Build());
 #endif
-				await Task.WhenAll(Hosts.Select(x => x.RunAsync()));
+				await Task.WhenAll(hosts.Select(x => x.RunAsync()));
 				return 0;
 			}
 		}
 
-		static IHostBuilder CreateHostBuilderWithCert(string[] Args, IConfiguration Config, ServerSettings ServerSettings, X509Certificate2? SslCert)
+		static IHostBuilder CreateHostBuilderWithCert(string[] args, IConfiguration config, ServerSettings serverSettings, X509Certificate2? sslCert)
 		{
-			return Host.CreateDefaultBuilder(Args)
+			return Host.CreateDefaultBuilder(args)
 				.UseSerilog()
-				.ConfigureAppConfiguration(Builder => Builder.AddConfiguration(Config))
-				.ConfigureWebHostDefaults(WebBuilder =>
+				.ConfigureAppConfiguration(builder => builder.AddConfiguration(config))
+				.ConfigureWebHostDefaults(webBuilder =>
 				{
-					WebBuilder.ConfigureKestrel(Options =>
+					webBuilder.ConfigureKestrel(options =>
 					{
-						Options.Limits.MaxRequestBodySize = 100 * 1024 * 1024;
+						options.Limits.MaxRequestBodySize = 100 * 1024 * 1024;
 
-						if (ServerSettings.HttpPort != 0)
+						if (serverSettings.HttpPort != 0)
 						{
-							Options.ListenAnyIP(ServerSettings.HttpPort, Configure => { Configure.Protocols = HttpProtocols.Http1AndHttp2; });
+							options.ListenAnyIP(serverSettings.HttpPort, configure => { configure.Protocols = HttpProtocols.Http1AndHttp2; });
 						}
 
-						if (ServerSettings.HttpsPort != 0)
+						if (serverSettings.HttpsPort != 0)
 						{
-							Options.ListenAnyIP(ServerSettings.HttpsPort, Configure => { if (SslCert != null) { Configure.UseHttps(SslCert); } });
+							options.ListenAnyIP(serverSettings.HttpsPort, configure => 
+							{ 
+								if (sslCert != null) 
+								{ 
+									configure.UseHttps(sslCert); 
+								} 
+							});
 						}
 
 						// To serve HTTP/2 with gRPC *without* TLS enabled, a separate port for HTTP/2 must be used.
 						// This is useful when having a load balancer in front that terminates TLS.
-						if (ServerSettings.Http2Port != 0)
+						if (serverSettings.Http2Port != 0)
 						{
-							Options.ListenAnyIP(ServerSettings.Http2Port, Configure => { Configure.Protocols = HttpProtocols.Http2; });
+							options.ListenAnyIP(serverSettings.Http2Port, configure => { configure.Protocols = HttpProtocols.Http2; });
 						}
 					});
-					WebBuilder.UseStartup<Startup>();
+					webBuilder.UseStartup<Startup>();
 				});
 		}
 
@@ -97,40 +102,39 @@ namespace Horde.Build.Commands
 		/// Gets the certificate to use for Grpc endpoints
 		/// </summary>
 		/// <returns>Custom certificate to use for Grpc endpoints, or null for the default.</returns>
-		public static X509Certificate2? ReadGrpcCertificate(ServerSettings HordeSettings)
+		public static X509Certificate2? ReadGrpcCertificate(ServerSettings hordeSettings)
 		{
-			string Base64Prefix = "base64:";
+			string base64Prefix = "base64:";
 
-			if (HordeSettings.ServerPrivateCert == null)
+			if (hordeSettings.ServerPrivateCert == null)
 			{
 				return null;
 			}
-			else if (HordeSettings.ServerPrivateCert.StartsWith(Base64Prefix, StringComparison.Ordinal))
+			else if (hordeSettings.ServerPrivateCert.StartsWith(base64Prefix, StringComparison.Ordinal))
 			{
-				byte[] CertData = Convert.FromBase64String(HordeSettings.ServerPrivateCert.Replace(Base64Prefix, "", StringComparison.Ordinal));
-				return new X509Certificate2(CertData);
+				byte[] certData = Convert.FromBase64String(hordeSettings.ServerPrivateCert.Replace(base64Prefix, "", StringComparison.Ordinal));
+				return new X509Certificate2(certData);
 			}
 			else
 			{
-				FileReference? ServerPrivateCert = null;
-
-				if (!Path.IsPathRooted(HordeSettings.ServerPrivateCert))
+				FileReference? serverPrivateCert;
+				if (!Path.IsPathRooted(hordeSettings.ServerPrivateCert))
 				{
-					ServerPrivateCert = FileReference.Combine(Program.AppDir, HordeSettings.ServerPrivateCert);
+					serverPrivateCert = FileReference.Combine(Program.AppDir, hordeSettings.ServerPrivateCert);
 				}
 				else
 				{
-					ServerPrivateCert = new FileReference(HordeSettings.ServerPrivateCert);
+					serverPrivateCert = new FileReference(hordeSettings.ServerPrivateCert);
 				}
 
-				return new X509Certificate2(FileReference.ReadAllBytes(ServerPrivateCert));
+				return new X509Certificate2(FileReference.ReadAllBytes(serverPrivateCert));
 			}
 		}
 
-		public static IHostBuilder CreateHostBuilderForTesting(string[] Args)
+		public static IHostBuilder CreateHostBuilderForTesting(string[] args)
 		{
-			ServerSettings HordeSettings = new ServerSettings();
-			return CreateHostBuilderWithCert(Args, new ConfigurationBuilder().Build(), HordeSettings, null);
+			ServerSettings hordeSettings = new ServerSettings();
+			return CreateHostBuilderWithCert(args, new ConfigurationBuilder().Build(), hordeSettings, null);
 		}
 }
 }
diff --git a/Engine/Source/Programs/Horde/Horde.Build/Compute/IComputeService.cs b/Engine/Source/Programs/Horde/Horde.Build/Compute/IComputeService.cs
--- a/Engine/Source/Programs/Horde/Horde.Build/Compute/IComputeService.cs
+++ b/Engine/Source/Programs/Horde/Horde.Build/Compute/IComputeService.cs
@@ -1,21 +1,15 @@
 // Copyright Epic Games, Inc. All Rights Reserved.
 
-using EpicGames.Core;
+using System;
+using System.Collections.Generic;
+using System.Threading;
+using System.Threading.Tasks;
 using EpicGames.Horde.Compute;
 using EpicGames.Horde.Storage;
 using EpicGames.Serialization;
 using Horde.Build.Models;
-using Horde.Build.Storage;
 using Horde.Build.Tasks;
 using Horde.Build.Utilities;
-using Microsoft.Extensions.Hosting;
-using MongoDB.Bson;
-using System;
-using System.Collections.Generic;
-using System.Diagnostics.CodeAnalysis;
-using System.Linq;
-using System.Threading;
-using System.Threading.Tasks;
 
 namespace Horde.Build.Compute
 {
@@ -75,34 +69,34 @@ namespace Horde.Build.Compute
 		/// <summary>
 		/// Gets information about a compute cluster
 		/// </summary>
-		/// <param name="ClusterId">Cluster to use for execution</param>
-		Task<IComputeClusterInfo> GetClusterInfoAsync(ClusterId ClusterId);
+		/// <param name="clusterId">Cluster to use for execution</param>
+		Task<IComputeClusterInfo> GetClusterInfoAsync(ClusterId clusterId);
 
 		/// <summary>
 		/// Post tasks to be executed to a channel
 		/// </summary>
-		/// <param name="ClusterId">Cluster to use for execution</param>
-		/// <param name="ChannelId">Unique identifier of the client</param>
-		/// <param name="TaskRefIds">List of task hashes</param>
-		/// <param name="RequirementsHash">Requirements document for execution</param>
+		/// <param name="clusterId">Cluster to use for execution</param>
+		/// <param name="channelId">Unique identifier of the client</param>
+		/// <param name="taskRefIds">List of task hashes</param>
+		/// <param name="requirementsHash">Requirements document for execution</param>
 		/// <returns>Async task</returns>
-		Task AddTasksAsync(ClusterId ClusterId, ChannelId ChannelId, List<RefId> TaskRefIds, CbObjectAttachment RequirementsHash);
+		Task AddTasksAsync(ClusterId clusterId, ChannelId channelId, List<RefId> taskRefIds, CbObjectAttachment requirementsHash);
 
 		/// <summary>
 		/// Dequeue completed items from a queue and return immediately
 		/// </summary>
-		/// <param name="ClusterId">Cluster containing the channel</param>
-		/// <param name="ChannelId">Queue to remove items from</param>
+		/// <param name="clusterId">Cluster containing the channel</param>
+		/// <param name="channelId">Queue to remove items from</param>
 		/// <returns>List of status updates</returns>
-		Task<List<IComputeTaskStatus>> GetTaskUpdatesAsync(ClusterId ClusterId, ChannelId ChannelId);
+		Task<List<IComputeTaskStatus>> GetTaskUpdatesAsync(ClusterId clusterId, ChannelId channelId);
 
 		/// <summary>
 		/// Dequeue completed items from a queue
 		/// </summary>
-		/// <param name="ClusterId">Cluster containing the channel</param>
-		/// <param name="ChannelId">Queue to remove items from</param>
-		/// <param name="CancellationToken">Cancellation token to stop waiting for items</param>
+		/// <param name="clusterId">Cluster containing the channel</param>
+		/// <param name="channelId">Queue to remove items from</param>
+		/// <param name="cancellationToken">Cancellation token to stop waiting for items</param>
 		/// <returns>List of status updates</returns>
-		Task<List<IComputeTaskStatus>> WaitForTaskUpdatesAsync(ClusterId ClusterId, ChannelId ChannelId, CancellationToken CancellationToken);
+		Task<List<IComputeTaskStatus>> WaitForTaskUpdatesAsync(ClusterId clusterId, ChannelId channelId, CancellationToken cancellationToken);
 	}
 }
diff --git a/Engine/Source/Programs/Horde/Horde.Build/Compute/Impl/ComputeController.cs b/Engine/Source/Programs/Horde/Horde.Build/Compute/Impl/ComputeController.cs
--- a/Engine/Source/Programs/Horde/Horde.Build/Compute/Impl/ComputeController.cs
+++ b/Engine/Source/Programs/Horde/Horde.Build/Compute/Impl/ComputeController.cs
@@ -1,20 +1,14 @@
-﻿// Copyright Epic Games, Inc. All Rights Reserved.
+// Copyright Epic Games, Inc. All Rights Reserved.
 
-using EpicGames.Core;
+using System.Collections.Generic;
+using System.Threading;
+using System.Threading.Tasks;
 using EpicGames.Horde.Compute;
 using EpicGames.Horde.Compute.Impl;
 using Horde.Build.Acls;
-using Horde.Build.Models;
-using Horde.Build.Services;
 using Horde.Build.Utilities;
 using Microsoft.AspNetCore.Authorization;
 using Microsoft.AspNetCore.Mvc;
-using System;
-using System.Collections.Generic;
-using System.ComponentModel.DataAnnotations;
-using System.Linq;
-using System.Threading;
-using System.Threading.Tasks;
 
 namespace Horde.Build.Compute.Impl
 {
@@ -26,124 +20,124 @@ namespace Horde.Build.Compute.Impl
 	[Route("[controller]")]
 	public class ComputeController : HordeControllerBase
 	{
-		AclService AclService;
-		IComputeService ComputeService;
+		readonly AclService _aclService;
+		readonly IComputeService _computeService;
 
 		/// <summary>
 		/// Constructor
 		/// </summary>
-		/// <param name="AclService"></param>
-		/// <param name="ComputeService">The compute service singleton</param>
-		public ComputeController(AclService AclService, IComputeService ComputeService)
+		/// <param name="aclService"></param>
+		/// <param name="computeService">The compute service singleton</param>
+		public ComputeController(AclService aclService, IComputeService computeService)
 		{
-			this.AclService = AclService;
-			this.ComputeService = ComputeService;
+			_aclService = aclService;
+			_computeService = computeService;
 		}
 
 		/// <summary>
 		/// Gets information about a cluster
 		/// </summary>
-		/// <param name="ClusterId">The cluster to add to</param>
+		/// <param name="clusterId">The cluster to add to</param>
 		/// <returns></returns>
 		[HttpGet]
 		[Authorize]
-		[Route("/api/v1/compute/{ClusterId}")]
-		public async Task<ActionResult<GetComputeClusterInfo>> GetClusterInfoAsync([FromRoute] ClusterId ClusterId)
+		[Route("/api/v1/compute/{clusterId}")]
+		public async Task<ActionResult<GetComputeClusterInfo>> GetClusterInfoAsync([FromRoute] ClusterId clusterId)
 		{
-			if (!await AclService.AuthorizeAsync(AclAction.ViewComputeTasks, User))
+			if (!await _aclService.AuthorizeAsync(AclAction.ViewComputeTasks, User))
 			{
 				return Forbid(AclAction.ViewComputeTasks);
 			}
 
-			IComputeClusterInfo ClusterInfo;
+			IComputeClusterInfo clusterInfo;
 			try
 			{
-				ClusterInfo = await ComputeService.GetClusterInfoAsync(ClusterId);
+				clusterInfo = await _computeService.GetClusterInfoAsync(clusterId);
 			}
 			catch (KeyNotFoundException)
 			{
-				return NotFound("Cluster '{ClusterId}' was found", ClusterId);
+				return NotFound("Cluster '{ClusterId}' was found", clusterId);
 			}
 
-			GetComputeClusterInfo Response = new GetComputeClusterInfo();
-			Response.Id = ClusterInfo.Id;
-			Response.NamespaceId = ClusterInfo.NamespaceId;
-			Response.RequestBucketId = ClusterInfo.RequestBucketId;
-			Response.ResponseBucketId = ClusterInfo.ResponseBucketId;
+			GetComputeClusterInfo response = new GetComputeClusterInfo();
+			response.Id = clusterInfo.Id;
+			response.NamespaceId = clusterInfo.NamespaceId;
+			response.RequestBucketId = clusterInfo.RequestBucketId;
+			response.ResponseBucketId = clusterInfo.ResponseBucketId;
 
-			return Response;
+			return response;
 		}
 
 		/// <summary>
 		/// Add tasks to be executed remotely
 		/// </summary>
-		/// <param name="ClusterId">The cluster to add to</param>
-		/// <param name="Request">The request parameters</param>
+		/// <param name="clusterId">The cluster to add to</param>
+		/// <param name="request">The request parameters</param>
 		/// <returns></returns>
 		[HttpPost]
 		[Authorize]
-		[Route("/api/v1/compute/{ClusterId}")]
-		public async Task<ActionResult> AddTasksAsync([FromRoute] ClusterId ClusterId, [FromBody] AddTasksRequest Request)
+		[Route("/api/v1/compute/{clusterId}")]
+		public async Task<ActionResult> AddTasksAsync([FromRoute] ClusterId clusterId, [FromBody] AddTasksRequest request)
 		{
-			if(!await AclService.AuthorizeAsync(AclAction.AddComputeTasks, User))
+			if(!await _aclService.AuthorizeAsync(AclAction.AddComputeTasks, User))
 			{
 				return Forbid(AclAction.AddComputeTasks);
 			}
-			if (Request.TaskRefIds.Count == 0)
+			if (request.TaskRefIds.Count == 0)
 			{
 				return BadRequest("No task hashes specified");
 			}
 
-			await ComputeService.AddTasksAsync(ClusterId, Request.ChannelId, Request.TaskRefIds, Request.RequirementsHash);
+			await _computeService.AddTasksAsync(clusterId, request.ChannelId, request.TaskRefIds, request.RequirementsHash);
 			return Ok();
 		}
 
 		/// <summary>
 		/// Read updates for a particular channel
 		/// </summary>
-		/// <param name="ClusterId"></param>
-		/// <param name="ChannelId">The channel to add to</param>
-		/// <param name="Wait">Amount of time to wait before responding, in seconds</param>
+		/// <param name="clusterId"></param>
+		/// <param name="channelId">The channel to add to</param>
+		/// <param name="wait">Amount of time to wait before responding, in seconds</param>
 		/// <returns></returns>
 		[HttpPost]
 		[Authorize]
-		[Route("/api/v1/compute/{ClusterId}/updates/{ChannelId}")]
-		public async Task<ActionResult<GetTaskUpdatesResponse>> GetUpdatesAsync([FromRoute] ClusterId ClusterId, [FromRoute] ChannelId ChannelId, [FromQuery] int Wait = 0)
+		[Route("/api/v1/compute/{clusterId}/updates/{channelId}")]
+		public async Task<ActionResult<GetTaskUpdatesResponse>> GetUpdatesAsync([FromRoute] ClusterId clusterId, [FromRoute] ChannelId channelId, [FromQuery] int wait = 0)
 		{
-			if (!await AclService.AuthorizeAsync(AclAction.ViewComputeTasks, User))
+			if (!await _aclService.AuthorizeAsync(AclAction.ViewComputeTasks, User))
 			{
 				return Forbid(AclAction.ViewComputeTasks);
 			}
 
-			List<IComputeTaskStatus> Results;
-			if (Wait == 0)
+			List<IComputeTaskStatus> results;
+			if (wait == 0)
 			{
-				Results = await ComputeService.GetTaskUpdatesAsync(ClusterId, ChannelId);
+				results = await _computeService.GetTaskUpdatesAsync(clusterId, channelId);
 			}
 			else
 			{
-				using CancellationTokenSource DelaySource = new CancellationTokenSource(Wait * 1000);
-				Results = await ComputeService.WaitForTaskUpdatesAsync(ClusterId, ChannelId, DelaySource.Token);
+				using CancellationTokenSource delaySource = new CancellationTokenSource(wait * 1000);
+				results = await _computeService.WaitForTaskUpdatesAsync(clusterId, channelId, delaySource.Token);
 			}
 
-			GetTaskUpdatesResponse Response = new GetTaskUpdatesResponse();
-			foreach (IComputeTaskStatus Result in Results)
+			GetTaskUpdatesResponse response = new GetTaskUpdatesResponse();
+			foreach (IComputeTaskStatus result in results)
 			{
-				GetTaskUpdateResponse Update = new GetTaskUpdateResponse();
+				GetTaskUpdateResponse update = new GetTaskUpdateResponse();
 
-				Update.TaskRefId = Result.TaskRefId;
-				Update.Time = Result.Time;
-				Update.State = Result.State;
-				Update.Outcome = Result.Outcome;
-				Update.ResultRefId = Result.ResultRefId;
-				Update.AgentId = Result.AgentId?.ToString();
-				Update.LeaseId = Result.LeaseId?.ToString();
-				Update.Detail = Result.Detail;
+				update.TaskRefId = result.TaskRefId;
+				update.Time = result.Time;
+				update.State = result.State;
+				update.Outcome = result.Outcome;
+				update.ResultRefId = result.ResultRefId;
+				update.AgentId = result.AgentId?.ToString();
+				update.LeaseId = result.LeaseId?.ToString();
+				update.Detail = result.Detail;
 
-				Response.Updates.Add(Update);
+				response.Updates.Add(update);
 			}
 
-			return Response;
+			return response;
 		}
 	}
 }
