Horde: Fix coding conventions within Horde.Build.

diff --git a/Engine/Source/Programs/Horde/Horde.Build/Compute/Impl/ComputeService.cs b/Engine/Source/Programs/Horde/Horde.Build/Compute/Impl/ComputeService.cs
--- a/Engine/Source/Programs/Horde/Horde.Build/Compute/Impl/ComputeService.cs
+++ b/Engine/Source/Programs/Horde/Horde.Build/Compute/Impl/ComputeService.cs
@@ -1,42 +1,27 @@
 // Copyright Epic Games, Inc. All Rights Reserved.
 
-using EpicGames.Core;
-using Google.Protobuf;
-using Google.Protobuf.Reflection;
-using Google.Protobuf.WellKnownTypes;
-using Grpc.Core;
-using HordeCommon;
-using HordeCommon.Rpc.Tasks;
-using Horde.Build.Api;
-using Horde.Build.Models;
-using Horde.Build.Services;
-using Microsoft.Extensions.Logging;
-using MongoDB.Bson;
 using System;
-using System.Collections.Concurrent;
 using System.Collections.Generic;
-using System.Diagnostics;
-using System.Diagnostics.CodeAnalysis;
 using System.Linq;
 using System.Threading;
 using System.Threading.Tasks;
-using Microsoft.Extensions.Options;
+using EpicGames.Core;
+using EpicGames.Horde.Compute;
+using EpicGames.Horde.Storage;
+using EpicGames.Redis;
+using EpicGames.Serialization;
+using Google.Protobuf;
+using Google.Protobuf.WellKnownTypes;
+using Horde.Build.Models;
+using Horde.Build.Services;
 using Horde.Build.Tasks;
 using Horde.Build.Utilities;
-using Horde.Build.Storage;
-using StackExchange.Redis;
-using MongoDB.Bson.Serialization.Attributes;
-using MongoDB.Bson.Serialization;
-using EpicGames.Serialization;
+using HordeCommon;
+using HordeCommon.Rpc.Tasks;
 using Microsoft.Extensions.Caching.Memory;
-using EpicGames.Horde.Compute;
 using Microsoft.Extensions.Hosting;
-using System.Threading.Channels;
-using EpicGames.Redis;
-using Horde.Build.Collections;
-using EpicGames.Horde.Storage;
-using EpicGames.Horde.Common;
-using System.Text.Json;
+using Microsoft.Extensions.Logging;
+using StackExchange.Redis;
 
 namespace Horde.Build.Compute.Impl
 {
@@ -61,11 +46,11 @@ namespace Horde.Build.Compute.Impl
 		{
 		}
 
-		public ComputeTaskInfo(ClusterId ClusterId, RefId TaskRefId, ChannelId ChannelId)
+		public ComputeTaskInfo(ClusterId clusterId, RefId taskRefId, ChannelId channelId)
 		{
-			this.ClusterId = ClusterId;
-			this.TaskRefId = TaskRefId;
-			this.ChannelId = ChannelId;
+			ClusterId = clusterId;
+			TaskRefId = taskRefId;
+			ChannelId = channelId;
 		}
 	}
 
@@ -114,20 +99,20 @@ namespace Horde.Build.Compute.Impl
 		public LeaseId? LeaseId
 		{
 			get => (LeaseIdBytes == null || LeaseIdBytes.Length == 0) ? (LeaseId?)null : new LeaseId(LeaseIdBytes);
-			set => LeaseIdBytes = (value != null) ? value.Value.Value.ToByteArray() : null;
+			set => LeaseIdBytes = value?.Value.ToByteArray();
 		}
 
 		private ComputeTaskStatus()
 		{
 		}
 
-		public ComputeTaskStatus(RefId TaskRefId, ComputeTaskState State, AgentId? AgentId, LeaseId? LeaseId)
+		public ComputeTaskStatus(RefId taskRefId, ComputeTaskState state, AgentId? agentId, LeaseId? leaseId)
 		{
-			this.TaskRefId = TaskRefId;
-			this.Time = DateTime.UtcNow;
-			this.State = State;
-			this.AgentId = (AgentId == null) ? Utf8String.Empty : AgentId.Value.ToString();
-			this.LeaseIdBytes = (LeaseId != null) ? LeaseId.Value.Value.ToByteArray() : null;
+			TaskRefId = taskRefId;
+			Time = DateTime.UtcNow;
+			State = state;
+			AgentId = (agentId == null) ? Utf8String.Empty : agentId.Value.ToString();
+			LeaseIdBytes = leaseId?.Value.ToByteArray();
 		}
 	}
 
@@ -142,10 +127,10 @@ namespace Horde.Build.Compute.Impl
 			public ClusterId ClusterId { get; set; }
 			public IoHash RequirementsHash { get; set; }
 
-			public QueueKey(ClusterId ClusterId, IoHash RequirementsHash)
+			public QueueKey(ClusterId clusterId, IoHash requirementsHash)
 			{
-				this.ClusterId = ClusterId;
-				this.RequirementsHash = RequirementsHash;
+				ClusterId = clusterId;
+				RequirementsHash = requirementsHash;
 			}
 
 			public override string ToString() => $"{ClusterId}/{RequirementsHash}";
@@ -153,14 +138,14 @@ namespace Horde.Build.Compute.Impl
 
 		class QueueKeySerializer : IRedisConverter<QueueKey>
 		{
-			public QueueKey FromRedisValue(RedisValue Value)
+			public QueueKey FromRedisValue(RedisValue value)
 			{
-				string Str = Value.ToString();
-				int Idx = Str.LastIndexOf("/", StringComparison.Ordinal);
-				return new QueueKey(new ClusterId(Str.Substring(0, Idx)), IoHash.Parse(Str.Substring(Idx + 1)));
+				string str = value.ToString();
+				int idx = str.LastIndexOf("/", StringComparison.Ordinal);
+				return new QueueKey(new ClusterId(str.Substring(0, idx)), IoHash.Parse(str.Substring(idx + 1)));
 			}
 
-			public RedisValue ToRedisValue(QueueKey Value) => $"{Value.ClusterId}/{Value.RequirementsHash}";
+			public RedisValue ToRedisValue(QueueKey value) => $"{value.ClusterId}/{value.RequirementsHash}";
 		}
 
 		class ClusterInfo : IComputeClusterInfo
@@ -170,12 +155,12 @@ namespace Horde.Build.Compute.Impl
 			public BucketId RequestBucketId { get; set; }
 			public BucketId ResponseBucketId { get; set; }
 
-			public ClusterInfo(ComputeClusterConfig Config)
+			public ClusterInfo(ComputeClusterConfig config)
 			{
-				Id = new ClusterId(Config.Id);
-				NamespaceId = new NamespaceId(Config.NamespaceId);
-				RequestBucketId = new BucketId(Config.RequestBucketId);
-				ResponseBucketId = new BucketId(Config.ResponseBucketId);
+				Id = new ClusterId(config.Id);
+				NamespaceId = new NamespaceId(config.NamespaceId);
+				RequestBucketId = new BucketId(config.RequestBucketId);
+				ResponseBucketId = new BucketId(config.ResponseBucketId);
 			}
 		}
 
@@ -187,62 +172,60 @@ namespace Horde.Build.Compute.Impl
 
 		public static NamespaceId DefaultNamespaceId { get; } = new NamespaceId("default");
 
-		IStorageClient StorageClient;
-		ITaskScheduler<QueueKey, ComputeTaskInfo> TaskScheduler;
-		RedisMessageQueue<ComputeTaskStatus> MessageQueue;
-		ITicker ExpireTasksTicker;
-		IMemoryCache RequirementsCache;
-		LazyCachedValue<Task<Globals>> Globals;
-		LazyCachedValue<Task<List<IPool>>> CachedPools;
-		ILogger Logger;
+		readonly IStorageClient _storageClient;
+		readonly ITaskScheduler<QueueKey, ComputeTaskInfo> _taskScheduler;
+		readonly RedisMessageQueue<ComputeTaskStatus> _messageQueue;
+		readonly ITicker _expireTasksTicker;
+		readonly IMemoryCache _requirementsCache;
+		readonly LazyCachedValue<Task<Globals>> _globals;
+		readonly ILogger _logger;
 
 		/// <summary>
 		/// Constructor
 		/// </summary>
-		public ComputeService(DatabaseService DatabaseService, IDatabase Redis, IStorageClient StorageClient, IPoolCollection PoolCollection, IClock Clock, ILogger<ComputeService> Logger)
+		public ComputeService(DatabaseService databaseService, IDatabase redis, IStorageClient storageClient, IClock clock, ILogger<ComputeService> logger)
 		{
-			this.StorageClient = StorageClient;
-			this.TaskScheduler = new RedisTaskScheduler<QueueKey, ComputeTaskInfo>(Redis, "compute/tasks/", Logger);
-			this.MessageQueue = new RedisMessageQueue<ComputeTaskStatus>(Redis, "compute/messages/");
-			this.ExpireTasksTicker = Clock.AddTicker(TimeSpan.FromMinutes(2.0), ExpireTasksAsync, Logger);
-			this.RequirementsCache = new MemoryCache(new MemoryCacheOptions());
-			this.Globals = new LazyCachedValue<Task<Globals>>(() => DatabaseService.GetGlobalsAsync(), TimeSpan.FromSeconds(120.0));
-			this.CachedPools = new LazyCachedValue<Task<List<IPool>>>(() => PoolCollection.GetAsync(), TimeSpan.FromSeconds(30.0));
-			this.Logger = Logger;
+			_storageClient = storageClient;
+			_taskScheduler = new RedisTaskScheduler<QueueKey, ComputeTaskInfo>(redis, "compute/tasks/", logger);
+			_messageQueue = new RedisMessageQueue<ComputeTaskStatus>(redis, "compute/messages/");
+			_expireTasksTicker = clock.AddTicker<ComputeService>(TimeSpan.FromMinutes(2.0), ExpireTasksAsync, logger);
+			_requirementsCache = new MemoryCache(new MemoryCacheOptions());
+			_globals = new LazyCachedValue<Task<Globals>>(() => databaseService.GetGlobalsAsync(), TimeSpan.FromSeconds(120.0));
+			_logger = logger;
 
 			OnLeaseStartedProperties.Add(x => x.TaskRefId);
 		}
 
 		/// <inheritdoc/>
-		public Task StartAsync(CancellationToken Token) => ExpireTasksTicker.StartAsync();
+		public Task StartAsync(CancellationToken token) => _expireTasksTicker.StartAsync();
 
 		/// <inheritdoc/>
-		public Task StopAsync(CancellationToken Token) => ExpireTasksTicker.StopAsync();
+		public Task StopAsync(CancellationToken token) => _expireTasksTicker.StopAsync();
 
 		/// <summary>
 		/// Expire tasks that are in inactive queues (ie. no machines can execute them)
 		/// </summary>
-		/// <param name="CancellationToken"></param>
+		/// <param name="cancellationToken"></param>
 		/// <returns></returns>
-		async ValueTask ExpireTasksAsync(CancellationToken CancellationToken)
+		async ValueTask ExpireTasksAsync(CancellationToken cancellationToken)
 		{
-			List<QueueKey> QueueKeys = await TaskScheduler.GetInactiveQueuesAsync();
-			foreach (QueueKey QueueKey in QueueKeys)
+			List<QueueKey> queueKeys = await _taskScheduler.GetInactiveQueuesAsync();
+			foreach (QueueKey queueKey in queueKeys)
 			{
-				Logger.LogInformation("Inactive queue: {QueueKey}", QueueKey);
+				_logger.LogInformation("Inactive queue: {QueueKey}", queueKey);
 				for (; ; )
 				{
-					ComputeTaskInfo? ComputeTask = await TaskScheduler.DequeueAsync(QueueKey);
-					if (ComputeTask == null)
+					ComputeTaskInfo? computeTask = await _taskScheduler.DequeueAsync(queueKey);
+					if (computeTask == null)
 					{
 						break;
 					}
 
-					ComputeTaskStatus Status = new ComputeTaskStatus(ComputeTask.TaskRefId, ComputeTaskState.Complete, null, null);
-					Status.Outcome = ComputeTaskOutcome.Expired;
-					Status.Detail = $"No agents monitoring queue {QueueKey}";
-					Logger.LogInformation("Compute task expired (queue: {RequirementsHash}, task: {TaskHash}, channel: {ChannelId})", QueueKey, ComputeTask.TaskRefId, ComputeTask.ChannelId);
-					await PostStatusMessageAsync(ComputeTask, Status);
+					ComputeTaskStatus status = new ComputeTaskStatus(computeTask.TaskRefId, ComputeTaskState.Complete, null, null);
+					status.Outcome = ComputeTaskOutcome.Expired;
+					status.Detail = $"No agents monitoring queue {queueKey}";
+					_logger.LogInformation("Compute task expired (queue: {RequirementsHash}, task: {TaskHash}, channel: {ChannelId})", queueKey, computeTask.TaskRefId, computeTask.ChannelId);
+					await PostStatusMessageAsync(computeTask, status);
 				}
 			}
 		}
@@ -250,276 +233,276 @@ namespace Horde.Build.Compute.Impl
 		/// <inheritdoc/>
 		public void Dispose()
 		{
-			MessageQueue.Dispose();
-			ExpireTasksTicker.Dispose();
-			RequirementsCache.Dispose();
+			_messageQueue.Dispose();
+			_expireTasksTicker.Dispose();
+			_requirementsCache.Dispose();
 		}
 
 		/// <summary>
 		/// Gets information about a compute cluster
 		/// </summary>
-		/// <param name="ClusterId">Cluster to use for execution</param>
-		public async Task<IComputeClusterInfo> GetClusterInfoAsync(ClusterId ClusterId)
+		/// <param name="clusterId">Cluster to use for execution</param>
+		public async Task<IComputeClusterInfo> GetClusterInfoAsync(ClusterId clusterId)
 		{
-			ComputeClusterConfig? Config = await GetClusterAsync(ClusterId);
-			if (Config == null)
+			ComputeClusterConfig? config = await GetClusterAsync(clusterId);
+			if (config == null)
 			{
 				throw new KeyNotFoundException();
 			}
-			return new ClusterInfo(Config);
+			return new ClusterInfo(config);
 		}
 
 		/// <inheritdoc/>
-		public async Task AddTasksAsync(ClusterId ClusterId, ChannelId ChannelId, List<RefId> TaskRefIds, CbObjectAttachment RequirementsHash)
+		public async Task AddTasksAsync(ClusterId clusterId, ChannelId channelId, List<RefId> taskRefIds, CbObjectAttachment requirementsHash)
 		{
-			List<Task> Tasks = new List<Task>();
-			foreach (RefId TaskRefId in TaskRefIds)
+			List<Task> tasks = new List<Task>();
+			foreach (RefId taskRefId in taskRefIds)
 			{
-				ComputeTaskInfo TaskInfo = new ComputeTaskInfo(ClusterId, TaskRefId, ChannelId);
-				Logger.LogDebug("Adding task {TaskHash} from channel {ChannelId} to queue {ClusterId}{RequirementsHash}", TaskRefId.Hash, ChannelId, ClusterId, RequirementsHash);
-				Tasks.Add(TaskScheduler.EnqueueAsync(new QueueKey(ClusterId, RequirementsHash), TaskInfo, false));
+				ComputeTaskInfo taskInfo = new ComputeTaskInfo(clusterId, taskRefId, channelId);
+				_logger.LogDebug("Adding task {TaskHash} from channel {ChannelId} to queue {ClusterId}{RequirementsHash}", taskRefId.Hash, channelId, clusterId, requirementsHash);
+				tasks.Add(_taskScheduler.EnqueueAsync(new QueueKey(clusterId, requirementsHash), taskInfo, false));
 			}
-			await Task.WhenAll(Tasks);
+			await Task.WhenAll(tasks);
 		}
 
-		async ValueTask<ComputeClusterConfig?> GetClusterAsync(ClusterId ClusterId)
+		async ValueTask<ComputeClusterConfig?> GetClusterAsync(ClusterId clusterId)
 		{
-			Globals GlobalsInstance = await Globals.GetCached();
-			return GlobalsInstance.ComputeClusters.FirstOrDefault(x => new ClusterId(x.Id) == ClusterId);
+			Globals globalsInstance = await _globals.GetCached();
+			return globalsInstance.ComputeClusters.FirstOrDefault(x => new ClusterId(x.Id) == clusterId);
 		}
 
 		/// <inheritdoc/>
-		public override async Task<AgentLease?> AssignLeaseAsync(IAgent Agent, CancellationToken CancellationToken)
+		public override async Task<AgentLease?> AssignLeaseAsync(IAgent agent, CancellationToken cancellationToken)
 		{
 			// If the agent is disabled, just block until we cancel
-			if (!Agent.Enabled)
+			if (!agent.Enabled)
 			{
-				using CancellationTask Task = new CancellationTask(CancellationToken);
-				await Task.Task;
+				using CancellationTask task = new CancellationTask(cancellationToken);
+				await task.Task;
 				return null;
 			}
 
 			// Find a task to execute
-			(QueueKey, ComputeTaskInfo)? Entry = await TaskScheduler.DequeueAsync(QueueKey => CheckRequirements(Agent, QueueKey), CancellationToken);
-			if (Entry != null)
+			(QueueKey, ComputeTaskInfo)? entry = await _taskScheduler.DequeueAsync(queueKey => CheckRequirements(agent, queueKey), cancellationToken);
+			if (entry != null)
 			{
-				(QueueKey QueueKey, ComputeTaskInfo TaskInfo) = Entry.Value;
+				(QueueKey queueKey, ComputeTaskInfo taskInfo) = entry.Value;
 
-				ComputeClusterConfig? Cluster = await GetClusterAsync(TaskInfo.ClusterId);
-				if (Cluster == null)
+				ComputeClusterConfig? cluster = await GetClusterAsync(taskInfo.ClusterId);
+				if (cluster == null)
 				{
-					Logger.LogWarning("Invalid cluster '{ClusterId}'; failing task {TaskRefId}", TaskInfo.ClusterId, TaskInfo.TaskRefId);
-					ComputeTaskStatus Status = new ComputeTaskStatus(TaskInfo.TaskRefId, ComputeTaskState.Complete, Agent.Id, null) { Detail = $"Invalid cluster '{TaskInfo.ClusterId}'" };
-					await PostStatusMessageAsync(TaskInfo, Status);
+					_logger.LogWarning("Invalid cluster '{ClusterId}'; failing task {TaskRefId}", taskInfo.ClusterId, taskInfo.TaskRefId);
+					ComputeTaskStatus status = new ComputeTaskStatus(taskInfo.TaskRefId, ComputeTaskState.Complete, agent.Id, null) { Detail = $"Invalid cluster '{taskInfo.ClusterId}'" };
+					await PostStatusMessageAsync(taskInfo, status);
 					return null;
 				}
 
-				Requirements? Requirements = await GetCachedRequirementsAsync(QueueKey);
-				if (Requirements == null)
+				Requirements? requirements = await GetCachedRequirementsAsync(queueKey);
+				if (requirements == null)
 				{
-					Logger.LogWarning("Unable to fetch requirements {RequirementsHash}", QueueKey);
-					ComputeTaskStatus Status = new ComputeTaskStatus(TaskInfo.TaskRefId, ComputeTaskState.Complete, Agent.Id, null) { Detail = $"Unable to retrieve requirements '{QueueKey}'" };
-					await PostStatusMessageAsync(TaskInfo, Status);
+					_logger.LogWarning("Unable to fetch requirements {RequirementsHash}", queueKey);
+					ComputeTaskStatus status = new ComputeTaskStatus(taskInfo.TaskRefId, ComputeTaskState.Complete, agent.Id, null) { Detail = $"Unable to retrieve requirements '{queueKey}'" };
+					await PostStatusMessageAsync(taskInfo, status);
 					return null;
 				}
 
-				ComputeTaskMessage ComputeTask = new ComputeTaskMessage();
-				ComputeTask.ClusterId = TaskInfo.ClusterId.ToString();
-				ComputeTask.ChannelId = TaskInfo.ChannelId.ToString();
-				ComputeTask.NamespaceId = Cluster.NamespaceId.ToString();
-				ComputeTask.InputBucketId = Cluster.RequestBucketId.ToString();
-				ComputeTask.OutputBucketId = Cluster.ResponseBucketId.ToString();
-				ComputeTask.RequirementsHash = QueueKey.RequirementsHash;
-				ComputeTask.TaskRefId = TaskInfo.TaskRefId;
-
-				string LeaseName = $"Remote action ({TaskInfo.TaskRefId})";
-				byte[] Payload = Any.Pack(ComputeTask).ToByteArray();
-
-				AgentLease Lease = new AgentLease(LeaseId.GenerateNewId(), LeaseName, null, null, null, LeaseState.Pending, Requirements.Resources, Requirements.Exclusive, Payload);
-				Logger.LogDebug("Created lease {LeaseId} for channel {ChannelId} task {TaskHash} req {RequirementsHash}", Lease.Id, ComputeTask.ChannelId, ComputeTask.TaskRefId, ComputeTask.RequirementsHash);
-				return Lease;
+				ComputeTaskMessage computeTask = new ComputeTaskMessage();
+				computeTask.ClusterId = taskInfo.ClusterId.ToString();
+				computeTask.ChannelId = taskInfo.ChannelId.ToString();
+				computeTask.NamespaceId = cluster.NamespaceId.ToString();
+				computeTask.InputBucketId = cluster.RequestBucketId.ToString();
+				computeTask.OutputBucketId = cluster.ResponseBucketId.ToString();
+				computeTask.RequirementsHash = queueKey.RequirementsHash;
+				computeTask.TaskRefId = taskInfo.TaskRefId;
+
+				string leaseName = $"Remote action ({taskInfo.TaskRefId})";
+				byte[] payload = Any.Pack(computeTask).ToByteArray();
+
+				AgentLease lease = new AgentLease(LeaseId.GenerateNewId(), leaseName, null, null, null, LeaseState.Pending, requirements.Resources, requirements.Exclusive, payload);
+				_logger.LogDebug("Created lease {LeaseId} for channel {ChannelId} task {TaskHash} req {RequirementsHash}", lease.Id, computeTask.ChannelId, computeTask.TaskRefId, computeTask.RequirementsHash);
+				return lease;
 			}
 			return null;
 		}
 
 		/// <inheritdoc/>
-		public override Task CancelLeaseAsync(IAgent Agent, LeaseId LeaseId, ComputeTaskMessage Message)
+		public override Task CancelLeaseAsync(IAgent agent, LeaseId leaseId, ComputeTaskMessage message)
 		{
-			ClusterId ClusterId = new ClusterId(Message.ClusterId);
-			ComputeTaskInfo TaskInfo = new ComputeTaskInfo(ClusterId, new RefId(new IoHash(Message.TaskRefId.ToByteArray())), new ChannelId(Message.ChannelId));
-			return TaskScheduler.EnqueueAsync(new QueueKey(ClusterId, new IoHash(Message.RequirementsHash.ToByteArray())), TaskInfo, true);
+			ClusterId clusterId = new ClusterId(message.ClusterId);
+			ComputeTaskInfo taskInfo = new ComputeTaskInfo(clusterId, new RefId(new IoHash(message.TaskRefId.ToByteArray())), new ChannelId(message.ChannelId));
+			return _taskScheduler.EnqueueAsync(new QueueKey(clusterId, new IoHash(message.RequirementsHash.ToByteArray())), taskInfo, true);
 		}
 
 		/// <inheritdoc/>
-		public async Task<List<IComputeTaskStatus>> GetTaskUpdatesAsync(ClusterId ClusterId, ChannelId ChannelId)
+		public async Task<List<IComputeTaskStatus>> GetTaskUpdatesAsync(ClusterId clusterId, ChannelId channelId)
 		{
-			List<ComputeTaskStatus> Messages = await MessageQueue.ReadMessagesAsync(GetMessageQueueId(ClusterId, ChannelId));
-			return Messages.ConvertAll<IComputeTaskStatus>(x => x);
+			List<ComputeTaskStatus> messages = await _messageQueue.ReadMessagesAsync(GetMessageQueueId(clusterId, channelId));
+			return messages.ConvertAll<IComputeTaskStatus>(x => x);
 		}
 
-		public async Task<List<IComputeTaskStatus>> WaitForTaskUpdatesAsync(ClusterId ClusterId, ChannelId ChannelId, CancellationToken CancellationToken)
+		public async Task<List<IComputeTaskStatus>> WaitForTaskUpdatesAsync(ClusterId clusterId, ChannelId channelId, CancellationToken cancellationToken)
 		{
-			List<ComputeTaskStatus> Messages = await MessageQueue.WaitForMessagesAsync(GetMessageQueueId(ClusterId, ChannelId), CancellationToken);
-			return Messages.ConvertAll<IComputeTaskStatus>(x => x);
+			List<ComputeTaskStatus> messages = await _messageQueue.WaitForMessagesAsync(GetMessageQueueId(clusterId, channelId), cancellationToken);
+			return messages.ConvertAll<IComputeTaskStatus>(x => x);
 		}
 
-		public override async Task OnLeaseStartedAsync(IAgent Agent, LeaseId LeaseId, ComputeTaskMessage ComputeTask, ILogger Logger)
+		public override async Task OnLeaseStartedAsync(IAgent agent, LeaseId leaseId, ComputeTaskMessage computeTask, ILogger logger)
 		{
-			await base.OnLeaseStartedAsync(Agent, LeaseId, ComputeTask, Logger);
+			await base.OnLeaseStartedAsync(agent, leaseId, computeTask, logger);
 
-			ComputeTaskStatus Status = new ComputeTaskStatus(ComputeTask.TaskRefId, ComputeTaskState.Executing, Agent.Id, LeaseId);
-			await PostStatusMessageAsync(ComputeTask, Status);
+			ComputeTaskStatus status = new ComputeTaskStatus(computeTask.TaskRefId, ComputeTaskState.Executing, agent.Id, leaseId);
+			await PostStatusMessageAsync(computeTask, status);
 		}
 
-		public override async Task OnLeaseFinishedAsync(IAgent Agent, LeaseId LeaseId, ComputeTaskMessage ComputeTask, LeaseOutcome Outcome, ReadOnlyMemory<byte> Output, ILogger Logger)
+		public override async Task OnLeaseFinishedAsync(IAgent agent, LeaseId leaseId, ComputeTaskMessage computeTask, LeaseOutcome outcome, ReadOnlyMemory<byte> output, ILogger logger)
 		{
-			await base.OnLeaseFinishedAsync(Agent, LeaseId, ComputeTask, Outcome, Output, Logger);
+			await base.OnLeaseFinishedAsync(agent, leaseId, computeTask, outcome, output, logger);
 
-			ComputeTaskResultMessage Message = ComputeTaskResultMessage.Parser.ParseFrom(Output.ToArray());
+			ComputeTaskResultMessage message = ComputeTaskResultMessage.Parser.ParseFrom(output.ToArray());
 
-			ComputeTaskStatus Status = new ComputeTaskStatus(ComputeTask.TaskRefId, ComputeTaskState.Complete, Agent.Id, LeaseId);
-			if (Message.ResultRefId != null)
+			ComputeTaskStatus status = new ComputeTaskStatus(computeTask.TaskRefId, ComputeTaskState.Complete, agent.Id, leaseId);
+			if (message.ResultRefId != null)
 			{
-				Status.ResultRefId = Message.ResultRefId;
+				status.ResultRefId = message.ResultRefId;
 			}
-			else if ((ComputeTaskOutcome)Message.Outcome != ComputeTaskOutcome.Success)
+			else if ((ComputeTaskOutcome)message.Outcome != ComputeTaskOutcome.Success)
 			{
-				(Status.Outcome, Status.Detail) = ((ComputeTaskOutcome)Message.Outcome, Message.Detail);
+				(status.Outcome, status.Detail) = ((ComputeTaskOutcome)message.Outcome, message.Detail);
 			}
-			else if (Outcome == LeaseOutcome.Failed)
+			else if (outcome == LeaseOutcome.Failed)
 			{
-				Status.Outcome = ComputeTaskOutcome.Failed;
+				status.Outcome = ComputeTaskOutcome.Failed;
 			}
-			else if (Outcome == LeaseOutcome.Cancelled)
+			else if (outcome == LeaseOutcome.Cancelled)
 			{
-				Status.Outcome = ComputeTaskOutcome.Cancelled;
+				status.Outcome = ComputeTaskOutcome.Cancelled;
 			}
 			else
 			{
-				Status.Outcome = ComputeTaskOutcome.NoResult;
+				status.Outcome = ComputeTaskOutcome.NoResult;
 			}
 
-			Logger.LogInformation("Compute lease finished (lease: {LeaseId}, task: {TaskHash}, agent: {AgentId}, channel: {ChannelId}, result: {ResultHash}, outcome: {Outcome})", LeaseId, ComputeTask.TaskRefId.AsRefId(), Agent.Id, ComputeTask.ChannelId, Status.ResultRefId?.ToString() ?? "(none)", Status.Outcome);
-			await PostStatusMessageAsync(ComputeTask, Status);
+			logger.LogInformation("Compute lease finished (lease: {LeaseId}, task: {TaskHash}, agent: {AgentId}, channel: {ChannelId}, result: {ResultHash}, outcome: {Outcome})", leaseId, computeTask.TaskRefId.AsRefId(), agent.Id, computeTask.ChannelId, status.ResultRefId?.ToString() ?? "(none)", status.Outcome);
+			await PostStatusMessageAsync(computeTask, status);
 		}
 
 		/// <summary>
 		/// Checks that an agent matches the necessary criteria to execute a task
 		/// </summary>
-		/// <param name="Agent"></param>
-		/// <param name="QueueKey"></param>
+		/// <param name="agent"></param>
+		/// <param name="queueKey"></param>
 		/// <returns></returns>
-		async ValueTask<bool> CheckRequirements(IAgent Agent, QueueKey QueueKey)
+		async ValueTask<bool> CheckRequirements(IAgent agent, QueueKey queueKey)
 		{
-			Requirements? Requirements = await GetCachedRequirementsAsync(QueueKey);
-			if (Requirements == null)
+			Requirements? requirements = await GetCachedRequirementsAsync(queueKey);
+			if (requirements == null)
 			{
 				return false;
 			}
-			return Agent.MeetsRequirements(Requirements);
+			return agent.MeetsRequirements(requirements);
 		}
 
 		/// <summary>
 		/// Gets the requirements object from the CAS
 		/// </summary>
-		/// <param name="QueueKey">Queue identifier</param>
+		/// <param name="queueKey">Queue identifier</param>
 		/// <returns>Requirements object for the queue</returns>
-		async ValueTask<Requirements?> GetCachedRequirementsAsync(QueueKey QueueKey)
+		async ValueTask<Requirements?> GetCachedRequirementsAsync(QueueKey queueKey)
 		{
-			Requirements? Requirements;
-			if (!RequirementsCache.TryGetValue(QueueKey.RequirementsHash, out Requirements))
+			Requirements? requirements;
+			if (!_requirementsCache.TryGetValue(queueKey.RequirementsHash, out requirements))
 			{
-				Requirements = await GetRequirementsAsync(QueueKey);
-				if (Requirements != null)
+				requirements = await GetRequirementsAsync(queueKey);
+				if (requirements != null)
 				{
-					using (ICacheEntry Entry = RequirementsCache.CreateEntry(QueueKey.RequirementsHash))
+					using (ICacheEntry entry = _requirementsCache.CreateEntry(queueKey.RequirementsHash))
 					{
-						Entry.SetSlidingExpiration(TimeSpan.FromMinutes(10.0));
-						Entry.SetValue(Requirements);
+						entry.SetSlidingExpiration(TimeSpan.FromMinutes(10.0));
+						entry.SetValue(requirements);
 					}
 				}
 			}
-			return Requirements;
+			return requirements;
 		}
 
 		/// <summary>
 		/// Gets the requirements object for a given queue. Fails tasks in the queue if the requirements object is missing.
 		/// </summary>
-		/// <param name="QueueKey">Queue identifier</param>
+		/// <param name="queueKey">Queue identifier</param>
 		/// <returns>Requirements object for the queue</returns>
-		async ValueTask<Requirements?> GetRequirementsAsync(QueueKey QueueKey)
+		async ValueTask<Requirements?> GetRequirementsAsync(QueueKey queueKey)
 		{
-			Requirements? Requirements = null;
+			Requirements? requirements = null;
 
-			ComputeClusterConfig? ClusterConfig = await GetClusterAsync(QueueKey.ClusterId);
-			if (ClusterConfig != null)
+			ComputeClusterConfig? clusterConfig = await GetClusterAsync(queueKey.ClusterId);
+			if (clusterConfig != null)
 			{
-				NamespaceId NamespaceId = new NamespaceId(ClusterConfig.NamespaceId);
+				NamespaceId namespaceId = new NamespaceId(clusterConfig.NamespaceId);
 				try
 				{
-					Requirements = await StorageClient.ReadBlobAsync<Requirements>(NamespaceId, QueueKey.RequirementsHash);
+					requirements = await _storageClient.ReadBlobAsync<Requirements>(namespaceId, queueKey.RequirementsHash);
 				}
 				catch (BlobNotFoundException)
 				{
 				}
-				catch (Exception Ex)
+				catch (Exception ex)
 				{
-					Logger.LogError(Ex, "Unable to read blob {NamespaceId}/{RequirementsHash} from storage service", ClusterConfig.NamespaceId, QueueKey.RequirementsHash);
+					_logger.LogError(ex, "Unable to read blob {NamespaceId}/{RequirementsHash} from storage service", clusterConfig.NamespaceId, queueKey.RequirementsHash);
 				}
 			}
 
-			if (Requirements == null)
+			if (requirements == null)
 			{
-				Logger.LogWarning("Unable to fetch requirements object for queue {QueueKey}; failing queued tasks.", QueueKey);
+				_logger.LogWarning("Unable to fetch requirements object for queue {QueueKey}; failing queued tasks.", queueKey);
 				for (; ; )
 				{
-					ComputeTaskInfo? ComputeTask = await TaskScheduler.DequeueAsync(QueueKey);
-					if (ComputeTask == null)
+					ComputeTaskInfo? computeTask = await _taskScheduler.DequeueAsync(queueKey);
+					if (computeTask == null)
 					{
 						break;
 					}
 
-					ComputeTaskStatus Status = new ComputeTaskStatus(ComputeTask.TaskRefId, ComputeTaskState.Complete, null, null);
-					Status.Outcome = ComputeTaskOutcome.BlobNotFound;
-					Status.Detail = $"Missing requirements object {QueueKey.RequirementsHash}";
-					Logger.LogInformation("Compute task failed due to missing requirements (queue: {QueueKey}, task: {TaskHash}, channel: {ChannelId})", QueueKey, ComputeTask.TaskRefId, ComputeTask.ChannelId);
-					await PostStatusMessageAsync(ComputeTask, Status);
+					ComputeTaskStatus status = new ComputeTaskStatus(computeTask.TaskRefId, ComputeTaskState.Complete, null, null);
+					status.Outcome = ComputeTaskOutcome.BlobNotFound;
+					status.Detail = $"Missing requirements object {queueKey.RequirementsHash}";
+					_logger.LogInformation("Compute task failed due to missing requirements (queue: {QueueKey}, task: {TaskHash}, channel: {ChannelId})", queueKey, computeTask.TaskRefId, computeTask.ChannelId);
+					await PostStatusMessageAsync(computeTask, status);
 				}
 			}
 
-			return Requirements;
+			return requirements;
 		}
 
 		/// <summary>
 		/// Post a status message for a particular task
 		/// </summary>
-		/// <param name="ComputeTask">The compute task instance</param>
-		/// <param name="Status">New status for the task</param>
-		async Task PostStatusMessageAsync(ComputeTaskInfo ComputeTask, ComputeTaskStatus Status)
+		/// <param name="computeTask">The compute task instance</param>
+		/// <param name="status">New status for the task</param>
+		async Task PostStatusMessageAsync(ComputeTaskInfo computeTask, ComputeTaskStatus status)
 		{
-			await MessageQueue.PostAsync(GetMessageQueueId(ComputeTask.ClusterId, ComputeTask.ChannelId), Status);
+			await _messageQueue.PostAsync(GetMessageQueueId(computeTask.ClusterId, computeTask.ChannelId), status);
 		}
 
 		/// <summary>
 		/// Post a status message for a particular task
 		/// </summary>
-		/// <param name="ComputeTaskMessage">The compute task lease</param>
-		/// <param name="Status">New status for the task</param>
+		/// <param name="computeTaskMessage">The compute task lease</param>
+		/// <param name="status">New status for the task</param>
 		/// <returns></returns>
-		async Task PostStatusMessageAsync(ComputeTaskMessage ComputeTaskMessage, ComputeTaskStatus Status)
+		async Task PostStatusMessageAsync(ComputeTaskMessage computeTaskMessage, ComputeTaskStatus status)
 		{
-			await MessageQueue.PostAsync(GetMessageQueueId(new ClusterId(ComputeTaskMessage.ClusterId), new ChannelId(ComputeTaskMessage.ChannelId)), Status);
+			await _messageQueue.PostAsync(GetMessageQueueId(new ClusterId(computeTaskMessage.ClusterId), new ChannelId(computeTaskMessage.ChannelId)), status);
 		}
 
 		/// <summary>
 		/// Gets the name of a particular message queue
 		/// </summary>
-		/// <param name="ClusterId">The compute cluster</param>
-		/// <param name="ChannelId">Identifier for the message channel</param>
+		/// <param name="clusterId">The compute cluster</param>
+		/// <param name="channelId">Identifier for the message channel</param>
 		/// <returns>Name of the message queue</returns>
-		static string GetMessageQueueId(ClusterId ClusterId, ChannelId ChannelId)
+		static string GetMessageQueueId(ClusterId clusterId, ChannelId channelId)
 		{
-			return $"{ClusterId}/{ChannelId}";
+			return $"{clusterId}/{channelId}";
 		}
 	}
 }
diff --git a/Engine/Source/Programs/Horde/Horde.Build/Compute/Impl/RedisMessageQueue.cs b/Engine/Source/Programs/Horde/Horde.Build/Compute/Impl/RedisMessageQueue.cs
--- a/Engine/Source/Programs/Horde/Horde.Build/Compute/Impl/RedisMessageQueue.cs
+++ b/Engine/Source/Programs/Horde/Horde.Build/Compute/Impl/RedisMessageQueue.cs
@@ -1,15 +1,11 @@
 // Copyright Epic Games, Inc. All Rights Reserved.
 
-using EpicGames.Redis;
-using Horde.Build.Compute.Impl;
-using Microsoft.Extensions.Logging;
-using StackExchange.Redis;
 using System;
-using System.Collections.Concurrent;
 using System.Collections.Generic;
-using System.Linq;
 using System.Threading;
 using System.Threading.Tasks;
+using EpicGames.Redis;
+using StackExchange.Redis;
 
 namespace Horde.Build.Compute.Impl
 {
@@ -22,24 +18,24 @@ namespace Horde.Build.Compute.Impl
 		/// <summary>
 		/// Adds messages to the queue
 		/// </summary>
-		/// <param name="ChannelId">The channel to post to</param>
-		/// <param name="Message">Message to post</param>
-		Task PostAsync(string ChannelId, T Message);
+		/// <param name="channelId">The channel to post to</param>
+		/// <param name="message">Message to post</param>
+		Task PostAsync(string channelId, T message);
 
 		/// <summary>
 		/// Waits for a message to be available on the given channel
 		/// </summary>
-		/// <param name="ChannelId">The channel to read from</param>
+		/// <param name="channelId">The channel to read from</param>
 		/// <returns>True if a message was available, false otherwise</returns>
-		Task<List<T>> ReadMessagesAsync(string ChannelId);
+		Task<List<T>> ReadMessagesAsync(string channelId);
 
 		/// <summary>
 		/// Waits for a message to be available on the given channel
 		/// </summary>
-		/// <param name="ChannelId">The channel to read from</param>
-		/// <param name="CancellationToken">May be signalled to stop the wait, without throwing an exception</param>
+		/// <param name="channelId">The channel to read from</param>
+		/// <param name="cancellationToken">May be signalled to stop the wait, without throwing an exception</param>
 		/// <returns>True if a message was available, false otherwise</returns>
-		Task<List<T>> WaitForMessagesAsync(string ChannelId, CancellationToken CancellationToken);
+		Task<List<T>> WaitForMessagesAsync(string channelId, CancellationToken cancellationToken);
 	}
 
 	/// <summary>
@@ -48,10 +44,10 @@ namespace Horde.Build.Compute.Impl
 	/// <typeparam name="T">The message type. Must be serailizable by RedisSerializer.</typeparam>
 	class RedisMessageQueue<T> : IMessageQueue<T>, IDisposable where T : class
 	{
-		IDatabase Redis;
-		RedisKey KeyPrefix;
-		RedisChannel UpdateChannel;
-		Dictionary<string, TaskCompletionSource<bool>> ChannelWakeEvents = new Dictionary<string, TaskCompletionSource<bool>>();
+		readonly IDatabase _redis;
+		readonly RedisKey _keyPrefix;
+		readonly RedisChannel _updateChannel;
+		readonly Dictionary<string, TaskCompletionSource<bool>> _channelWakeEvents = new Dictionary<string, TaskCompletionSource<bool>>();
 
 		// Time after which entries should be removed
 		TimeSpan ExpireTime { get; set; } = TimeSpan.FromSeconds(30);
@@ -59,126 +55,126 @@ namespace Horde.Build.Compute.Impl
 		/// <summary>
 		/// Constructor
 		/// </summary>
-		/// <param name="Redis">The Redis database instance</param>
-		/// <param name="KeyPrefix">Prefix for keys to use in the database</param>
-		public RedisMessageQueue(IDatabase Redis, RedisKey KeyPrefix)
+		/// <param name="redis">The Redis database instance</param>
+		/// <param name="keyPrefix">Prefix for keys to use in the database</param>
+		public RedisMessageQueue(IDatabase redis, RedisKey keyPrefix)
 		{
-			this.Redis = Redis;
-			this.KeyPrefix = KeyPrefix;
-			this.UpdateChannel = KeyPrefix.Append("updates").ToString();
+			_redis = redis;
+			_keyPrefix = keyPrefix;
+			_updateChannel = keyPrefix.Append("updates").ToString();
 
-			Redis.Multiplexer.GetSubscriber().Subscribe(UpdateChannel, OnChannelUpdate);
+			redis.Multiplexer.GetSubscriber().Subscribe(_updateChannel, OnChannelUpdate);
 		}
 
 		/// <inheritdoc/>
 		public void Dispose()
 		{
-			Redis.Multiplexer.GetSubscriber().Unsubscribe(UpdateChannel, OnChannelUpdate);
+			_redis.Multiplexer.GetSubscriber().Unsubscribe(_updateChannel, OnChannelUpdate);
 		}
 
 		/// <summary>
 		/// Callback for a message being posted to a channel
 		/// </summary>
-		/// <param name="Channel"></param>
-		/// <param name="Value"></param>
-		void OnChannelUpdate(RedisChannel Channel, RedisValue Value)
+		/// <param name="channel"></param>
+		/// <param name="value"></param>
+		void OnChannelUpdate(RedisChannel channel, RedisValue value)
 		{
-			TaskCompletionSource<bool>? CompletionSource;
-			lock (ChannelWakeEvents)
+			TaskCompletionSource<bool>? completionSource;
+			lock (_channelWakeEvents)
 			{
-				ChannelWakeEvents.TryGetValue(Value.ToString(), out CompletionSource);
+				_channelWakeEvents.TryGetValue(value.ToString(), out completionSource);
 			}
-			if (CompletionSource != null)
+			if (completionSource != null)
 			{
-				CompletionSource.TrySetResult(true);
+				completionSource.TrySetResult(true);
 			}
 		}
 
 		/// <summary>
 		/// Gets the set of messages for a channel
 		/// </summary>
-		/// <param name="ChannelId"></param>
+		/// <param name="channelId"></param>
 		/// <returns></returns>
-		RedisList<T> GetChannel(string ChannelId)
+		RedisList<T> GetChannel(string channelId)
 		{
-			return new RedisList<T>(Redis, KeyPrefix.Append(ChannelId));
+			return new RedisList<T>(_redis, _keyPrefix.Append(channelId));
 		}
 
 		/// <inheritdoc/>
-		public async Task PostAsync(string ChannelId, T Message)
+		public async Task PostAsync(string channelId, T message)
 		{
-			RedisList<T> Channel = GetChannel(ChannelId);
+			RedisList<T> channel = GetChannel(channelId);
 
-			long Length = await Channel.RightPushAsync(Message);
-			if (Length == 1)
+			long length = await channel.RightPushAsync(message);
+			if (length == 1)
 			{
-				await Redis.PublishAsync(UpdateChannel, ChannelId, CommandFlags.FireAndForget);
+				await _redis.PublishAsync(_updateChannel, channelId, CommandFlags.FireAndForget);
 			}
-			await Redis.KeyExpireAsync(Channel.Key, ExpireTime, CommandFlags.FireAndForget);
+			await _redis.KeyExpireAsync(channel.Key, ExpireTime, CommandFlags.FireAndForget);
 		}
 
-		static async Task<bool> ReadMessagesAsync(RedisList<T> List, List<T> Messages)
+		static async Task<bool> ReadMessagesAsync(RedisList<T> list, List<T> messages)
 		{
-			T? Message = await List.LeftPopAsync();
-			while(Message != null)
+			T? message = await list.LeftPopAsync();
+			while(message != null)
 			{
-				Messages.Add(Message);
-				Message = await List.LeftPopAsync();
+				messages.Add(message);
+				message = await list.LeftPopAsync();
 			}
-			return Messages.Count > 0;
+			return messages.Count > 0;
 		}
 
 		/// <inheritdoc/>
-		public async Task<List<T>> ReadMessagesAsync(string ChannelId)
+		public async Task<List<T>> ReadMessagesAsync(string channelId)
 		{
-			List<T> Messages = new List<T>();
-			await ReadMessagesAsync(GetChannel(ChannelId), Messages);
-			return Messages;
+			List<T> messages = new List<T>();
+			await ReadMessagesAsync(GetChannel(channelId), messages);
+			return messages;
 		}
 
 		/// <inheritdoc/>
-		public async Task<List<T>> WaitForMessagesAsync(string ChannelId, CancellationToken CancellationToken)
+		public async Task<List<T>> WaitForMessagesAsync(string channelId, CancellationToken cancellationToken)
 		{
-			List<T> Messages = new List<T>();
+			List<T> messages = new List<T>();
 
-			RedisList<T> Channel = GetChannel(ChannelId);
-			while (!CancellationToken.IsCancellationRequested && !await ReadMessagesAsync(Channel, Messages))
+			RedisList<T> channel = GetChannel(channelId);
+			while (!cancellationToken.IsCancellationRequested && !await ReadMessagesAsync(channel, messages))
 			{
 				// Register for notifications on this channel
-				TaskCompletionSource<bool>? CompletionSource;
-				lock (ChannelWakeEvents)
+				TaskCompletionSource<bool>? completionSource;
+				lock (_channelWakeEvents)
 				{
-					if (!ChannelWakeEvents.TryGetValue(ChannelId, out CompletionSource))
+					if (!_channelWakeEvents.TryGetValue(channelId, out completionSource))
 					{
-						CompletionSource = new TaskCompletionSource<bool>();
-						ChannelWakeEvents.Add(ChannelId, CompletionSource);
+						completionSource = new TaskCompletionSource<bool>();
+						_channelWakeEvents.Add(channelId, completionSource);
 					}
 				}
 
 				try
 				{
 					// Read the current queue state again, in case it was modified since we registered.
-					if(await ReadMessagesAsync(Channel, Messages))
+					if(await ReadMessagesAsync(channel, messages))
 					{
 						break;
 					}
 
 					// Wait for messages to be available
-					using (CancellationToken.Register(() => CompletionSource.TrySetResult(false)))
+					using (cancellationToken.Register(() => completionSource.TrySetResult(false)))
 					{
-						await CompletionSource.Task;
+						await completionSource.Task;
 					}
 				}
 				finally
 				{
-					lock (ChannelWakeEvents)
+					lock (_channelWakeEvents)
 					{
-						ChannelWakeEvents.Remove(ChannelId);
+						_channelWakeEvents.Remove(channelId);
 					}
 				}
 			}
 
-			return Messages;
+			return messages;
 		}
 	}
 }
diff --git a/Engine/Source/Programs/Horde/Horde.Build/Compute/Impl/RedisTaskScheduler.cs b/Engine/Source/Programs/Horde/Horde.Build/Compute/Impl/RedisTaskScheduler.cs
--- a/Engine/Source/Programs/Horde/Horde.Build/Compute/Impl/RedisTaskScheduler.cs
+++ b/Engine/Source/Programs/Horde/Horde.Build/Compute/Impl/RedisTaskScheduler.cs
@@ -1,24 +1,16 @@
 // Copyright Epic Games, Inc. All Rights Reserved.
 
-using EpicGames.Core;
-using EpicGames.Horde.Compute;
-using EpicGames.Redis;
-using EpicGames.Serialization;
-using Horde.Build.Models;
-using Horde.Build.Services;
-using Horde.Build.Storage;
-using Horde.Build.Utilities;
-using Microsoft.Extensions.Caching.Memory;
-using Microsoft.Extensions.Logging;
-using StackExchange.Redis;
 using System;
-using System.Collections.Concurrent;
 using System.Collections.Generic;
 using System.Diagnostics;
 using System.Linq;
 using System.Threading;
 using System.Threading.Channels;
 using System.Threading.Tasks;
+using EpicGames.Core;
+using EpicGames.Redis;
+using Microsoft.Extensions.Logging;
+using StackExchange.Redis;
 
 namespace Horde.Build.Compute.Impl
 {
@@ -36,25 +28,25 @@ namespace Horde.Build.Compute.Impl
 		/// <summary>
 		/// Adds a task to a queue
 		/// </summary>
-		/// <param name="QueueId">The queue identifier</param>
-		/// <param name="TaskId">The task to add</param>
-		/// <param name="AtFront">Whether to insert at the front of the queue</param>
-		Task EnqueueAsync(TQueueId QueueId, TTask TaskId, bool AtFront);
+		/// <param name="queueId">The queue identifier</param>
+		/// <param name="taskId">The task to add</param>
+		/// <param name="atFront">Whether to insert at the front of the queue</param>
+		Task EnqueueAsync(TQueueId queueId, TTask taskId, bool atFront);
 
 		/// <summary>
 		/// Dequeue any task from a particular queue
 		/// </summary>
-		/// <param name="QueueId">The queue to remove a task from</param>
+		/// <param name="queueId">The queue to remove a task from</param>
 		/// <returns>Information about the task to be executed</returns>
-		Task<TTask?> DequeueAsync(TQueueId QueueId);
+		Task<TTask?> DequeueAsync(TQueueId queueId);
 
 		/// <summary>
 		/// Dequeues a task that the given agent can execute
 		/// </summary>
-		/// <param name="Predicate">Predicate for determining which queues can be removed from</param>
-		/// <param name="Token">Cancellation token for the operation. Will return a null entry rather than throwing an exception.</param>
+		/// <param name="predicate">Predicate for determining which queues can be removed from</param>
+		/// <param name="token">Cancellation token for the operation. Will return a null entry rather than throwing an exception.</param>
 		/// <returns>Information about the task to be executed</returns>
-		Task<(TQueueId, TTask)?> DequeueAsync(Func<TQueueId, ValueTask<bool>> Predicate, CancellationToken Token = default);
+		Task<(TQueueId, TTask)?> DequeueAsync(Func<TQueueId, ValueTask<bool>> predicate, CancellationToken token = default);
 
 		/// <summary>
 		/// Gets hashes of all the inactive task queues
@@ -74,56 +66,55 @@ namespace Horde.Build.Compute.Impl
 	{
 		class Listener
 		{
-			public Func<TQueueId, ValueTask<bool>> Predicate;
+			public Func<TQueueId, ValueTask<bool>> _predicate;
 			public TaskCompletionSource<(TQueueId, TTask)?> CompletionSource { get; }
 
-			public Listener(Func<TQueueId, ValueTask<bool>> Predicate)
+			public Listener(Func<TQueueId, ValueTask<bool>> predicate)
 			{
-				this.Predicate = Predicate;
-				this.CompletionSource = new TaskCompletionSource<(TQueueId, TTask)?>();
+				_predicate = predicate;
+				CompletionSource = new TaskCompletionSource<(TQueueId, TTask)?>();
 			}
 		}
 
-		IDatabase Redis;
-		RedisKey BaseKey;
-		RedisSet<TQueueId> QueueIndex;
-		RedisHash<TQueueId, DateTime> ActiveQueues; // Queues which are actively being dequeued from
-		ReadOnlyHashSet<TQueueId> LocalActiveQueues = new HashSet<TQueueId>();
-		Stopwatch ResetActiveQueuesTimer = Stopwatch.StartNew();
-
-		List<Listener> Listeners = new List<Listener>();
-		RedisChannel<TQueueId> NewQueueChannel;
-		Task QueueUpdateTask;
-		CancellationTokenSource CancellationSource = new CancellationTokenSource();
-		ILogger Logger;
+		readonly IDatabase _redis;
+		readonly RedisKey _baseKey;
+		readonly RedisSet<TQueueId> _queueIndex;
+		readonly RedisHash<TQueueId, DateTime> _activeQueues; // Queues which are actively being dequeued from
+		ReadOnlyHashSet<TQueueId> _localActiveQueues = new HashSet<TQueueId>();
+		readonly Stopwatch _resetActiveQueuesTimer = Stopwatch.StartNew();
+		readonly List<Listener> _listeners = new List<Listener>();
+		readonly RedisChannel<TQueueId> _newQueueChannel;
+		readonly Task _queueUpdateTask;
+		readonly CancellationTokenSource _cancellationSource = new CancellationTokenSource();
+		readonly ILogger _logger;
 
 		/// <summary>
 		/// Constructor
 		/// </summary>
-		/// <param name="Redis">The Redis database instance</param>
-		/// <param name="BaseKey">Base key for all keys used by this scheduler</param>
-		/// <param name="Logger"></param>
-		public RedisTaskScheduler(IDatabase Redis, RedisKey BaseKey, ILogger Logger)
+		/// <param name="redis">The Redis database instance</param>
+		/// <param name="baseKey">Base key for all keys used by this scheduler</param>
+		/// <param name="logger"></param>
+		public RedisTaskScheduler(IDatabase redis, RedisKey baseKey, ILogger logger)
 		{
-			this.Redis = Redis;
-			this.BaseKey = BaseKey;
-			this.QueueIndex = new RedisSet<TQueueId>(Redis, BaseKey.Append("index"));
-			this.ActiveQueues = new RedisHash<TQueueId, DateTime>(Redis, BaseKey.Append("active"));
-			this.NewQueueChannel = new RedisChannel<TQueueId>(BaseKey.Append("new_queues").ToString());
-			this.Logger = Logger;
-
-			QueueUpdateTask = Task.Run(() => UpdateQueuesAsync(CancellationSource.Token));
+			_redis = redis;
+			_baseKey = baseKey;
+			_queueIndex = new RedisSet<TQueueId>(redis, baseKey.Append("index"));
+			_activeQueues = new RedisHash<TQueueId, DateTime>(redis, baseKey.Append("active"));
+			_newQueueChannel = new RedisChannel<TQueueId>(baseKey.Append("new_queues").ToString());
+			_logger = logger;
+
+			_queueUpdateTask = Task.Run(() => UpdateQueuesAsync(_cancellationSource.Token));
 		}
 
 		/// <inheritdoc/>
 		public async ValueTask DisposeAsync()
 		{
-			if (!QueueUpdateTask.IsCompleted)
+			if (!_queueUpdateTask.IsCompleted)
 			{
-				CancellationSource.Cancel();
-				await QueueUpdateTask;
+				_cancellationSource.Cancel();
+				await _queueUpdateTask;
 			}
-			CancellationSource.Dispose();
+			_cancellationSource.Dispose();
 		}
 
 		/// <inheritdoc/>
@@ -135,118 +126,118 @@ namespace Horde.Build.Compute.Impl
 		/// <summary>
 		/// Gets the key for a list of tasks for a particular queue
 		/// </summary>
-		/// <param name="QueueId">The queue identifier</param>
+		/// <param name="queueId">The queue identifier</param>
 		/// <returns></returns>
-		RedisKey GetQueueKey(TQueueId QueueId)
+		RedisKey GetQueueKey(TQueueId queueId)
 		{
-			return BaseKey.Append(RedisSerializer.Serialize(QueueId).AsKey());
+			return _baseKey.Append(RedisSerializer.Serialize(queueId).AsKey());
 		}
 
 		/// <summary>
 		/// Gets the key for a list of tasks for a particular queue
 		/// </summary>
-		/// <param name="QueueId">The queue identifier</param>
+		/// <param name="queueId">The queue identifier</param>
 		/// <returns></returns>
-		RedisList<TTask> GetQueue(TQueueId QueueId)
+		RedisList<TTask> GetQueue(TQueueId queueId)
 		{
-			return new RedisList<TTask>(Redis, GetQueueKey(QueueId));
+			return new RedisList<TTask>(_redis, GetQueueKey(queueId));
 		}
 
 		/// <summary>
 		/// Pushes a task onto either end of a queue
 		/// </summary>
-		static Task<long> PushTaskAsync(RedisList<TTask> List, TTask Task, When When, CommandFlags Flags, bool AtFront)
+		static Task<long> PushTaskAsync(RedisList<TTask> list, TTask task, When when, CommandFlags flags, bool atFront)
 		{
-			if (AtFront)
+			if (atFront)
 			{
-				return List.LeftPushAsync(Task, When, Flags);
+				return list.LeftPushAsync(task, when, flags);
 			}
 			else
 			{
-				return List.RightPushAsync(Task, When, Flags);
+				return list.RightPushAsync(task, when, flags);
 			}
 		}
 
 		/// <summary>
 		/// Adds a task to a particular queue, creating and adding that queue to the index if necessary
 		/// </summary>
-		/// <param name="QueueId">The queue to add the task to</param>
-		/// <param name="Task">Task to be scheduled</param>
-		/// <param name="AtFront">Whether to add to the front of the queue</param>
-		public async Task EnqueueAsync(TQueueId QueueId, TTask Task, bool AtFront)
+		/// <param name="queueId">The queue to add the task to</param>
+		/// <param name="task">Task to be scheduled</param>
+		/// <param name="atFront">Whether to add to the front of the queue</param>
+		public async Task EnqueueAsync(TQueueId queueId, TTask task, bool atFront)
 		{
-			RedisList<TTask> List = GetQueue(QueueId);
+			RedisList<TTask> list = GetQueue(queueId);
 			for (; ; )
 			{
-				long NewLength = await PushTaskAsync(List, Task, When.Exists, CommandFlags.None, AtFront);
-				if (NewLength > 0)
+				long newLength = await PushTaskAsync(list, task, When.Exists, CommandFlags.None, atFront);
+				if (newLength > 0)
 				{
-					Logger.LogInformation("Length of queue {QueueId} is {Length}", QueueId, NewLength);
+					_logger.LogInformation("Length of queue {QueueId} is {Length}", queueId, newLength);
 					break;
 				}
 
-				ITransaction Transaction = Redis.CreateTransaction();
-				_ = Transaction.With(QueueIndex).AddAsync(QueueId, CommandFlags.FireAndForget);
-				_ = PushTaskAsync(Transaction.With(List), Task, When.Always, CommandFlags.FireAndForget, AtFront);
+				ITransaction transaction = _redis.CreateTransaction();
+				_ = transaction.With(_queueIndex).AddAsync(queueId, CommandFlags.FireAndForget);
+				_ = PushTaskAsync(transaction.With(list), task, When.Always, CommandFlags.FireAndForget, atFront);
 
-				if (await Transaction.ExecuteAsync())
+				if (await transaction.ExecuteAsync())
 				{
-					Logger.LogInformation("Created queue {QueueId}", QueueId);
-					await Redis.PublishAsync(NewQueueChannel, QueueId);
+					_logger.LogInformation("Created queue {QueueId}", queueId);
+					await _redis.PublishAsync(_newQueueChannel, queueId);
 					break;
 				}
 
-				Logger.LogDebug("EnqueueAsync() retrying...");
+				_logger.LogDebug("EnqueueAsync() retrying...");
 			}
 		}
 
 		/// <summary>
 		/// Dequeues an item for execution by the given agent
 		/// </summary>
-		/// <param name="Predicate">Predicate for queues that tasks can be removed from</param>
-		/// <param name="Token">Cancellation token for waiting for an item</param>
+		/// <param name="predicate">Predicate for queues that tasks can be removed from</param>
+		/// <param name="token">Cancellation token for waiting for an item</param>
 		/// <returns>The dequeued item, or null if no item is available</returns>
-		public async Task<(TQueueId, TTask)?> DequeueAsync(Func<TQueueId, ValueTask<bool>> Predicate, CancellationToken Token = default)
+		public async Task<(TQueueId, TTask)?> DequeueAsync(Func<TQueueId, ValueTask<bool>> predicate, CancellationToken token = default)
 		{
 			// Compare against all the list of cached queues to see if we can dequeue something from any of them
-			Listener? Listener = null;
+			Listener? listener = null;
 			try
 			{
-				TQueueId[] Queues = await QueueIndex.MembersAsync();
-				while (!Token.IsCancellationRequested)
+				TQueueId[] queues = await _queueIndex.MembersAsync();
+				while (!token.IsCancellationRequested)
 				{
 					// Try to dequeue an item from the list
-					(TQueueId, TTask)? Entry = await TryAssignToLocalAgentAsync(Queues, Predicate);
-					if (Entry != null)
+					(TQueueId, TTask)? entry = await TryAssignToLocalAgentAsync(queues, predicate);
+					if (entry != null)
 					{
-						return Entry;
+						return entry;
 					}
 
 					// Create and register a listener for this waiter 
-					if (Listener == null)
+					if (listener == null)
 					{
-						Listener = new Listener(Predicate);
-						lock (Listeners)
+						listener = new Listener(predicate);
+						lock (_listeners)
 						{
-							Listeners.Add(Listener);
+							_listeners.Add(listener);
 						}
 					}
 					else
 					{
-						using (IDisposable Registration = Token.Register(() => Listener.CompletionSource.TrySetResult(null)))
+						using (IDisposable registration = token.Register(() => listener.CompletionSource.TrySetResult(null)))
 						{
-							return await Listener.CompletionSource.Task;
+							return await listener.CompletionSource.Task;
 						}
 					}
 				}
 			}
 			finally
 			{
-				if (Listener != null)
+				if (listener != null)
 				{
-					lock (Listeners)
+					lock (_listeners)
 					{
-						Listeners.Remove(Listener);
+						_listeners.Remove(listener);
 					}
 				}
 			}
@@ -256,19 +247,19 @@ namespace Horde.Build.Compute.Impl
 		/// <summary>
 		/// Attempts to dequeue a task from a set of queue
 		/// </summary>
-		/// <param name="QueueIds">The current array of queues</param>
-		/// <param name="Predicate">Predicate for queues that tasks can be removed from</param>
+		/// <param name="queueIds">The current array of queues</param>
+		/// <param name="predicate">Predicate for queues that tasks can be removed from</param>
 		/// <returns>The dequeued item, or null if no item is available</returns>
-		async Task<(TQueueId, TTask)?> TryAssignToLocalAgentAsync(TQueueId[] QueueIds, Func<TQueueId, ValueTask<bool>> Predicate)
+		async Task<(TQueueId, TTask)?> TryAssignToLocalAgentAsync(TQueueId[] queueIds, Func<TQueueId, ValueTask<bool>> predicate)
 		{
-			foreach (TQueueId QueueId in QueueIds)
+			foreach (TQueueId queueId in queueIds)
 			{
-				if (await Predicate(QueueId))
+				if (await predicate(queueId))
 				{
-					TTask? Task = await DequeueAsync(QueueId);
-					if (Task != null)
+					TTask? task = await DequeueAsync(queueId);
+					if (task != null)
 					{
-						return (QueueId, Task);
+						return (queueId, task);
 					}
 				}
 			}
@@ -278,45 +269,45 @@ namespace Horde.Build.Compute.Impl
 		/// <summary>
 		/// Dequeues the item at the front of a queue
 		/// </summary>
-		/// <param name="QueueId">The queue to dequeue from</param>
+		/// <param name="queueId">The queue to dequeue from</param>
 		/// <returns>The dequeued item, or null if the queue is empty</returns>
-		public async Task<TTask?> DequeueAsync(TQueueId QueueId)
+		public async Task<TTask?> DequeueAsync(TQueueId queueId)
 		{
-			await AddActiveQueue(QueueId);
+			await AddActiveQueue(queueId);
 
-			TTask? Item = await GetQueue(QueueId).LeftPopAsync();
-			if (Item == null)
+			TTask? item = await GetQueue(queueId).LeftPopAsync();
+			if (item == null)
 			{
-				ITransaction Transaction = Redis.CreateTransaction();
-				Transaction.AddCondition(Condition.KeyNotExists(GetQueueKey(QueueId)));
-				Task<bool> WasRemoved = Transaction.With(QueueIndex).RemoveAsync(QueueId);
+				ITransaction transaction = _redis.CreateTransaction();
+				transaction.AddCondition(Condition.KeyNotExists(GetQueueKey(queueId)));
+				Task<bool> wasRemoved = transaction.With(_queueIndex).RemoveAsync(queueId);
 
-				if (await Transaction.ExecuteAsync() && await WasRemoved)
+				if (await transaction.ExecuteAsync() && await wasRemoved)
 				{
-					Logger.LogInformation("Removed queue {QueueId} from index", QueueId);
+					_logger.LogInformation("Removed queue {QueueId} from index", queueId);
 				}
 			}
 
-			return Item;
+			return item;
 		}
 
 		/// <summary>
 		/// Marks a queue key as being actively monitored, preventing it being returned by <see cref="GetInactiveQueuesAsync"/>
 		/// </summary>
-		/// <param name="QueueId">The queue key</param>
+		/// <param name="queueId">The queue key</param>
 		/// <returns></returns>
-		async ValueTask AddActiveQueue(TQueueId QueueId)
+		async ValueTask AddActiveQueue(TQueueId queueId)
 		{
 			// Periodically clear out the set of active keys
-			TimeSpan ResetTime = TimeSpan.FromSeconds(10.0);
-			if (ResetActiveQueuesTimer.Elapsed > ResetTime)
+			TimeSpan resetTime = TimeSpan.FromSeconds(10.0);
+			if (_resetActiveQueuesTimer.Elapsed > resetTime)
 			{
-				lock (ResetActiveQueuesTimer)
+				lock (_resetActiveQueuesTimer)
 				{
-					if (ResetActiveQueuesTimer.Elapsed > ResetTime)
+					if (_resetActiveQueuesTimer.Elapsed > resetTime)
 					{
-						LocalActiveQueues = new HashSet<TQueueId>();
-						ResetActiveQueuesTimer.Restart();
+						_localActiveQueues = new HashSet<TQueueId>();
+						_resetActiveQueuesTimer.Restart();
 					}
 				}
 			}
@@ -326,19 +317,19 @@ namespace Horde.Build.Compute.Impl
 			// readers can thus access it without the need for any locking.
 			for(; ;)
 			{
-				ReadOnlyHashSet<TQueueId> LocalActiveQueuesCopy = LocalActiveQueues;
-				if (LocalActiveQueuesCopy.Contains(QueueId))
+				ReadOnlyHashSet<TQueueId> localActiveQueuesCopy = _localActiveQueues;
+				if (localActiveQueuesCopy.Contains(queueId))
 				{
 					break;
 				}
 
-				HashSet<TQueueId> NewLocalActiveQueues = new HashSet<TQueueId>(LocalActiveQueuesCopy);
-				NewLocalActiveQueues.Add(QueueId);
+				HashSet<TQueueId> newLocalActiveQueues = new HashSet<TQueueId>(localActiveQueuesCopy);
+				newLocalActiveQueues.Add(queueId);
 
-				if (Interlocked.CompareExchange(ref LocalActiveQueues, NewLocalActiveQueues, LocalActiveQueuesCopy) == LocalActiveQueuesCopy)
+				if (Interlocked.CompareExchange(ref _localActiveQueues, newLocalActiveQueues, localActiveQueuesCopy) == localActiveQueuesCopy)
 				{
-					Logger.LogInformation("Refreshing active queue {QueueId}", QueueId);
-					await ActiveQueues.SetAsync(QueueId, DateTime.UtcNow);
+					_logger.LogInformation("Refreshing active queue {QueueId}", queueId);
+					await _activeQueues.SetAsync(queueId, DateTime.UtcNow);
 					break;
 				}
 			}
@@ -350,114 +341,114 @@ namespace Horde.Build.Compute.Impl
 		/// <returns></returns>
 		public async Task<List<TQueueId>> GetInactiveQueuesAsync()
 		{
-			HashSet<TQueueId> Keys = new HashSet<TQueueId>(await QueueIndex.MembersAsync());
-			HashSet<TQueueId> InvalidKeys = new HashSet<TQueueId>();
+			HashSet<TQueueId> keys = new HashSet<TQueueId>(await _queueIndex.MembersAsync());
+			HashSet<TQueueId> invalidKeys = new HashSet<TQueueId>();
 
-			DateTime MinTime = DateTime.UtcNow - TimeSpan.FromMinutes(10.0);
+			DateTime minTime = DateTime.UtcNow - TimeSpan.FromMinutes(10.0);
 
-			HashEntry<TQueueId, DateTime>[] Entries = await ActiveQueues.GetAllAsync();
-			foreach (HashEntry<TQueueId, DateTime> Entry in Entries)
+			HashEntry<TQueueId, DateTime>[] entries = await _activeQueues.GetAllAsync();
+			foreach (HashEntry<TQueueId, DateTime> entry in entries)
 			{
-				if (Entry.Value < MinTime)
+				if (entry.Value < minTime)
 				{
-					InvalidKeys.Add(Entry.Name);
+					invalidKeys.Add(entry.Name);
 				}
 				else
 				{
-					Keys.Remove(Entry.Name);
+					keys.Remove(entry.Name);
 				}
 			}
 
-			if (InvalidKeys.Count > 0)
+			if (invalidKeys.Count > 0)
 			{
-				await ActiveQueues.DeleteAsync(InvalidKeys.ToArray());
+				await _activeQueues.DeleteAsync(invalidKeys.ToArray());
 			}
 
-			return Keys.ToList();
+			return keys.ToList();
 		}
 
-		async Task UpdateQueuesAsync(CancellationToken CancellationToken)
+		async Task UpdateQueuesAsync(CancellationToken cancellationToken)
 		{
-			Channel<TQueueId> NewQueues = Channel.CreateUnbounded<TQueueId>();
+			Channel<TQueueId> newQueues = Channel.CreateUnbounded<TQueueId>();
 
-			ISubscriber Subscriber = Redis.Multiplexer.GetSubscriber();
-			await using var _ = await Subscriber.SubscribeAsync(NewQueueChannel, (_, v) => NewQueues.Writer.TryWrite(v));
+			ISubscriber subscriber = _redis.Multiplexer.GetSubscriber();
+			await using RedisChannelSubscription<TQueueId>? _ = await subscriber.SubscribeAsync(_newQueueChannel, (_, v) => newQueues.Writer.TryWrite(v));
 
-			while (await NewQueues.Reader.WaitToReadAsync(CancellationToken))
+			while (await newQueues.Reader.WaitToReadAsync(cancellationToken))
 			{
-				HashSet<TQueueId> NewQueueIds = new HashSet<TQueueId>();
-				while (NewQueues.Reader.TryRead(out TQueueId? QueueId))
+				HashSet<TQueueId> newQueueIds = new HashSet<TQueueId>();
+				while (newQueues.Reader.TryRead(out TQueueId? queueId))
 				{
-					NewQueueIds.Add(QueueId);
+					newQueueIds.Add(queueId);
 				}
-				foreach (TQueueId NewQueueId in NewQueueIds)
+				foreach (TQueueId newQueueId in newQueueIds)
 				{
-					await TryDispatchToNewQueueAsync(NewQueueId);
+					await TryDispatchToNewQueueAsync(newQueueId);
 				}
 			}
 		}
 
-		async Task<bool> TryDispatchToNewQueueAsync(TQueueId QueueId)
+		async Task<bool> TryDispatchToNewQueueAsync(TQueueId queueId)
 		{
-			RedisList<TTask> Queue = GetQueue(QueueId);
+			RedisList<TTask> queue = GetQueue(queueId);
 
 			// Find a local listener that can execute the work
-			(TQueueId QueueId, TTask TaskId)? Entry = null;
+			(TQueueId QueueId, TTask TaskId)? entry = null;
 			try
 			{
 				for (; ; )
 				{
-					Listener? Listener = null;
+					Listener? listener = null;
 
 					// Look for a listener that can execute the task
-					HashSet<Listener> CheckedListeners = new HashSet<Listener>();
-					while(Listener == null)
+					HashSet<Listener> checkedListeners = new HashSet<Listener>();
+					while(listener == null)
 					{
 						// Find up to 10 listeners we haven't seen before
-						List<Listener> NewListeners = new List<Listener>();
-						lock (Listeners)
+						List<Listener> newListeners = new List<Listener>();
+						lock (_listeners)
 						{
-							NewListeners.AddRange(Listeners.Where(x => !x.CompletionSource.Task.IsCompleted && CheckedListeners.Add(x)).Take(10));
+							newListeners.AddRange(_listeners.Where(x => !x.CompletionSource.Task.IsCompleted && checkedListeners.Add(x)).Take(10));
 						}
-						if (NewListeners.Count == 0)
+						if (newListeners.Count == 0)
 						{
 							return false;
 						}
 
 						// Check predicates for each one against the new queue
-						foreach (Listener NewListener in NewListeners)
+						foreach (Listener newListener in newListeners)
 						{
-							if (await NewListener.Predicate(QueueId))
+							if (await newListener._predicate(queueId))
 							{
-								Listener = NewListener;
+								listener = newListener;
 								break;
 							}
 						}
 					}
 
 					// Pop an entry from the queue
-					if (Entry == null)
+					if (entry == null)
 					{
-						TTask? Task = await DequeueAsync(QueueId);
-						if (Task == null)
+						TTask? task = await DequeueAsync(queueId);
+						if (task == null)
 						{
 							return false;
 						}
-						Entry = (QueueId, Task);
+						entry = (queueId, task);
 					}
 
 					// Assign it to the listener
-					if (Listener.CompletionSource.TrySetResult(Entry))
+					if (listener.CompletionSource.TrySetResult(entry))
 					{
-						Entry = null;
+						entry = null;
 					}
 				}
 			}
 			finally
 			{
-				if (Entry != null)
+				if (entry != null)
 				{
-					await EnqueueAsync(Entry.Value.QueueId, Entry.Value.TaskId, true);
+					await EnqueueAsync(entry.Value.QueueId, entry.Value.TaskId, true);
 				}
 			}
 		}
diff --git a/Engine/Source/Programs/Horde/Horde.Build/Controllers/AccountController.cs b/Engine/Source/Programs/Horde/Horde.Build/Controllers/AccountController.cs
--- a/Engine/Source/Programs/Horde/Horde.Build/Controllers/AccountController.cs
+++ b/Engine/Source/Programs/Horde/Horde.Build/Controllers/AccountController.cs
@@ -1,21 +1,16 @@
 ﻿// Copyright Epic Games, Inc. All Rights Reserved.
 
+using System;
+using System.Globalization;
+using System.Net;
+using System.Text;
+using System.Threading.Tasks;
 using Horde.Build.Acls;
 using Horde.Build.Authentication;
-using Horde.Build.Models;
-using Horde.Build.Services;
 using Microsoft.AspNetCore.Authentication;
 using Microsoft.AspNetCore.Authentication.Cookies;
 using Microsoft.AspNetCore.Authentication.OpenIdConnect;
-using Microsoft.AspNetCore.Authorization;
 using Microsoft.AspNetCore.Mvc;
-using System;
-using System.Collections.Generic;
-using System.Globalization;
-using System.Linq;
-using System.Net;
-using System.Text;
-using System.Threading.Tasks;
 using Microsoft.Extensions.Options;
 
 namespace Horde.Build.Controllers
@@ -40,37 +35,37 @@ namespace Horde.Build.Controllers
 		/// <summary>
 		/// The ACL service singleton
 		/// </summary>
-		AclService AclService;
-		
+		readonly AclService _aclService;
+
 		/// <summary>
 		/// Authentication scheme in use
 		/// </summary>
-		string AuthenticationScheme;
+		readonly string _authenticationScheme;
 
 		/// <summary>
 		/// Constructor
 		/// </summary>
-		/// <param name="AclService">ACL service instance</param>
-		/// <param name="ServerSettings">Server settings</param>
-		public AccountController(AclService AclService, IOptionsMonitor<ServerSettings> ServerSettings)
+		/// <param name="aclService">ACL service instance</param>
+		/// <param name="serverSettings">Server settings</param>
+		public AccountController(AclService aclService, IOptionsMonitor<ServerSettings> serverSettings)
 		{
-			this.AclService = AclService;
-			this.AuthenticationScheme = GetAuthScheme(ServerSettings.CurrentValue.AuthMethod);
+			_aclService = aclService;
+			_authenticationScheme = GetAuthScheme(serverSettings.CurrentValue.AuthMethod);
 		}
 
 		/// <summary>
 		/// Get auth scheme name for a given auth method
 		/// </summary>
-		/// <param name="Method">Authentication method</param>
+		/// <param name="method">Authentication method</param>
 		/// <returns>Name of authentication scheme</returns>
-		public static string GetAuthScheme(AuthMethod Method)
+		public static string GetAuthScheme(AuthMethod method)
 		{
-			return Method switch
+			return method switch
 			{
 				AuthMethod.Anonymous => AnonymousAuthenticationHandler.AuthenticationScheme,
 				AuthMethod.Okta => OktaDefaults.AuthenticationScheme,
 				AuthMethod.OpenIdConnect => OpenIdConnectDefaults.AuthenticationScheme,
-				_ => throw new ArgumentOutOfRangeException(nameof(Method), Method, null)
+				_ => throw new ArgumentOutOfRangeException(nameof(method), method, null)
 			};
 		}
 
@@ -82,39 +77,39 @@ namespace Horde.Build.Controllers
 		[Route("/account")]
 		public async Task<ActionResult> State()
 		{
-			StringBuilder Content = new StringBuilder();
-			Content.Append($"<html><style>{StyleSheet}</style><h1>Horde Server</h1>");
+			StringBuilder content = new StringBuilder();
+			content.Append($"<html><style>{StyleSheet}</style><h1>Horde Server</h1>");
 			if (User.Identity?.IsAuthenticated ?? false)
 			{
-				Content.Append(CultureInfo.InvariantCulture, $"<p>User <b>{User.Identity?.Name}</b> is logged in. <a href=\"/account/logout\">Log out</a></p>");
-				if (await AclService.AuthorizeAsync(AclAction.AdminWrite, User))
+				content.Append(CultureInfo.InvariantCulture, $"<p>User <b>{User.Identity?.Name}</b> is logged in. <a href=\"/account/logout\">Log out</a></p>");
+				if (await _aclService.AuthorizeAsync(AclAction.AdminWrite, User))
 				{
-					Content.Append("<p>");
-					Content.Append("<a href=\"/api/v1/admin/token\">Get bearer token</a><br/>");
-					Content.Append("<a href=\"/api/v1/admin/registrationtoken\">Get agent registration token</a><br/>");
-					Content.Append("<a href=\"/api/v1/admin/softwaretoken\">Get agent software upload token</a><br/>");
-					Content.Append("<a href=\"/api/v1/admin/softwaredownloadtoken\">Get agent software download token</a><br/>");
-					Content.Append("<a href=\"/api/v1/admin/configtoken\">Get configuration token</a><br/>");
-					Content.Append("<a href=\"/api/v1/admin/chainedjobtoken\">Get chained job token</a><br/>");
-					Content.Append("</p>");
+					content.Append("<p>");
+					content.Append("<a href=\"/api/v1/admin/token\">Get bearer token</a><br/>");
+					content.Append("<a href=\"/api/v1/admin/registrationtoken\">Get agent registration token</a><br/>");
+					content.Append("<a href=\"/api/v1/admin/softwaretoken\">Get agent software upload token</a><br/>");
+					content.Append("<a href=\"/api/v1/admin/softwaredownloadtoken\">Get agent software download token</a><br/>");
+					content.Append("<a href=\"/api/v1/admin/configtoken\">Get configuration token</a><br/>");
+					content.Append("<a href=\"/api/v1/admin/chainedjobtoken\">Get chained job token</a><br/>");
+					content.Append("</p>");
 				}
-				Content.Append(CultureInfo.InvariantCulture, $"<p>Claims for {User.Identity?.Name}:");
-				Content.Append("<table>");
-				foreach (System.Security.Claims.Claim Claim in User.Claims)
+				content.Append(CultureInfo.InvariantCulture, $"<p>Claims for {User.Identity?.Name}:");
+				content.Append("<table>");
+				foreach (System.Security.Claims.Claim claim in User.Claims)
 				{
-					Content.Append(CultureInfo.InvariantCulture, $"<tr><td>{Claim.Type}</td><td>{Claim.Value}</td></tr>");
+					content.Append(CultureInfo.InvariantCulture, $"<tr><td>{claim.Type}</td><td>{claim.Value}</td></tr>");
 				}
-				Content.Append("</table>");
-				Content.Append("</p>");
+				content.Append("</table>");
+				content.Append("</p>");
 
-				Content.Append(CultureInfo.InvariantCulture, $"<p>Built from Perforce</p>");
+				content.Append(CultureInfo.InvariantCulture, $"<p>Built from Perforce</p>");
 			}
 			else
 			{
-				Content.Append("<p><a href=\"/account/login\"><b>Login with OAuth2</b></a></p>");
+				content.Append("<p><a href=\"/account/login\"><b>Login with OAuth2</b></a></p>");
 			}
-			Content.Append("</html>");
-			return new ContentResult { ContentType = "text/html", StatusCode = (int)HttpStatusCode.OK, Content = Content.ToString() };
+			content.Append("</html>");
+			return new ContentResult { ContentType = "text/html", StatusCode = (int)HttpStatusCode.OK, Content = content.ToString() };
 		}
 
 		/// <summary>
@@ -125,7 +120,7 @@ namespace Horde.Build.Controllers
 		[Route("/account/login")]
 		public IActionResult Login()
 		{
-			return new ChallengeResult(AuthenticationScheme, new AuthenticationProperties { RedirectUri = "/account" });
+			return new ChallengeResult(_authenticationScheme, new AuthenticationProperties { RedirectUri = "/account" });
 		}
 
 		/// <summary>
@@ -139,14 +134,14 @@ namespace Horde.Build.Controllers
 			await HttpContext.SignOutAsync(CookieAuthenticationDefaults.AuthenticationScheme);
 			try
 			{
-				await HttpContext.SignOutAsync(AuthenticationScheme);
+				await HttpContext.SignOutAsync(_authenticationScheme);
 			}
 			catch
 			{
 			}
 
-			string Content = $"<html><style>{StyleSheet}</style><body onload=\"setTimeout(function(){{ window.location = '/account'; }}, 2000)\"><p>User has been logged out. Returning to login page.</p></body></html>";
-			return new ContentResult { ContentType = "text/html", StatusCode = (int)HttpStatusCode.OK, Content = Content };
+			string content = $"<html><style>{StyleSheet}</style><body onload=\"setTimeout(function(){{ window.location = '/account'; }}, 2000)\"><p>User has been logged out. Returning to login page.</p></body></html>";
+			return new ContentResult { ContentType = "text/html", StatusCode = (int)HttpStatusCode.OK, Content = content };
 		}
 	}
 }
diff --git a/Engine/Source/Programs/Horde/Horde.Build/Controllers/AdminController.cs b/Engine/Source/Programs/Horde/Horde.Build/Controllers/AdminController.cs
--- a/Engine/Source/Programs/Horde/Horde.Build/Controllers/AdminController.cs
+++ b/Engine/Source/Programs/Horde/Horde.Build/Controllers/AdminController.cs
@@ -1,26 +1,16 @@
-﻿// Copyright Epic Games, Inc. All Rights Reserved.
+// Copyright Epic Games, Inc. All Rights Reserved.
 
-using EpicGames.Core;
+using System;
+using System.Collections.Generic;
+using System.Linq;
+using System.Security.Claims;
+using System.Threading.Tasks;
 using Horde.Build.Acls;
-using Horde.Build.Api;
-using Horde.Build.Collections;
-using HordeCommon;
-using Horde.Build.Models;
 using Horde.Build.Services;
 using Horde.Build.Utilities;
 using Microsoft.AspNetCore.Authorization;
-using Microsoft.AspNetCore.Authorization.Infrastructure;
 using Microsoft.AspNetCore.Mvc;
 using Microsoft.Extensions.Options;
-using MongoDB.Bson;
-using System;
-using System.Collections.Generic;
-using System.IdentityModel.Tokens.Jwt;
-using System.IO;
-using System.Linq;
-using System.Runtime.InteropServices;
-using System.Security.Claims;
-using System.Threading.Tasks;
 
 namespace Horde.Build.Controllers
 {
@@ -64,55 +54,37 @@ namespace Horde.Build.Controllers
 	[Route("[controller]")]
 	public class AdminController : HordeControllerBase
 	{
-		/// <summary>
-		/// The database service singleton
-		/// </summary>
-		DatabaseService DatabaseService;
-
-		/// <summary>
-		/// The acl service singleton
-		/// </summary>
-		AclService AclService;
-
-		/// <summary>
-		/// The upgrade service singleton
-		/// </summary>
-		UpgradeService UpgradeService;
-
-		/// <summary>
-		/// Settings for the server
-		/// </summary>
-		IOptionsMonitor<ServerSettings> Settings;
+		readonly AclService _aclService;
+		readonly UpgradeService _upgradeService;
+		readonly IOptionsMonitor<ServerSettings> _settings;
 
 		/// <summary>
 		/// Constructor
 		/// </summary>
-		/// <param name="DatabaseService">The database service singleton</param>
-		/// <param name="AclService">The ACL service singleton</param>
-		/// <param name="UpgradeService">The upgrade service singelton</param>
-		/// <param name="Settings">Server settings</param>
-		public AdminController(DatabaseService DatabaseService, AclService AclService, UpgradeService UpgradeService, IOptionsMonitor<ServerSettings> Settings)
+		/// <param name="aclService">The ACL service singleton</param>
+		/// <param name="upgradeService">The upgrade service singelton</param>
+		/// <param name="settings">Server settings</param>
+		public AdminController(AclService aclService, UpgradeService upgradeService, IOptionsMonitor<ServerSettings> settings)
 		{
-			this.DatabaseService = DatabaseService;
-			this.AclService = AclService;
-			this.UpgradeService = UpgradeService;
-			this.Settings = Settings;
+			_aclService = aclService;
+			_upgradeService = upgradeService;
+			_settings = settings;
 		}
 
 		/// <summary>
 		/// Upgrade the database to the latest schema
 		/// </summary>
-		/// <param name="FromVersion">The schema version to upgrade from.</param>
+		/// <param name="fromVersion">The schema version to upgrade from.</param>
 		[HttpPost]
 		[Route("/api/v1/admin/upgradeschema")]
-		public async Task<ActionResult> UpgradeSchemaAsync([FromQuery] int? FromVersion = null)
+		public async Task<ActionResult> UpgradeSchemaAsync([FromQuery] int? fromVersion = null)
 		{
-			if (!await AclService.AuthorizeAsync(AclAction.AdminWrite, User))
+			if (!await _aclService.AuthorizeAsync(AclAction.AdminWrite, User))
 			{
 				return Forbid(AclAction.AdminWrite);
 			}
 
-			await UpgradeService.UpgradeSchemaAsync(FromVersion);
+			await _upgradeService.UpgradeSchemaAsync(fromVersion);
 			return Ok();
 		}
 
@@ -124,32 +96,32 @@ namespace Horde.Build.Controllers
 		[Route("/api/v1/admin/token")]
 		public async Task<ActionResult<string>> GetTokenAsync()
 		{
-			if (!await AclService.AuthorizeAsync(AclAction.IssueBearerToken, User))
+			if (!await _aclService.AuthorizeAsync(AclAction.IssueBearerToken, User))
 			{
 				return Forbid(AclAction.IssueBearerToken);
 			}
 
-			return AclService.IssueBearerToken(User.Claims, GetDefaultExpiryTime());
+			return _aclService.IssueBearerToken(User.Claims, GetDefaultExpiryTime());
 		}
 
 		/// <summary>
 		/// Issues a token for the given roles. Issues a token for the current user if not specified.
 		/// </summary>
-		/// <param name="Roles">Roles for the new token</param>
+		/// <param name="roles">Roles for the new token</param>
 		/// <returns>Administrative settings for the server</returns>
 		[HttpGet]
 		[Route("/api/v1/admin/roletoken")]
-		public async Task<ActionResult<string>> GetRoleTokenAsync([FromQuery] string Roles)
+		public async Task<ActionResult<string>> GetRoleTokenAsync([FromQuery] string roles)
 		{
-			if (!await AclService.AuthorizeAsync(AclAction.AdminWrite, User))
+			if (!await _aclService.AuthorizeAsync(AclAction.AdminWrite, User))
 			{
 				return Forbid(AclAction.AdminWrite);
 			}
 
-			List<Claim> Claims = new List<Claim>();
-			Claims.AddRange(Roles.Split('+').Select(x => new Claim(ClaimTypes.Role, x)));
+			List<Claim> claims = new List<Claim>();
+			claims.AddRange(roles.Split('+').Select(x => new Claim(ClaimTypes.Role, x)));
 
-			return AclService.IssueBearerToken(Claims, GetDefaultExpiryTime());
+			return _aclService.IssueBearerToken(claims, GetDefaultExpiryTime());
 		}
 
 		/// <summary>
@@ -160,16 +132,16 @@ namespace Horde.Build.Controllers
 		[Route("/api/v1/admin/registrationtoken")]
 		public async Task<ActionResult<string>> GetRegistrationTokenAsync()
 		{
-			if (!await AclService.AuthorizeAsync(AclAction.AdminWrite, User))
+			if (!await _aclService.AuthorizeAsync(AclAction.AdminWrite, User))
 			{
 				return Forbid(AclAction.AdminWrite);
 			}
 
-			List<AclClaim> Claims = new List<AclClaim>();
-			Claims.Add(new AclClaim(ClaimTypes.Name, User.Identity?.Name ?? "Unknown"));
-			Claims.Add(AclService.AgentRegistrationClaim);
+			List<AclClaim> claims = new List<AclClaim>();
+			claims.Add(new AclClaim(ClaimTypes.Name, User.Identity?.Name ?? "Unknown"));
+			claims.Add(AclService.AgentRegistrationClaim);
 
-			return AclService.IssueBearerToken(Claims, null);
+			return _aclService.IssueBearerToken(claims, null);
 		}
 
 		/// <summary>
@@ -180,16 +152,16 @@ namespace Horde.Build.Controllers
 		[Route("/api/v1/admin/softwaretoken")]
 		public async Task<ActionResult<string>> GetSoftwareTokenAsync()
 		{
-			if (!await AclService.AuthorizeAsync(AclAction.AdminWrite, User))
+			if (!await _aclService.AuthorizeAsync(AclAction.AdminWrite, User))
 			{
 				return Forbid(AclAction.AdminWrite);
 			}
 
-			List<AclClaim> Claims = new List<AclClaim>();
-			Claims.Add(new AclClaim(ClaimTypes.Name, User.Identity?.Name ?? "Unknown"));
-			Claims.Add(AclService.UploadSoftwareClaim);
+			List<AclClaim> claims = new List<AclClaim>();
+			claims.Add(new AclClaim(ClaimTypes.Name, User.Identity?.Name ?? "Unknown"));
+			claims.Add(AclService.UploadSoftwareClaim);
 
-			return AclService.IssueBearerToken(Claims, null);
+			return _aclService.IssueBearerToken(claims, null);
 		}
 
 		/// <summary>
@@ -200,16 +172,16 @@ namespace Horde.Build.Controllers
 		[Route("/api/v1/admin/softwaredownloadtoken")]
 		public async Task<ActionResult<string>> GetSoftwareDownloadTokenAsync()
 		{
-			if (!await AclService.AuthorizeAsync(AclAction.AdminRead, User))
+			if (!await _aclService.AuthorizeAsync(AclAction.AdminRead, User))
 			{
 				return Forbid(AclAction.AdminRead);
 			}
 
-			List<AclClaim> Claims = new List<AclClaim>();
-			Claims.Add(new AclClaim(ClaimTypes.Name, User.Identity?.Name ?? "Unknown"));
-			Claims.Add(AclService.DownloadSoftwareClaim);
+			List<AclClaim> claims = new List<AclClaim>();
+			claims.Add(new AclClaim(ClaimTypes.Name, User.Identity?.Name ?? "Unknown"));
+			claims.Add(AclService.DownloadSoftwareClaim);
 
-			return AclService.IssueBearerToken(Claims, null);
+			return _aclService.IssueBearerToken(claims, null);
 		}
 
 		/// <summary>
@@ -220,16 +192,16 @@ namespace Horde.Build.Controllers
 		[Route("/api/v1/admin/configtoken")]
 		public async Task<ActionResult<string>> GetConfigToken()
 		{
-			if (!await AclService.AuthorizeAsync(AclAction.AdminRead, User))
+			if (!await _aclService.AuthorizeAsync(AclAction.AdminRead, User))
 			{
 				return Forbid(AclAction.AdminRead);
 			}
 
-			List<AclClaim> Claims = new List<AclClaim>();
-			Claims.Add(new AclClaim(ClaimTypes.Name, User.Identity?.Name ?? "Unknown"));
-			Claims.Add(AclService.ConfigureProjectsClaim);
+			List<AclClaim> claims = new List<AclClaim>();
+			claims.Add(new AclClaim(ClaimTypes.Name, User.Identity?.Name ?? "Unknown"));
+			claims.Add(AclService.ConfigureProjectsClaim);
 
-			return AclService.IssueBearerToken(Claims, null);
+			return _aclService.IssueBearerToken(claims, null);
 		}
 
 		/// <summary>
@@ -240,16 +212,16 @@ namespace Horde.Build.Controllers
 		[Route("/api/v1/admin/chainedjobtoken")]
 		public async Task<ActionResult<string>> GetChainedJobToken()
 		{
-			if (!await AclService.AuthorizeAsync(AclAction.AdminRead, User))
+			if (!await _aclService.AuthorizeAsync(AclAction.AdminRead, User))
 			{
 				return Forbid(AclAction.AdminRead);
 			}
 
-			List<AclClaim> Claims = new List<AclClaim>();
+			List<AclClaim> claims = new List<AclClaim>();
 			//Claims.Add(new AclClaim(ClaimTypes.Name, User.Identity.Name ?? "Unknown"));
-			Claims.Add(AclService.StartChainedJobClaim);
+			claims.Add(AclService.StartChainedJobClaim);
 
-			return AclService.IssueBearerToken(Claims, null);
+			return _aclService.IssueBearerToken(claims, null);
 		}
 
 		/// <summary>
@@ -258,12 +230,12 @@ namespace Horde.Build.Controllers
 		/// <returns></returns>
 		private TimeSpan? GetDefaultExpiryTime()
 		{
-			TimeSpan? ExpiryTime = null;
-			if (Settings.CurrentValue.JwtExpiryTimeHours != -1)
+			TimeSpan? expiryTime = null;
+			if (_settings.CurrentValue.JwtExpiryTimeHours != -1)
 			{
-				ExpiryTime = TimeSpan.FromHours(Settings.CurrentValue.JwtExpiryTimeHours);
+				expiryTime = TimeSpan.FromHours(_settings.CurrentValue.JwtExpiryTimeHours);
 			}
-			return ExpiryTime;
+			return expiryTime;
 		}
 	}
 }
diff --git a/Engine/Source/Programs/Horde/Horde.Build/Controllers/AgentSoftwareController.cs b/Engine/Source/Programs/Horde/Horde.Build/Controllers/AgentSoftwareController.cs
--- a/Engine/Source/Programs/Horde/Horde.Build/Controllers/AgentSoftwareController.cs
+++ b/Engine/Source/Programs/Horde/Horde.Build/Controllers/AgentSoftwareController.cs
@@ -1,20 +1,17 @@
 ﻿// Copyright Epic Games, Inc. All Rights Reserved.
 
+using System.Collections.Generic;
+using System.IO;
+using System.Threading.Tasks;
 using Horde.Build.Acls;
 using Horde.Build.Api;
+using Horde.Build.Models;
 using Horde.Build.Services;
+using Horde.Build.Utilities;
 using Microsoft.AspNetCore.Authorization;
 using Microsoft.AspNetCore.Http;
 using Microsoft.AspNetCore.Mvc;
 using Microsoft.Net.Http.Headers;
-using System;
-using System.IO;
-using System.Collections.Generic;
-using System.Linq;
-using System.Threading.Tasks;
-using Horde.Build.Models;
-using MongoDB.Bson;
-using Horde.Build.Utilities;
 
 namespace Horde.Build.Controllers
 {
@@ -31,124 +28,124 @@ namespace Horde.Build.Controllers
 		/// <summary>
 		/// Singleton instance of the ACL service
 		/// </summary>
-		private readonly AclService AclService;
+		private readonly AclService _aclService;
 
 		/// <summary>
 		/// Singleton instance of the client service
 		/// </summary>
-		private readonly AgentSoftwareService AgentSoftwareService;
+		private readonly AgentSoftwareService _agentSoftwareService;
 
 		/// <summary>
 		/// Constructor
 		/// </summary>
-		/// <param name="AclService">The ACL service</param>
-		/// <param name="AgentSoftwareService">The client service</param>
-		public AgentSoftwareController(AclService AclService, AgentSoftwareService AgentSoftwareService)
+		/// <param name="aclService">The ACL service</param>
+		/// <param name="agentSoftwareService">The client service</param>
+		public AgentSoftwareController(AclService aclService, AgentSoftwareService agentSoftwareService)
 		{
-			this.AclService = AclService;
-			this.AgentSoftwareService = AgentSoftwareService;
+			_aclService = aclService;
+			_agentSoftwareService = agentSoftwareService;
 		}
 
 		/// <summary>
 		/// Finds all uploaded software matching the given criteria
 		/// </summary>
-		/// <param name="Filter">Filter for the properties to return</param>
+		/// <param name="filter">Filter for the properties to return</param>
 		/// <returns>Http response</returns>
 		[HttpGet]
 		[Route("/api/v1/agentsoftware")]
 		[ProducesResponseType(typeof(List<GetAgentSoftwareChannelResponse>), 200)]
-		public async Task<ActionResult<List<object>>> FindSoftwareAsync([FromQuery] PropertyFilter? Filter = null)
+		public async Task<ActionResult<List<object>>> FindSoftwareAsync([FromQuery] PropertyFilter? filter = null)
 		{
-			if (!await AclService.AuthorizeAsync(AclAction.DownloadSoftware, User))
+			if (!await _aclService.AuthorizeAsync(AclAction.DownloadSoftware, User))
 			{
 				return Forbid();
 			}
 
-			List<IAgentSoftwareChannel> Results = await AgentSoftwareService.FindChannelsAsync();
+			List<IAgentSoftwareChannel> results = await _agentSoftwareService.FindChannelsAsync();
 
-			List<object> Responses = new List<object>();
-			foreach (IAgentSoftwareChannel Result in Results)
+			List<object> responses = new List<object>();
+			foreach (IAgentSoftwareChannel result in results)
 			{
-				Responses.Add(new GetAgentSoftwareChannelResponse(Result).ApplyFilter(Filter));
+				responses.Add(new GetAgentSoftwareChannelResponse(result).ApplyFilter(filter));
 			}
-			return Responses;
+			return responses;
 		}
 
 		/// <summary>
 		/// Finds all uploaded software matching the given criteria
 		/// </summary>
-		/// <param name="Name">Name of the channel to get</param>
-		/// <param name="Filter">Filter for the properties to return</param>
+		/// <param name="name">Name of the channel to get</param>
+		/// <param name="filter">Filter for the properties to return</param>
 		/// <returns>Http response</returns>
 		[HttpGet]
-		[Route("/api/v1/agentsoftware/{Name}")]
+		[Route("/api/v1/agentsoftware/{name}")]
 		[ProducesResponseType(typeof(GetAgentSoftwareChannelResponse), 200)]
-		public async Task<ActionResult<object>> FindSoftwareAsync(string Name, [FromQuery] PropertyFilter? Filter = null)
+		public async Task<ActionResult<object>> FindSoftwareAsync(string name, [FromQuery] PropertyFilter? filter = null)
 		{
-			if (!await AclService.AuthorizeAsync(AclAction.DownloadSoftware, User))
+			if (!await _aclService.AuthorizeAsync(AclAction.DownloadSoftware, User))
 			{
 				return Forbid();
 			}
 
-			IAgentSoftwareChannel? Channel = await AgentSoftwareService.GetChannelAsync(new AgentSoftwareChannelName(Name));
-			if(Channel == null)
+			IAgentSoftwareChannel? channel = await _agentSoftwareService.GetChannelAsync(new AgentSoftwareChannelName(name));
+			if(channel == null)
 			{
 				return NotFound();
 			}
 
-			return new GetAgentSoftwareChannelResponse(Channel).ApplyFilter(Filter);
+			return new GetAgentSoftwareChannelResponse(channel).ApplyFilter(filter);
 		}
 
 		/// <summary>
 		/// Uploads a new agent zip file
 		/// </summary>
-		/// <param name="Name">Name of the channel to post to</param>
-		/// <param name="File">Zip archive containing the new client software</param>
+		/// <param name="name">Name of the channel to post to</param>
+		/// <param name="file">Zip archive containing the new client software</param>
 		/// <returns>Http result code</returns>
 		[HttpPost]
-		[Route("/api/v1/agentsoftware/{Name}/zip")]
-		public async Task<ActionResult> SetArchiveAsync(string Name, [FromForm] IFormFile File)
+		[Route("/api/v1/agentsoftware/{name}/zip")]
+		public async Task<ActionResult> SetArchiveAsync(string name, [FromForm] IFormFile file)
 		{
-			if (!await AclService.AuthorizeAsync(AclAction.UploadSoftware, User))
+			if (!await _aclService.AuthorizeAsync(AclAction.UploadSoftware, User))
 			{
 				return Forbid();
 			}
 
-			byte[] Data;
-			using (MemoryStream MemoryStream = new MemoryStream())
+			byte[] data;
+			using (MemoryStream memoryStream = new MemoryStream())
 			{
-				using (System.IO.Stream Stream = File.OpenReadStream())
+				using (System.IO.Stream stream = file.OpenReadStream())
 				{
-					await Stream.CopyToAsync(MemoryStream);
+					await stream.CopyToAsync(memoryStream);
 				}
-				Data = MemoryStream.ToArray();
+				data = memoryStream.ToArray();
 			}
 
-			await AgentSoftwareService.SetArchiveAsync(new AgentSoftwareChannelName(Name), User.Identity?.Name, Data);
+			await _agentSoftwareService.SetArchiveAsync(new AgentSoftwareChannelName(name), User.Identity?.Name, data);
 			return Ok();
 		}
 
 		/// <summary>
 		/// Gets the zip file for a specific channel
 		/// </summary>
-		/// <param name="Name">Name of the channel</param>
+		/// <param name="name">Name of the channel</param>
 		/// <returns>Http response</returns>
 		[HttpGet]
-		[Route("/api/v1/agentsoftware/{Name}/zip")]
-		public async Task<ActionResult> GetArchiveAsync(string Name)
+		[Route("/api/v1/agentsoftware/{name}/zip")]
+		public async Task<ActionResult> GetArchiveAsync(string name)
 		{
-			if (!await AclService.AuthorizeAsync(AclAction.DownloadSoftware, User))
+			if (!await _aclService.AuthorizeAsync(AclAction.DownloadSoftware, User))
 			{
 				return Forbid();
 			}
 
-			byte[]? Data = await AgentSoftwareService.GetArchiveAsync(new AgentSoftwareChannelName(Name));
-			if (Data == null)
+			byte[]? data = await _agentSoftwareService.GetArchiveAsync(new AgentSoftwareChannelName(name));
+			if (data == null)
 			{
 				return NotFound();
 			}
 
-			return new FileStreamResult(new MemoryStream(Data), new MediaTypeHeaderValue("application/octet-stream")) { FileDownloadName = $"HordeAgent.zip" };
+			return new FileStreamResult(new MemoryStream(data), new MediaTypeHeaderValue("application/octet-stream")) { FileDownloadName = $"HordeAgent.zip" };
 			
 		}
 	}
