Horde: Fix coding conventions within Horde.Build.

diff --git a/Engine/Source/Programs/Horde/Horde.Build/Controllers/AgentsController.cs b/Engine/Source/Programs/Horde/Horde.Build/Controllers/AgentsController.cs
--- a/Engine/Source/Programs/Horde/Horde.Build/Controllers/AgentsController.cs
+++ b/Engine/Source/Programs/Horde/Horde.Build/Controllers/AgentsController.cs
@@ -1,5 +1,9 @@
 // Copyright Epic Games, Inc. All Rights Reserved.
 
+using System;
+using System.Collections.Generic;
+using System.Linq;
+using System.Threading.Tasks;
 using Horde.Build.Acls;
 using Horde.Build.Api;
 using Horde.Build.Collections;
@@ -9,12 +13,6 @@ using Horde.Build.Utilities;
 using Microsoft.AspNetCore.Authorization;
 using Microsoft.AspNetCore.Mvc;
 using Microsoft.Extensions.Logging;
-using MongoDB.Bson;
-using System;
-using System.Collections.Generic;
-using System.Linq;
-using System.Text;
-using System.Threading.Tasks;
 
 namespace Horde.Build.Controllers
 {
@@ -31,192 +29,178 @@ namespace Horde.Build.Controllers
 	[Route("[controller]")]
 	public class AgentsController : HordeControllerBase
 	{
-		/// <summary>
-		/// Singleton instance of the ACL service
-		/// </summary>
-		AclService AclService;
-
-		/// <summary>
-		/// Singleton instance of the agent service
-		/// </summary>
-		AgentService AgentService;
-
-		/// <summary>
-		/// The pool service
-		/// </summary>
-		PoolService PoolService;
+		readonly AclService _aclService;
+		readonly AgentService _agentService;
 
 		/// <summary>
 		/// Constructor
 		/// </summary>
-		/// <param name="AclService">The ACL service singleton</param>
-		/// <param name="AgentService">The agent service</param>
-		/// <param name="PoolService">The pool service</param>
-		public AgentsController(AclService AclService, AgentService AgentService, PoolService PoolService)
+		/// <param name="aclService">The ACL service singleton</param>
+		/// <param name="agentService">The agent service</param>
+		public AgentsController(AclService aclService, AgentService agentService)
 		{
-			this.AclService = AclService;
-			this.AgentService = AgentService;
-			this.PoolService = PoolService;
+			_aclService = aclService;
+			_agentService = agentService;
 		}
 
 		/// <summary>
 		/// Register an agent to perform remote work.
 		/// </summary>
-		/// <param name="Request">Request parameters</param>
+		/// <param name="request">Request parameters</param>
 		/// <returns>Information about the registered agent</returns>
 		[HttpPost]
 		[Route("/api/v1/agents")]
-		public async Task<ActionResult<CreateAgentResponse>> CreateAgentAsync([FromBody] CreateAgentRequest Request)
+		public async Task<ActionResult<CreateAgentResponse>> CreateAgentAsync([FromBody] CreateAgentRequest request)
 		{
-			if(!await AclService.AuthorizeAsync(AclAction.CreateAgent, User))
+			if(!await _aclService.AuthorizeAsync(AclAction.CreateAgent, User))
 			{
 				return Forbid(AclAction.CreateAgent);
 			}
 
-			AgentSoftwareChannelName? Channel = String.IsNullOrEmpty(Request.Channel) ? (AgentSoftwareChannelName?)null : new AgentSoftwareChannelName(Request.Channel);
-			IAgent Agent = await AgentService.CreateAgentAsync(Request.Name, Request.Enabled, Channel, Request.Pools?.ConvertAll(x => new PoolId(x)));
+			AgentSoftwareChannelName? channel = String.IsNullOrEmpty(request.Channel) ? (AgentSoftwareChannelName?)null : new AgentSoftwareChannelName(request.Channel);
+			IAgent agent = await _agentService.CreateAgentAsync(request.Name, request.Enabled, channel, request.Pools?.ConvertAll(x => new PoolId(x)));
 
-			return new CreateAgentResponse(Agent.Id.ToString());
+			return new CreateAgentResponse(agent.Id.ToString());
 		}
 
 		/// <summary>
 		/// Finds the agents matching specified criteria.
 		/// </summary>
-		/// <param name="PoolId">The pool containing the agent</param>
-		/// <param name="Index">First result to return</param>
-		/// <param name="Count">Number of results to return</param>
-		/// <param name="ModifiedAfter">If set, only returns agents modified after this time</param>
-		/// <param name="Filter">Filter for the properties to return</param>
+		/// <param name="poolId">The pool containing the agent</param>
+		/// <param name="index">First result to return</param>
+		/// <param name="count">Number of results to return</param>
+		/// <param name="modifiedAfter">If set, only returns agents modified after this time</param>
+		/// <param name="filter">Filter for the properties to return</param>
 		/// <returns>List of matching agents</returns>
 		[HttpGet]
 		[Route("/api/v1/agents")]
 		[ProducesResponseType(typeof(List<GetAgentResponse>), 200)]
-		public async Task<ActionResult<List<object>>> FindAgentsAsync([FromQuery] PoolId? PoolId = null, [FromQuery] int? Index = null, [FromQuery] int? Count = null, [FromQuery] DateTimeOffset? ModifiedAfter = null, [FromQuery] PropertyFilter? Filter = null)
+		public async Task<ActionResult<List<object>>> FindAgentsAsync([FromQuery] PoolId? poolId = null, [FromQuery] int? index = null, [FromQuery] int? count = null, [FromQuery] DateTimeOffset? modifiedAfter = null, [FromQuery] PropertyFilter? filter = null)
 		{
-			GlobalPermissionsCache PermissionsCache = new GlobalPermissionsCache();
-			List<IAgent> Agents = await AgentService.FindAgentsAsync(PoolId, ModifiedAfter?.UtcDateTime, Index, Count);
+			GlobalPermissionsCache permissionsCache = new GlobalPermissionsCache();
+			List<IAgent> agents = await _agentService.FindAgentsAsync(poolId, modifiedAfter?.UtcDateTime, index, count);
 
-			List<object> Responses = new List<object>();
-			foreach (IAgent Agent in Agents)
+			List<object> responses = new List<object>();
+			foreach (IAgent agent in agents)
 			{
-				Responses.Add(await GetAgentResponseAsync(Agent, PermissionsCache, Filter));
+				responses.Add(await GetAgentResponseAsync(agent, permissionsCache, filter));
 			}
 
-			return Responses;
+			return responses;
 		}
 
 		/// <summary>
 		/// Retrieve information about a specific agent
 		/// </summary>
-		/// <param name="AgentId">Id of the agent to get information about</param>
-		/// <param name="Filter">Filter for the properties to return</param>
+		/// <param name="agentId">Id of the agent to get information about</param>
+		/// <param name="filter">Filter for the properties to return</param>
 		/// <returns>Information about the requested agent</returns>
 		[HttpGet]
-		[Route("/api/v1/agents/{AgentId}")]
+		[Route("/api/v1/agents/{agentId}")]
 		[ProducesResponseType(typeof(GetAgentResponse), 200)]
-		public async Task<ActionResult<object>> GetAgentAsync(AgentId AgentId, [FromQuery] PropertyFilter? Filter = null)
+		public async Task<ActionResult<object>> GetAgentAsync(AgentId agentId, [FromQuery] PropertyFilter? filter = null)
 		{
-			IAgent? Agent = await AgentService.GetAgentAsync(AgentId);
-			if (Agent == null)
+			IAgent? agent = await _agentService.GetAgentAsync(agentId);
+			if (agent == null)
 			{
-				return NotFound(AgentId);
+				return NotFound(agentId);
 			}
 
-			return await GetAgentResponseAsync(Agent, new GlobalPermissionsCache(), Filter);
+			return await GetAgentResponseAsync(agent, new GlobalPermissionsCache(), filter);
 		}
 
 		/// <summary>
 		/// Gets an individual agent response
 		/// </summary>
-		async ValueTask<object> GetAgentResponseAsync(IAgent Agent, GlobalPermissionsCache PermissionsCache, PropertyFilter? Filter = null)
+		async ValueTask<object> GetAgentResponseAsync(IAgent agent, GlobalPermissionsCache permissionsCache, PropertyFilter? filter = null)
 		{
-			double? Rate = null;
-			if (await AgentService.AuthorizeAsync(Agent, AclAction.ViewCosts, User, PermissionsCache))
+			double? rate = null;
+			if (await _agentService.AuthorizeAsync(agent, AclAction.ViewCosts, User, permissionsCache))
 			{
-				Rate = await AgentService.GetRateAsync(Agent.Id);
+				rate = await _agentService.GetRateAsync(agent.Id);
 			}
 
-			List<GetAgentLeaseResponse> Leases = new List<GetAgentLeaseResponse>();
-			foreach (AgentLease Lease in Agent.Leases)
+			List<GetAgentLeaseResponse> leases = new List<GetAgentLeaseResponse>();
+			foreach (AgentLease lease in agent.Leases)
 			{
-				Dictionary<string, string>? Details = AgentService.GetPayloadDetails(Lease.Payload);
-				Leases.Add(new GetAgentLeaseResponse(Lease, Details));
+				Dictionary<string, string>? details = _agentService.GetPayloadDetails(lease.Payload);
+				leases.Add(new GetAgentLeaseResponse(lease, details));
 			}
 
-			bool bIncludeAcl = await AgentService.AuthorizeAsync(Agent, AclAction.ViewPermissions, User, PermissionsCache);
-			return new GetAgentResponse(Agent, Leases, Rate, bIncludeAcl).ApplyFilter(Filter);
+			bool bIncludeAcl = await _agentService.AuthorizeAsync(agent, AclAction.ViewPermissions, User, permissionsCache);
+			return new GetAgentResponse(agent, leases, rate, bIncludeAcl).ApplyFilter(filter);
 		}
 
 		/// <summary>
 		/// Update an agent's properties.
 		/// </summary>
-		/// <param name="AgentId">Id of the agent to update.</param>
-		/// <param name="Update">Properties on the agent to update.</param>
+		/// <param name="agentId">Id of the agent to update.</param>
+		/// <param name="update">Properties on the agent to update.</param>
 		/// <returns>Http result code</returns>
 		[HttpPut]
-		[Route("/api/v1/agents/{AgentId}")]
-		public async Task<ActionResult> UpdateAgentAsync(AgentId AgentId, [FromBody] UpdateAgentRequest Update)
+		[Route("/api/v1/agents/{agentId}")]
+		public async Task<ActionResult> UpdateAgentAsync(AgentId agentId, [FromBody] UpdateAgentRequest update)
 		{
-			List<PoolId>? UpdatePools = Update.Pools?.ConvertAll(x => new PoolId(x));
-			AgentSoftwareChannelName? Channel = String.IsNullOrEmpty(Update.Channel) ? (AgentSoftwareChannelName?)null : new AgentSoftwareChannelName(Update.Channel);
+			List<PoolId>? updatePools = update.Pools?.ConvertAll(x => new PoolId(x));
+			AgentSoftwareChannelName? channel = String.IsNullOrEmpty(update.Channel) ? (AgentSoftwareChannelName?)null : new AgentSoftwareChannelName(update.Channel);
 
-			string UserName = User.GetUserName() ?? "Unknown";
+			string userName = User.GetUserName() ?? "Unknown";
 
-			GlobalPermissionsCache Cache = new GlobalPermissionsCache();
+			GlobalPermissionsCache cache = new GlobalPermissionsCache();
 			for (; ; )
 			{
-				IAgent? Agent = await AgentService.GetAgentAsync(AgentId);
-				if (Agent == null)
+				IAgent? agent = await _agentService.GetAgentAsync(agentId);
+				if (agent == null)
 				{
-					return NotFound(AgentId);
+					return NotFound(agentId);
 				}
-				if (!await AgentService.AuthorizeAsync(Agent, AclAction.UpdateAgent, User, Cache))
+				if (!await _agentService.AuthorizeAsync(agent, AclAction.UpdateAgent, User, cache))
 				{
-					return Forbid(AclAction.UpdateAgent, AgentId);
+					return Forbid(AclAction.UpdateAgent, agentId);
 				}
-				if (Update.Acl != null && !await AgentService.AuthorizeAsync(Agent, AclAction.ChangePermissions, User, Cache))
+				if (update.Acl != null && !await _agentService.AuthorizeAsync(agent, AclAction.ChangePermissions, User, cache))
 				{
-					return Forbid(AclAction.ChangePermissions, AgentId);
+					return Forbid(AclAction.ChangePermissions, agentId);
 				}
 
-				IAgent? NewAgent = await AgentService.Agents.TryUpdateSettingsAsync(Agent, Update.Enabled, Update.RequestConform, Update.RequestFullConform, Update.RequestRestart, Update.RequestShutdown, $"Manual ({UserName})", Channel, Update.Pools?.ConvertAll(x => new PoolId(x)), Acl.Merge(Agent.Acl, Update.Acl), Update.Comment);
-				if (NewAgent == null)
+				IAgent? newAgent = await _agentService.Agents.TryUpdateSettingsAsync(agent, update.Enabled, update.RequestConform, update.RequestFullConform, update.RequestRestart, update.RequestShutdown, $"Manual ({userName})", channel, update.Pools?.ConvertAll(x => new PoolId(x)), Acl.Merge(agent.Acl, update.Acl), update.Comment);
+				if (newAgent == null)
 				{
 					continue;
 				}
 
-				IAuditLogChannel<AgentId> Logger = AgentService.Agents.GetLogger(Agent.Id);
-				if (Agent.Enabled != NewAgent.Enabled)
+				IAuditLogChannel<AgentId> logger = _agentService.Agents.GetLogger(agent.Id);
+				if (agent.Enabled != newAgent.Enabled)
 				{
-					Logger.LogInformation("Setting changed: Enabled = {State} ({UserName})", NewAgent.Enabled, UserName);
+					logger.LogInformation("Setting changed: Enabled = {State} ({UserName})", newAgent.Enabled, userName);
 				}
-				if (Agent.RequestConform != NewAgent.RequestConform)
+				if (agent.RequestConform != newAgent.RequestConform)
 				{
-					Logger.LogInformation("Setting changed: RequestConform = {State} ({UserName})", NewAgent.RequestConform, UserName);
+					logger.LogInformation("Setting changed: RequestConform = {State} ({UserName})", newAgent.RequestConform, userName);
 				}
-				if (Agent.RequestFullConform != NewAgent.RequestFullConform)
+				if (agent.RequestFullConform != newAgent.RequestFullConform)
 				{
-					Logger.LogInformation("Setting changed: RequestFullConform = {State} ({UserName})", NewAgent.RequestFullConform, UserName);
+					logger.LogInformation("Setting changed: RequestFullConform = {State} ({UserName})", newAgent.RequestFullConform, userName);
 				}
-				if (Agent.RequestRestart != NewAgent.RequestRestart)
+				if (agent.RequestRestart != newAgent.RequestRestart)
 				{
-					Logger.LogInformation("Setting changed: RequestRestart = {State} ({UserName})", NewAgent.RequestRestart, UserName);
+					logger.LogInformation("Setting changed: RequestRestart = {State} ({UserName})", newAgent.RequestRestart, userName);
 				}
-				if (Agent.RequestShutdown != NewAgent.RequestShutdown)
+				if (agent.RequestShutdown != newAgent.RequestShutdown)
 				{
-					Logger.LogInformation("Setting changed: RequestShutdown = {State} ({UserName})", NewAgent.RequestShutdown, UserName);
+					logger.LogInformation("Setting changed: RequestShutdown = {State} ({UserName})", newAgent.RequestShutdown, userName);
 				}
-				if (Agent.Comment != NewAgent.Comment)
+				if (agent.Comment != newAgent.Comment)
 				{
-					Logger.LogInformation("Setting changed: Comment = \"{Comment}\" ({UserName})", Update.Comment, UserName);
+					logger.LogInformation("Setting changed: Comment = \"{Comment}\" ({UserName})", update.Comment, userName);
 				}
-				foreach (PoolId AddedPool in NewAgent.ExplicitPools.Except(Agent.ExplicitPools))
+				foreach (PoolId addedPool in newAgent.ExplicitPools.Except(agent.ExplicitPools))
 				{
-					Logger.LogInformation("Added to pool {PoolId} ({UserName})", AddedPool, UserName);
+					logger.LogInformation("Added to pool {PoolId} ({UserName})", addedPool, userName);
 				}
-				foreach (PoolId RemovedPool in Agent.ExplicitPools.Except(NewAgent.ExplicitPools))
+				foreach (PoolId removedPool in agent.ExplicitPools.Except(newAgent.ExplicitPools))
 				{
-					Logger.LogInformation("Removed from pool {PoolId} ({UserName})", RemovedPool, UserName);
+					logger.LogInformation("Removed from pool {PoolId} ({UserName})", removedPool, userName);
 				}
 				break;
 			}
@@ -226,179 +210,179 @@ namespace Horde.Build.Controllers
 		/// <summary>
 		/// Remove a registered agent.
 		/// </summary>
-		/// <param name="AgentId">Id of the agent to delete.</param>
+		/// <param name="agentId">Id of the agent to delete.</param>
 		/// <returns>Http result code</returns>
 		[HttpDelete]
-		[Route("/api/v1/agents/{AgentId}")]
-		public async Task<ActionResult> DeleteAgentAsync(AgentId AgentId)
+		[Route("/api/v1/agents/{agentId}")]
+		public async Task<ActionResult> DeleteAgentAsync(AgentId agentId)
 		{
-			IAgent? Agent = await AgentService.GetAgentAsync(AgentId);
-			if (Agent == null)
+			IAgent? agent = await _agentService.GetAgentAsync(agentId);
+			if (agent == null)
 			{
-				return NotFound(AgentId);
+				return NotFound(agentId);
 			}
-			if (!await AgentService.AuthorizeAsync(Agent, AclAction.DeleteAgent, User, null))
+			if (!await _agentService.AuthorizeAsync(agent, AclAction.DeleteAgent, User, null))
 			{
-				return Forbid(AclAction.DeleteAgent, AgentId);
+				return Forbid(AclAction.DeleteAgent, agentId);
 			}
 
-			await AgentService.DeleteAgentAsync(Agent);
+			await _agentService.DeleteAgentAsync(agent);
 			return new OkResult();
 		}
 
 		/// <summary>
 		/// Retrieve historical information about a specific agent
 		/// </summary>
-		/// <param name="AgentId">Id of the agent to get information about</param>
-		/// <param name="MinTime">Minimum time for records to return</param>
-		/// <param name="MaxTime">Maximum time for records to return</param>
-		/// <param name="Index">Offset of the first result</param>
-		/// <param name="Count">Number of records to return</param>
+		/// <param name="agentId">Id of the agent to get information about</param>
+		/// <param name="minTime">Minimum time for records to return</param>
+		/// <param name="maxTime">Maximum time for records to return</param>
+		/// <param name="index">Offset of the first result</param>
+		/// <param name="count">Number of records to return</param>
 		/// <returns>Information about the requested agent</returns>
 		[HttpGet]
-		[Route("/api/v1/agents/{AgentId}/history")]
-		public async Task GetAgentHistoryAsync(AgentId AgentId, [FromQuery] DateTime? MinTime = null, [FromQuery] DateTime? MaxTime = null, [FromQuery] int Index = 0, [FromQuery] int Count = 50)
+		[Route("/api/v1/agents/{agentId}/history")]
+		public async Task GetAgentHistoryAsync(AgentId agentId, [FromQuery] DateTime? minTime = null, [FromQuery] DateTime? maxTime = null, [FromQuery] int index = 0, [FromQuery] int count = 50)
 		{
 			Response.ContentType = "application/json";
 			Response.StatusCode = 200;
 			await Response.StartAsync();
-			await AgentService.Agents.GetLogger(AgentId).FindAsync(Response.BodyWriter, MinTime, MaxTime, Index, Count);
+			await _agentService.Agents.GetLogger(agentId).FindAsync(Response.BodyWriter, minTime, maxTime, index, count);
 		}
 
 		/// <summary>
 		/// Find all the sessions of a particular agent
 		/// </summary>
-		/// <param name="AgentId">Unique id of the agent to find</param>
-		/// <param name="StartTime">Start time to include in the search</param>
-		/// <param name="FinishTime">Finish time to include in the search</param>
-		/// <param name="Index">Index of the first result to return</param>
-		/// <param name="Count">Number of results to return</param>
+		/// <param name="agentId">Unique id of the agent to find</param>
+		/// <param name="startTime">Start time to include in the search</param>
+		/// <param name="finishTime">Finish time to include in the search</param>
+		/// <param name="index">Index of the first result to return</param>
+		/// <param name="count">Number of results to return</param>
 		/// <returns>Sessions </returns>
 		[HttpGet]
-		[Route("/api/v1/agents/{AgentId}/sessions")]
-		public async Task<ActionResult<List<GetAgentSessionResponse>>> FindSessionsAsync(AgentId AgentId, [FromQuery] DateTimeOffset? StartTime, [FromQuery] DateTimeOffset? FinishTime, [FromQuery] int Index = 0, [FromQuery] int Count = 50)
+		[Route("/api/v1/agents/{agentId}/sessions")]
+		public async Task<ActionResult<List<GetAgentSessionResponse>>> FindSessionsAsync(AgentId agentId, [FromQuery] DateTimeOffset? startTime, [FromQuery] DateTimeOffset? finishTime, [FromQuery] int index = 0, [FromQuery] int count = 50)
 		{
-			IAgent? Agent = await AgentService.GetAgentAsync(AgentId);
-			if (Agent == null)
+			IAgent? agent = await _agentService.GetAgentAsync(agentId);
+			if (agent == null)
 			{
-				return NotFound(AgentId);
+				return NotFound(agentId);
 			}
-			if (!await AgentService.AuthorizeAsync(Agent, AclAction.ViewSession, User, null))
+			if (!await _agentService.AuthorizeAsync(agent, AclAction.ViewSession, User, null))
 			{
-				return Forbid(AclAction.ViewSession, AgentId);
+				return Forbid(AclAction.ViewSession, agentId);
 			}
 
-			List<ISession> Sessions = await AgentService.FindSessionsAsync(AgentId, StartTime?.UtcDateTime, FinishTime?.UtcDateTime, Index, Count);
-			return Sessions.ConvertAll(x => new GetAgentSessionResponse(x));
+			List<ISession> sessions = await _agentService.FindSessionsAsync(agentId, startTime?.UtcDateTime, finishTime?.UtcDateTime, index, count);
+			return sessions.ConvertAll(x => new GetAgentSessionResponse(x));
 		}
 
 		/// <summary>
 		/// Find all the sessions of a particular agent
 		/// </summary>
-		/// <param name="AgentId">Unique id of the agent to find</param>
-		/// <param name="SessionId">Unique id of the session</param>
+		/// <param name="agentId">Unique id of the agent to find</param>
+		/// <param name="sessionId">Unique id of the session</param>
 		/// <returns>Sessions </returns>
 		[HttpGet]
-		[Route("/api/v1/agents/{AgentId}/sessions/{SessionId}")]
-		public async Task<ActionResult<GetAgentSessionResponse>> GetSessionAsync(AgentId AgentId, SessionId SessionId)
+		[Route("/api/v1/agents/{agentId}/sessions/{sessionId}")]
+		public async Task<ActionResult<GetAgentSessionResponse>> GetSessionAsync(AgentId agentId, SessionId sessionId)
 		{
-			IAgent? Agent = await AgentService.GetAgentAsync(AgentId);
-			if (Agent == null)
+			IAgent? agent = await _agentService.GetAgentAsync(agentId);
+			if (agent == null)
 			{
-				return NotFound(AgentId);
+				return NotFound(agentId);
 			}
-			if (!await AgentService.AuthorizeAsync(Agent, AclAction.ViewSession, User, null))
+			if (!await _agentService.AuthorizeAsync(agent, AclAction.ViewSession, User, null))
 			{
-				return Forbid(AclAction.ViewSession, AgentId);
+				return Forbid(AclAction.ViewSession, agentId);
 			}
 
-			ISession? Session = await AgentService.GetSessionAsync(SessionId);
-			if(Session == null || Session.AgentId != AgentId)
+			ISession? session = await _agentService.GetSessionAsync(sessionId);
+			if(session == null || session.AgentId != agentId)
 			{
 				return NotFound();
 			}
 
-			return new GetAgentSessionResponse(Session);
+			return new GetAgentSessionResponse(session);
 		}
 
 		/// <summary>
 		/// Find all the leases for a particular agent
 		/// </summary>
-		/// <param name="AgentId">Unique id of the agent to find</param>
-		/// <param name="SessionId">The session to query</param>
-		/// <param name="StartTime">Start of the time window to consider</param>
-		/// <param name="FinishTime">End of the time window to consider</param>
-		/// <param name="Index">Index of the first result to return</param>
-		/// <param name="Count">Number of results to return</param>
-		/// <param name="Filter">Filter to apply to the properties</param>
+		/// <param name="agentId">Unique id of the agent to find</param>
+		/// <param name="sessionId">The session to query</param>
+		/// <param name="startTime">Start of the time window to consider</param>
+		/// <param name="finishTime">End of the time window to consider</param>
+		/// <param name="index">Index of the first result to return</param>
+		/// <param name="count">Number of results to return</param>
+		/// <param name="filter">Filter to apply to the properties</param>
 		/// <returns>Sessions </returns>
 		[HttpGet]
-		[Route("/api/v1/agents/{AgentId}/leases")]
+		[Route("/api/v1/agents/{agentId}/leases")]
 		[ProducesResponseType(200, Type = typeof(List<GetAgentLeaseResponse>))]
-		public async Task<ActionResult<List<object>>> FindLeasesAsync(AgentId AgentId, [FromQuery] SessionId? SessionId, [FromQuery] DateTimeOffset? StartTime, [FromQuery] DateTimeOffset? FinishTime, [FromQuery] int Index = 0, [FromQuery] int Count = 1000, [FromQuery] PropertyFilter? Filter = null)
+		public async Task<ActionResult<List<object>>> FindLeasesAsync(AgentId agentId, [FromQuery] SessionId? sessionId, [FromQuery] DateTimeOffset? startTime, [FromQuery] DateTimeOffset? finishTime, [FromQuery] int index = 0, [FromQuery] int count = 1000, [FromQuery] PropertyFilter? filter = null)
 		{
-			IAgent? Agent = await AgentService.GetAgentAsync(AgentId);
-			if (Agent == null)
+			IAgent? agent = await _agentService.GetAgentAsync(agentId);
+			if (agent == null)
 			{
-				return NotFound(AgentId);
+				return NotFound(agentId);
 			}
-			if (!await AgentService.AuthorizeAsync(Agent, AclAction.ViewLeases, User, null))
+			if (!await _agentService.AuthorizeAsync(agent, AclAction.ViewLeases, User, null))
 			{
-				return Forbid(AclAction.ViewLeases, AgentId);
+				return Forbid(AclAction.ViewLeases, agentId);
 			}
 
-			List<ILease> Leases = await AgentService.FindLeasesAsync(AgentId, SessionId, StartTime?.UtcDateTime, FinishTime?.UtcDateTime, Index, Count);
+			List<ILease> leases = await _agentService.FindLeasesAsync(agentId, sessionId, startTime?.UtcDateTime, finishTime?.UtcDateTime, index, count);
 
-			double? AgentRate = null;
-			if (await AclService.AuthorizeAsync(AclAction.ViewCosts, User))
+			double? agentRate = null;
+			if (await _aclService.AuthorizeAsync(AclAction.ViewCosts, User))
 			{
-				AgentRate = await AgentService.GetRateAsync(AgentId);
+				agentRate = await _agentService.GetRateAsync(agentId);
 			}
 
-			List<object> Responses = new List<object>();
-			foreach(ILease Lease in Leases)
+			List<object> responses = new List<object>();
+			foreach(ILease lease in leases)
 			{
-				Dictionary<string, string>? Details = AgentService.GetPayloadDetails(Lease.Payload);
-				Responses.Add(PropertyFilter.Apply(new GetAgentLeaseResponse(Lease, Details, AgentRate), Filter));
+				Dictionary<string, string>? details = _agentService.GetPayloadDetails(lease.Payload);
+				responses.Add(PropertyFilter.Apply(new GetAgentLeaseResponse(lease, details, agentRate), filter));
 			}
 
-			return Responses;
+			return responses;
 		}
 
 		/// <summary>
 		/// Get info about a particular lease
 		/// </summary>
-		/// <param name="AgentId">Unique id of the agent to find</param>
-		/// <param name="LeaseId">Unique id of the particular lease</param>
+		/// <param name="agentId">Unique id of the agent to find</param>
+		/// <param name="leaseId">Unique id of the particular lease</param>
 		/// <returns>Lease matching the given id</returns>
 		[HttpGet]
-		[Route("/api/v1/agents/{AgentId}/leases/{LeaseId}")]
-		public async Task<ActionResult<GetAgentLeaseResponse>> GetLeaseAsync(AgentId AgentId, LeaseId LeaseId)
+		[Route("/api/v1/agents/{agentId}/leases/{leaseId}")]
+		public async Task<ActionResult<GetAgentLeaseResponse>> GetLeaseAsync(AgentId agentId, LeaseId leaseId)
 		{
-			IAgent? Agent = await AgentService.GetAgentAsync(AgentId);
-			if(Agent == null)
+			IAgent? agent = await _agentService.GetAgentAsync(agentId);
+			if(agent == null)
 			{
-				return NotFound(AgentId);
+				return NotFound(agentId);
 			}
-			if (!await AclService.AuthorizeAsync(AclAction.ViewLeases, User))
+			if (!await _aclService.AuthorizeAsync(AclAction.ViewLeases, User))
 			{
-				return Forbid(AclAction.ViewLeases, AgentId);
+				return Forbid(AclAction.ViewLeases, agentId);
 			}
 
-			ILease? Lease = await AgentService.GetLeaseAsync(LeaseId);
-			if (Lease == null || Lease.AgentId != AgentId)
+			ILease? lease = await _agentService.GetLeaseAsync(leaseId);
+			if (lease == null || lease.AgentId != agentId)
 			{
-				return NotFound(AgentId, LeaseId);
+				return NotFound(agentId, leaseId);
 			}
 
-			double? AgentRate = null;
-			if (await AclService.AuthorizeAsync(AclAction.ViewCosts, User))
+			double? agentRate = null;
+			if (await _aclService.AuthorizeAsync(AclAction.ViewCosts, User))
 			{
-				AgentRate = await AgentService.GetRateAsync(AgentId);
+				agentRate = await _agentService.GetRateAsync(agentId);
 			}
 
-			Dictionary<string, string>? Details = AgentService.GetPayloadDetails(Lease.Payload);
-			return new GetAgentLeaseResponse(Lease, Details, AgentRate);
+			Dictionary<string, string>? details = _agentService.GetPayloadDetails(lease.Payload);
+			return new GetAgentLeaseResponse(lease, details, agentRate);
 		}
 	}
 }
diff --git a/Engine/Source/Programs/Horde/Horde.Build/Controllers/ArtifactsController.cs b/Engine/Source/Programs/Horde/Horde.Build/Controllers/ArtifactsController.cs
--- a/Engine/Source/Programs/Horde/Horde.Build/Controllers/ArtifactsController.cs
+++ b/Engine/Source/Programs/Horde/Horde.Build/Controllers/ArtifactsController.cs
@@ -1,6 +1,14 @@
-﻿// Copyright Epic Games, Inc. All Rights Reserved.
+// Copyright Epic Games, Inc. All Rights Reserved.
 
-using Amazon.S3.Transfer;
+using System;
+using System.Collections.Generic;
+using System.IdentityModel.Tokens.Jwt;
+using System.IO;
+using System.IO.Compression;
+using System.Linq;
+using System.Net.Mime;
+using System.Security.Claims;
+using System.Threading.Tasks;
 using Horde.Build.Acls;
 using Horde.Build.Api;
 using Horde.Build.Jobs;
@@ -10,21 +18,8 @@ using Horde.Build.Utilities;
 using Microsoft.AspNetCore.Authorization;
 using Microsoft.AspNetCore.Http;
 using Microsoft.AspNetCore.Mvc;
-using Microsoft.AspNetCore.StaticFiles;
 using Microsoft.IdentityModel.Tokens;
 using MongoDB.Bson;
-using System;
-using System.Collections.Generic;
-using System.Globalization;
-using System.IdentityModel.Tokens.Jwt;
-using System.IO;
-using System.IO.Compression;
-using System.Linq;
-using System.Net.Http.Headers;
-using System.Net.Mime;
-using System.Security.Claims;
-using System.Text;
-using System.Threading.Tasks;
 
 namespace Horde.Build.Controllers
 {
@@ -40,226 +35,226 @@ namespace Horde.Build.Controllers
 		/// <summary>
 		/// Instance of the database service
 		/// </summary>
-		private readonly DatabaseService DatabaseService;
+		private readonly DatabaseService _databaseService;
 
 		/// <summary>
 		/// Instance of the artifact collection
 		/// </summary>
-		private readonly IArtifactCollection ArtifactCollection;
+		private readonly IArtifactCollection _artifactCollection;
 
 		/// <summary>
 		/// Instance of the ACL service
 		/// </summary>
-		private readonly AclService AclService;
+		private readonly AclService _aclService;
 
 		/// <summary>
 		/// Instance of the Job service
 		/// </summary>
-		private readonly JobService JobService;
+		private readonly JobService _jobService;
 
 		/// <summary>
 		/// Constructor
 		/// </summary>
-		public ArtifactsController(DatabaseService DatabaseSerivce, IArtifactCollection ArtifactCollection, AclService AclService, JobService JobService)
+		public ArtifactsController(DatabaseService databaseSerivce, IArtifactCollection artifactCollection, AclService aclService, JobService jobService)
 		{
-			this.DatabaseService = DatabaseSerivce;
-			this.ArtifactCollection = ArtifactCollection;
-			this.AclService = AclService;
-			this.JobService = JobService;
+			_databaseService = databaseSerivce;
+			_artifactCollection = artifactCollection;
+			_aclService = aclService;
+			_jobService = jobService;
 		}
 
 		/// <summary>
 		/// Creates an artifact
 		/// </summary>
-		/// <param name="JobId">BatchId</param>
-		/// <param name="StepId">StepId</param>
-		/// <param name="File">The file contents</param>
+		/// <param name="jobId">BatchId</param>
+		/// <param name="stepId">StepId</param>
+		/// <param name="file">The file contents</param>
 		/// <returns>Http result code</returns>
 		[HttpPost]
 		[Authorize]
 		[Route("/api/v1/artifacts")]
-		public async Task<ActionResult<CreateArtifactResponse>> CreateArtifact([FromQuery] JobId JobId, [FromQuery]string? StepId, IFormFile File)
+		public async Task<ActionResult<CreateArtifactResponse>> CreateArtifact([FromQuery] JobId jobId, [FromQuery]string? stepId, IFormFile file)
 		{
-			IJob? Job = await JobService.GetJobAsync(JobId);
-			if(Job == null)
+			IJob? job = await _jobService.GetJobAsync(jobId);
+			if(job == null)
 			{
 				return NotFound();
 			}
 
-			if (!await JobService.AuthorizeAsync(Job, AclAction.UploadArtifact, User, null))
+			if (!await _jobService.AuthorizeAsync(job, AclAction.UploadArtifact, User, null))
 			{
 				return Forbid();
 			}
 
-			IJobStep? Step = null;
-			if(StepId != null)
+			IJobStep? step = null;
+			if(stepId != null)
 			{
-				foreach(IJobStepBatch Batch in Job.Batches)
+				foreach(IJobStepBatch batch in job.Batches)
 				{
-					if(Batch.TryGetStep(StepId.ToSubResourceId(), out Step))
+					if(batch.TryGetStep(stepId.ToSubResourceId(), out step))
 					{
 						break;
 					}
 				}
-				if(Step == null)
+				if(step == null)
 				{
 					// if the step doesn't exist in any of the batches, not found
 					return NotFound();
 				}
 			}
 
-			IArtifact NewArtifact = await ArtifactCollection.CreateArtifactAsync(Job.Id, Step?.Id, File.FileName, File.ContentType ?? "horde-mime/unknown", File.OpenReadStream());
-			return new CreateArtifactResponse(NewArtifact.Id.ToString());
+			IArtifact newArtifact = await _artifactCollection.CreateArtifactAsync(job.Id, step?.Id, file.FileName, file.ContentType ?? "horde-mime/unknown", file.OpenReadStream());
+			return new CreateArtifactResponse(newArtifact.Id.ToString());
 		}
 
 		/// <summary>
 		/// Updates an artifact
 		/// </summary>
-		/// <param name="ArtifactId">JobId</param>
-		/// <param name="File">The file contents</param>
+		/// <param name="artifactId">JobId</param>
+		/// <param name="file">The file contents</param>
 		/// <returns>Http result code</returns>
 		[HttpPut]
 		[Authorize]
-		[Route("/api/v1/artifacts/{ArtifactId}")]
-		public async Task<ActionResult<CreateArtifactResponse>> UpdateArtifact(string ArtifactId, IFormFile File)
+		[Route("/api/v1/artifacts/{artifactId}")]
+		public async Task<ActionResult<CreateArtifactResponse>> UpdateArtifact(string artifactId, IFormFile file)
 		{
-			IArtifact? Artifact = await ArtifactCollection.GetArtifactAsync(ArtifactId.ToObjectId());
-			if (Artifact == null)
+			IArtifact? artifact = await _artifactCollection.GetArtifactAsync(artifactId.ToObjectId());
+			if (artifact == null)
 			{
 				return NotFound();
 			}
-			if (!await JobService.AuthorizeAsync(Artifact.JobId, AclAction.UploadArtifact, User, null))
+			if (!await _jobService.AuthorizeAsync(artifact.JobId, AclAction.UploadArtifact, User, null))
 			{
 				return Forbid();
 			}
 
-			await ArtifactCollection.UpdateArtifactAsync(Artifact, File.ContentType ?? "horde-mime/unknown", File.OpenReadStream());
+			await _artifactCollection.UpdateArtifactAsync(artifact, file.ContentType ?? "horde-mime/unknown", file.OpenReadStream());
 			return Ok();
 		}
 
 		/// <summary>
 		/// Query artifacts for a job step
 		/// </summary>
-		/// <param name="JobId">Optional JobId to filter by</param>
-		/// <param name="StepId">Optional StepId to filter by</param>
-		/// <param name="Code">Whether to generate a direct download code</param>
-		/// <param name="Filter">Filter for the properties to return</param>
+		/// <param name="jobId">Optional JobId to filter by</param>
+		/// <param name="stepId">Optional StepId to filter by</param>
+		/// <param name="code">Whether to generate a direct download code</param>
+		/// <param name="filter">Filter for the properties to return</param>
 		/// <returns>Information about all the artifacts</returns>
 		[HttpGet]
 		[Authorize]
 		[Route("/api/v1/artifacts")]
 		[ProducesResponseType(typeof(List<GetArtifactResponse>), 200)]
-		public async Task<ActionResult<List<object>>> GetArtifacts([FromQuery] JobId JobId, [FromQuery] string? StepId = null, [FromQuery] bool Code = false, [FromQuery] PropertyFilter? Filter = null)
+		public async Task<ActionResult<List<object>>> GetArtifacts([FromQuery] JobId jobId, [FromQuery] string? stepId = null, [FromQuery] bool code = false, [FromQuery] PropertyFilter? filter = null)
 		{
-			if (!await JobService.AuthorizeAsync(JobId, AclAction.DownloadArtifact, User, null))
+			if (!await _jobService.AuthorizeAsync(jobId, AclAction.DownloadArtifact, User, null))
 			{
 				return Forbid();
 			}
 
-			string? DownloadCode = Code ? (string?)GetDirectDownloadCodeForJob(JobId) : null;
+			string? downloadCode = code ? (string?)GetDirectDownloadCodeForJob(jobId) : null;
 
-			List<IArtifact> Artifacts = await ArtifactCollection.GetArtifactsAsync(JobId, StepId?.ToSubResourceId(), null);
-			return Artifacts.ConvertAll(x => new GetArtifactResponse(x, DownloadCode).ApplyFilter(Filter));
+			List<IArtifact> artifacts = await _artifactCollection.GetArtifactsAsync(jobId, stepId?.ToSubResourceId(), null);
+			return artifacts.ConvertAll(x => new GetArtifactResponse(x, downloadCode).ApplyFilter(filter));
 		}
 
 		/// <summary>
 		/// Gets the claim required to download artifacts for a particular job
 		/// </summary>
-		/// <param name="JobId">The job id</param>
+		/// <param name="jobId">The job id</param>
 		/// <returns>The required claim</returns>
-		static Claim GetDirectDownloadClaim(JobId JobId)
+		static Claim GetDirectDownloadClaim(JobId jobId)
 		{
-			return new Claim(HordeClaimTypes.JobArtifacts, JobId.ToString());
+			return new Claim(HordeClaimTypes.JobArtifacts, jobId.ToString());
 		}
 
 		/// <summary>
 		/// Get a download code for the artifacts of a job
 		/// </summary>
-		/// <param name="JobId">The job id</param>
+		/// <param name="jobId">The job id</param>
 		/// <returns>The download code</returns>
-		string GetDirectDownloadCodeForJob(JobId JobId)
+		string GetDirectDownloadCodeForJob(JobId jobId)
 		{
-			Claim DownloadClaim = GetDirectDownloadClaim(JobId);
-			return AclService.IssueBearerToken(new[] { DownloadClaim }, TimeSpan.FromHours(4.0));
+			Claim downloadClaim = GetDirectDownloadClaim(jobId);
+			return _aclService.IssueBearerToken(new[] { downloadClaim }, TimeSpan.FromHours(4.0));
 		}
 
 		/// <summary>
 		/// Retrieve metadata about a specific artifact
 		/// </summary>
-		/// <param name="ArtifactId">Id of the artifact to get information about</param>
-		/// <param name="Code">Whether to generate a direct download code</param>
-		/// <param name="Filter">Filter for the properties to return</param>
+		/// <param name="artifactId">Id of the artifact to get information about</param>
+		/// <param name="code">Whether to generate a direct download code</param>
+		/// <param name="filter">Filter for the properties to return</param>
 		/// <returns>Information about the requested project</returns>
 		[HttpGet]
 		[Authorize]
-		[Route("/api/v1/artifacts/{ArtifactId}")]
+		[Route("/api/v1/artifacts/{artifactId}")]
 		[ProducesResponseType(typeof(GetArtifactResponse), 200)]
-		public async Task<ActionResult<object>> GetArtifact(string ArtifactId, bool Code = false, [FromQuery] PropertyFilter? Filter = null)
+		public async Task<ActionResult<object>> GetArtifact(string artifactId, bool code = false, [FromQuery] PropertyFilter? filter = null)
 		{
-			IArtifact? Artifact = await ArtifactCollection.GetArtifactAsync(ArtifactId.ToObjectId());
-			if (Artifact == null)
+			IArtifact? artifact = await _artifactCollection.GetArtifactAsync(artifactId.ToObjectId());
+			if (artifact == null)
 			{
 				return NotFound();
 			}
-			if (!await JobService.AuthorizeAsync(Artifact.JobId, AclAction.DownloadArtifact, User, null))
+			if (!await _jobService.AuthorizeAsync(artifact.JobId, AclAction.DownloadArtifact, User, null))
 			{
 				return Forbid();
 			}
 
-			string? DownloadCode = Code? (string?)GetDirectDownloadCodeForJob(Artifact.JobId) : null;
-			return new GetArtifactResponse(Artifact, DownloadCode).ApplyFilter(Filter);
+			string? downloadCode = code? (string?)GetDirectDownloadCodeForJob(artifact.JobId) : null;
+			return new GetArtifactResponse(artifact, downloadCode).ApplyFilter(filter);
 		}
 
 		/// <summary>
 		/// Retrieve raw data for an artifact
 		/// </summary>
-		/// <param name="ArtifactId">Id of the artifact to get information about</param>
+		/// <param name="artifactId">Id of the artifact to get information about</param>
 		/// <returns>Raw artifact data</returns>
 		[HttpGet]
 		[Authorize]
-		[Route("/api/v1/artifacts/{ArtifactId}/data")]
-		public async Task<ActionResult> GetArtifactData(string ArtifactId)
+		[Route("/api/v1/artifacts/{artifactId}/data")]
+		public async Task<ActionResult> GetArtifactData(string artifactId)
 		{
-			IArtifact? Artifact = await ArtifactCollection.GetArtifactAsync(ArtifactId.ToObjectId());
-			if (Artifact == null)
+			IArtifact? artifact = await _artifactCollection.GetArtifactAsync(artifactId.ToObjectId());
+			if (artifact == null)
 			{
 				return NotFound();
 			}
-			if (!await JobService.AuthorizeAsync(Artifact.JobId, AclAction.DownloadArtifact, User, null))
+			if (!await _jobService.AuthorizeAsync(artifact.JobId, AclAction.DownloadArtifact, User, null))
 			{
 				return Forbid();
 			}
 
 			// Fun, filestream result automatically closes the stream!
-			return new FileStreamResult(await ArtifactCollection.OpenArtifactReadStreamAsync(Artifact), Artifact.MimeType);
+			return new FileStreamResult(await _artifactCollection.OpenArtifactReadStreamAsync(artifact), artifact.MimeType);
 		}
 		
 		/// <summary>
 		/// Retrieve raw data for an artifact by filename
 		/// </summary>
-		/// <param name="JobId">Unique id for the job</param>
-		/// <param name="StepId">Unique id for the step</param>
-		/// <param name="Filename">Filename of artifact from step</param>
+		/// <param name="jobId">Unique id for the job</param>
+		/// <param name="stepId">Unique id for the step</param>
+		/// <param name="filename">Filename of artifact from step</param>
 		/// <returns>Raw artifact data</returns>
 		[HttpGet]
-		[Route("/api/v1/jobs/{JobId}/steps/{StepId}/artifacts/{FileName}/data")]
-		public async Task<ActionResult<object>> GetArtifactDataByFilename(JobId JobId, string StepId, string Filename)
+		[Route("/api/v1/jobs/{jobId}/steps/{stepId}/artifacts/{filename}/data")]
+		public async Task<ActionResult<object>> GetArtifactDataByFilename(JobId jobId, string stepId, string filename)
 		{
-			SubResourceId StepIdValue = StepId.ToSubResourceId();
+			SubResourceId stepIdValue = stepId.ToSubResourceId();
 
-			if (!await JobService.AuthorizeAsync(JobId, AclAction.DownloadArtifact, User, null))
+			if (!await _jobService.AuthorizeAsync(jobId, AclAction.DownloadArtifact, User, null))
 			{
 				return Forbid();
 			}
 			
-			List<IArtifact> Artifacts = await ArtifactCollection.GetArtifactsAsync(JobId, StepIdValue, Filename);
-			if (Artifacts.Count == 0)
+			List<IArtifact> artifacts = await _artifactCollection.GetArtifactsAsync(jobId, stepIdValue, filename);
+			if (artifacts.Count == 0)
 			{
 				return NotFound();
 			}
 
-			IArtifact Artifact = Artifacts[0];
-			return new FileStreamResult(await ArtifactCollection.OpenArtifactReadStreamAsync(Artifact), Artifact.MimeType);
+			IArtifact artifact = artifacts[0];
+			return new FileStreamResult(await _artifactCollection.OpenArtifactReadStreamAsync(artifact), artifact.MimeType);
 		}
 
 		/// <summary>
@@ -270,114 +265,113 @@ namespace Horde.Build.Controllers
 			/// <summary>
 			/// The suggested download filename
 			/// </summary>
-			string FileName;
+			readonly string _fileName;
 
 			/// <summary>
 			/// Constructor
 			/// </summary>
-			/// <param name="Stream"></param>
-			/// <param name="MimeType"></param>
-			/// <param name="FileName"></param>
-			public InlineFileStreamResult(System.IO.Stream Stream, string MimeType, string FileName)
-				: base(Stream, MimeType)
+			/// <param name="stream"></param>
+			/// <param name="mimeType"></param>
+			/// <param name="fileName"></param>
+			public InlineFileStreamResult(System.IO.Stream stream, string mimeType, string fileName)
+				: base(stream, mimeType)
 			{
-				this.FileName = FileName;
+				_fileName = fileName;
 			}
 
 			/// <inheritdoc/>
-			public override Task ExecuteResultAsync(ActionContext Context)
+			public override Task ExecuteResultAsync(ActionContext context)
 			{
-				ContentDisposition ContentDisposition = new ContentDisposition();
-				ContentDisposition.Inline = true;
-				ContentDisposition.FileName = FileName;
-				Context.HttpContext.Response.Headers.Add("Content-Disposition", ContentDisposition.ToString());
+				ContentDisposition contentDisposition = new ContentDisposition();
+				contentDisposition.Inline = true;
+				contentDisposition.FileName = _fileName;
+				context.HttpContext.Response.Headers.Add("Content-Disposition", contentDisposition.ToString());
 				
-				return base.ExecuteResultAsync(Context);
+				return base.ExecuteResultAsync(context);
 			}
 		}
 
 		/// <summary>
 		/// Retrieve raw data for an artifact
 		/// </summary>
-		/// <param name="ArtifactId">Id of the artifact to get information about</param>
-		/// <param name="Code">The authorization code for this resource</param>
+		/// <param name="artifactId">Id of the artifact to get information about</param>
+		/// <param name="code">The authorization code for this resource</param>
 		/// <returns>Raw artifact data</returns>
 		[HttpGet]
 		[AllowAnonymous]
-		[Route("/api/v1/artifacts/{ArtifactId}/download")]
-		public async Task<ActionResult> DownloadArtifact(string ArtifactId, [FromQuery] string Code)
+		[Route("/api/v1/artifacts/{artifactId}/download")]
+		public async Task<ActionResult> DownloadArtifact(string artifactId, [FromQuery] string code)
 		{
-			TokenValidationParameters Parameters = new TokenValidationParameters();
-			Parameters.ValidateAudience = false;
-			Parameters.RequireExpirationTime = true;
-			Parameters.ValidateLifetime = true;
-			Parameters.ValidIssuer = DatabaseService.JwtIssuer;
-			Parameters.ValidateIssuer = true;
-			Parameters.ValidateIssuerSigningKey = true;
-			Parameters.IssuerSigningKey = DatabaseService.JwtSigningKey;
-
-			SecurityToken Token;
-			JwtSecurityTokenHandler Handler = new JwtSecurityTokenHandler();
-			ClaimsPrincipal Principal = Handler.ValidateToken(Code, Parameters, out Token);
-
-			IArtifact? Artifact = await ArtifactCollection.GetArtifactAsync(ArtifactId.ToObjectId());
-			if (Artifact == null)
+			TokenValidationParameters parameters = new TokenValidationParameters();
+			parameters.ValidateAudience = false;
+			parameters.RequireExpirationTime = true;
+			parameters.ValidateLifetime = true;
+			parameters.ValidIssuer = _databaseService.JwtIssuer;
+			parameters.ValidateIssuer = true;
+			parameters.ValidateIssuerSigningKey = true;
+			parameters.IssuerSigningKey = _databaseService.JwtSigningKey;
+
+			JwtSecurityTokenHandler handler = new JwtSecurityTokenHandler();
+			ClaimsPrincipal principal = handler.ValidateToken(code, parameters, out _);
+
+			IArtifact? artifact = await _artifactCollection.GetArtifactAsync(artifactId.ToObjectId());
+			if (artifact == null)
 			{
 				return NotFound();
 			}
 
-			Claim DirectDownloadClaim = GetDirectDownloadClaim(Artifact.JobId);
-			if (!Principal.HasClaim(DirectDownloadClaim.Type, DirectDownloadClaim.Value))
+			Claim directDownloadClaim = GetDirectDownloadClaim(artifact.JobId);
+			if (!principal.HasClaim(directDownloadClaim.Type, directDownloadClaim.Value))
 			{
 				return Forbid();
 			}
 
-			return new InlineFileStreamResult(await ArtifactCollection.OpenArtifactReadStreamAsync(Artifact), Artifact.MimeType, Path.GetFileName(Artifact.Name));
+			return new InlineFileStreamResult(await _artifactCollection.OpenArtifactReadStreamAsync(artifact), artifact.MimeType, Path.GetFileName(artifact.Name));
 		}
 
 		/// <summary>
 		/// Returns a zip archive of many artifacts
 		/// </summary>
-		/// <param name="ArtifactZipRequest">Artifact request params</param>
+		/// <param name="artifactZipRequest">Artifact request params</param>
 		/// <returns>Zip of many artifacts</returns>
 		[HttpPost]
 		[Authorize]
 		[Route("/api/v1/artifacts/zip")]
-		public async Task<ActionResult> ZipArtifacts(GetArtifactZipRequest ArtifactZipRequest)
+		public async Task<ActionResult> ZipArtifacts(GetArtifactZipRequest artifactZipRequest)
 		{
-			if (ArtifactZipRequest.JobId == null)
+			if (artifactZipRequest.JobId == null)
 			{
 				return BadRequest("Must specify a JobId");
 			}
 
-			IJob? Job = await JobService.GetJobAsync(new JobId(ArtifactZipRequest.JobId!));
-			if (Job == null)
+			IJob? job = await _jobService.GetJobAsync(new JobId(artifactZipRequest.JobId!));
+			if (job == null)
 			{
 				return NotFound();
 			}
-			if (!await JobService.AuthorizeAsync(Job, AclAction.DownloadArtifact, User, null))
+			if (!await _jobService.AuthorizeAsync(job, AclAction.DownloadArtifact, User, null))
 			{
 				return Forbid();
 			}
 
-			List<IArtifact> Artifacts = await ArtifactCollection.GetArtifactsAsync(Job.Id, ArtifactZipRequest.StepId?.ToSubResourceId(), null);
+			List<IArtifact> artifacts = await _artifactCollection.GetArtifactsAsync(job.Id, artifactZipRequest.StepId?.ToSubResourceId(), null);
 
-			Dictionary<ObjectId, IArtifact> IdToArtifact = Artifacts.ToDictionary(x => x.Id, x => x);
+			Dictionary<ObjectId, IArtifact> idToArtifact = artifacts.ToDictionary(x => x.Id, x => x);
 
-			List<IArtifact> ZipArtifacts;
-			if (ArtifactZipRequest.ArtifactIds == null)
+			List<IArtifact> zipArtifacts;
+			if (artifactZipRequest.ArtifactIds == null)
 			{
-				ZipArtifacts = Artifacts;
+				zipArtifacts = artifacts;
 			}
 			else
 			{
-				ZipArtifacts = new List<IArtifact>();
-				foreach (string ArtifactId in ArtifactZipRequest.ArtifactIds)
+				zipArtifacts = new List<IArtifact>();
+				foreach (string artifactId in artifactZipRequest.ArtifactIds)
 				{
-					IArtifact? Artifact;
-					if (IdToArtifact.TryGetValue(ArtifactId.ToObjectId(), out Artifact))
+					IArtifact? artifact;
+					if (idToArtifact.TryGetValue(artifactId.ToObjectId(), out artifact))
 					{
-						ZipArtifacts.Add(Artifact);
+						zipArtifacts.Add(artifact);
 					}
 					else
 					{
@@ -386,61 +380,61 @@ namespace Horde.Build.Controllers
 				}
 			}
 
-			IGraph Graph = await JobService.GetGraphAsync(Job);
+			IGraph graph = await _jobService.GetGraphAsync(job);
 
-			return new CustomFileCallbackResult("Artifacts.zip", "application/octet-stream", false, async (OutputStream, Context) =>
+			return new CustomFileCallbackResult("Artifacts.zip", "application/octet-stream", false, async (outputStream, context) =>
 			{
 				// Make an unseekable MemoryStream for the ZipArchive. We have to do this because the ZipEntry stream falls back to a synchronous write to it's own stream wrappers.
-				using (CustomBufferStream ZipOutputStream = new CustomBufferStream())
+				using (CustomBufferStream zipOutputStream = new CustomBufferStream())
 				{
 					// Keep the stream open after dispose so we can write the EOF bits.
-					using (ZipArchive ZipArchive = new ZipArchive(ZipOutputStream, ZipArchiveMode.Create, true))
+					using (ZipArchive zipArchive = new ZipArchive(zipOutputStream, ZipArchiveMode.Create, true))
 					{
-						foreach (IArtifact Artifact in ZipArtifacts)
+						foreach (IArtifact artifact in zipArtifacts)
 						{
-							await using (System.IO.Stream ArtifactStream = await ArtifactCollection.OpenArtifactReadStreamAsync(Artifact))
+							await using (System.IO.Stream artifactStream = await _artifactCollection.OpenArtifactReadStreamAsync(artifact))
 							{
 								// tack on the step name into the directory if it exists
-								string StepName = string.Empty;
-								if (Artifact.StepId.HasValue)
+								string stepName = String.Empty;
+								if (artifact.StepId.HasValue)
 								{
-									foreach (IJobStepBatch Batch in Job.Batches)
+									foreach (IJobStepBatch batch in job.Batches)
 									{
-										IJobStep? Step;
-										if (Batch.TryGetStep(Artifact.StepId.Value, out Step))
+										IJobStep? step;
+										if (batch.TryGetStep(artifact.StepId.Value, out step))
 										{
-											StepName = Graph.Groups[Batch.GroupIdx].Nodes[Step.NodeIdx].Name;
+											stepName = graph.Groups[batch.GroupIdx].Nodes[step.NodeIdx].Name;
 											break;
 										}
 									}
 								}
 
-								ZipArchiveEntry ZipEntry = ZipArchive.CreateEntry(Artifact.Name);
-								using (System.IO.Stream EntryStream = ZipEntry.Open())
+								ZipArchiveEntry zipEntry = zipArchive.CreateEntry(artifact.Name);
+								using (System.IO.Stream entryStream = zipEntry.Open())
 								{
-									byte[] Buffer = new byte[4096];
-									int TotalBytesRead = 0;
-									while (TotalBytesRead < Artifact.Length)
+									byte[] buffer = new byte[4096];
+									int totalBytesRead = 0;
+									while (totalBytesRead < artifact.Length)
 									{
-										int BytesRead = await ArtifactStream.ReadAsync(Buffer, 0, Buffer.Length);
+										int bytesRead = await artifactStream.ReadAsync(buffer, 0, buffer.Length);
 
 										// Write bytes to the entry stream.  Also advances the MemStream pos.
-										await EntryStream.WriteAsync(Buffer, 0, BytesRead);
+										await entryStream.WriteAsync(buffer, 0, bytesRead);
 										// Dump what we have to the output stream
-										await OutputStream.WriteAsync(ZipOutputStream.GetBuffer(), 0, (int)ZipOutputStream.Position);
+										await outputStream.WriteAsync(zipOutputStream.GetBuffer(), 0, (int)zipOutputStream.Position);
 
 										// Reset everything.
-										ZipOutputStream.Position = 0;
-										ZipOutputStream.SetLength(0);
-										TotalBytesRead += BytesRead;
+										zipOutputStream.Position = 0;
+										zipOutputStream.SetLength(0);
+										totalBytesRead += bytesRead;
 									}
 								}
 							}
 						}
 					}
 					// Write out the EOF stuff
-					ZipOutputStream.Position = 0;
-					await ZipOutputStream.CopyToAsync(OutputStream);
+					zipOutputStream.Position = 0;
+					await zipOutputStream.CopyToAsync(outputStream);
 				}
 			});
 		}
diff --git a/Engine/Source/Programs/Horde/Horde.Build/Controllers/ConfigController.cs b/Engine/Source/Programs/Horde/Horde.Build/Controllers/ConfigController.cs
--- a/Engine/Source/Programs/Horde/Horde.Build/Controllers/ConfigController.cs
+++ b/Engine/Source/Programs/Horde/Horde.Build/Controllers/ConfigController.cs
@@ -1,18 +1,17 @@
-﻿// Copyright Epic Games, Inc. All Rights Reserved.
+// Copyright Epic Games, Inc. All Rights Reserved.
 
+using System;
+using System.Text.Json;
+using System.Threading.Tasks;
 using EpicGames.Core;
 using Horde.Build.Acls;
 using Horde.Build.Api;
 using Horde.Build.Models;
 using Horde.Build.Services;
+using Horde.Build.Utilities;
 using Microsoft.AspNetCore.Authorization;
 using Microsoft.AspNetCore.Mvc;
 using Microsoft.Extensions.Options;
-using System.Threading.Tasks;
-using System.Text.Json;
-using Microsoft.Extensions.Logging;
-using System.Runtime.InteropServices;
-using Horde.Build.Utilities;
 
 namespace Horde.Build.Controllers
 {
@@ -25,45 +24,32 @@ namespace Horde.Build.Controllers
 	[Route("[controller]")]
 	public class ConfigController : HordeControllerBase
 	{
-		/// <summary>
-		/// The database service singleton
-		/// </summary>
-		DatabaseService DatabaseService;
-
 		/// <summary>
 		/// The acl service singleton
 		/// </summary>
-		AclService AclService;
+		readonly AclService _aclService;
 
 		/// <summary>
 		/// The config service singleton
 		/// </summary>
-		ConfigService ConfigService;
-		/// <summary>
-		/// Settings for the server
-		/// </summary>
-		IOptionsMonitor<ServerSettings> Settings;
+		readonly ConfigService _configService;
 
 		/// <summary>
-		///  Logger for controller
+		/// Settings for the server
 		/// </summary>
-		private readonly ILogger<ConfigController> Logger;
+		readonly IOptionsMonitor<ServerSettings> _settings;
 
 		/// <summary>
 		/// Constructor
 		/// </summary>
-		/// <param name="DatabaseService">The database service singleton</param>
-		/// <param name="AclService">The ACL service singleton</param>
-		/// <param name="ConfigService">The Config service singleton</param>
-		/// <param name="Logger">Logger for the controller</param>
-		/// <param name="Settings">Server settings</param>
-		public ConfigController(DatabaseService DatabaseService, AclService AclService, ConfigService ConfigService, IOptionsMonitor<ServerSettings> Settings, ILogger<ConfigController> Logger)
+		/// <param name="aclService">The ACL service singleton</param>
+		/// <param name="configService">The Config service singleton</param>
+		/// <param name="settings">Server settings</param>
+		public ConfigController(AclService aclService, ConfigService configService, IOptionsMonitor<ServerSettings> settings)
 		{
-			this.DatabaseService = DatabaseService;
-			this.AclService = AclService;
-			this.ConfigService = ConfigService;
-			this.Settings = Settings;
-			this.Logger = Logger;			
+			_aclService = aclService;
+			_configService = configService;
+			_settings = settings;
 		}
 
 		// Global Configuration
@@ -73,19 +59,19 @@ namespace Horde.Build.Controllers
 		/// </summary>
 		[HttpPut]
 		[Route("/api/v1/config/global")]
-		public async Task<ActionResult> UpdateGlobalSettings([FromBody] UpdateGlobalConfigRequest Request)
+		public async Task<ActionResult> UpdateGlobalSettings([FromBody] UpdateGlobalConfigRequest request)
 		{
-			if (!await AclService.AuthorizeAsync(AclAction.AdminWrite, User))
+			if (!await _aclService.AuthorizeAsync(AclAction.AdminWrite, User))
 			{
 				return Forbid(AclAction.AdminWrite);
 			}
 
-			if (!Settings.CurrentValue.SingleInstance)
+			if (!_settings.CurrentValue.SingleInstance)
 			{
 				return BadRequest("Updating global configuration settings with ConfigController currently only supported in single instance mode");
 			}
 
-			await ConfigService.UpdateGlobalConfig(Request);
+			await _configService.UpdateGlobalConfig(request);
 
 			return Ok();
 		}
@@ -98,23 +84,23 @@ namespace Horde.Build.Controllers
 		[Route("/api/v1/config/global")]
 		public async Task<ActionResult<GlobalConfig?>> GetGlobalConfig()
 		{
-			if (!await AclService.AuthorizeAsync(AclAction.AdminRead, User))
+			if (!await _aclService.AuthorizeAsync(AclAction.AdminRead, User))
 			{
 				return Forbid(AclAction.AdminRead);
 			}
 
-			if (string.IsNullOrEmpty(Settings.CurrentValue.ConfigPath))
+			if (String.IsNullOrEmpty(_settings.CurrentValue.ConfigPath))
 			{
 				return BadRequest("Missing config path for settings");
 			}
 
-			FileReference GlobalConfigFile = new FileReference(Settings.CurrentValue.ConfigPath);
+			FileReference globalConfigFile = new FileReference(_settings.CurrentValue.ConfigPath);
 
-			byte[] Data = await FileReference.ReadAllBytesAsync(GlobalConfigFile);
-			JsonSerializerOptions Options = new JsonSerializerOptions();
-			Startup.ConfigureJsonSerializer(Options);
+			byte[] data = await FileReference.ReadAllBytesAsync(globalConfigFile);
+			JsonSerializerOptions options = new JsonSerializerOptions();
+			Startup.ConfigureJsonSerializer(options);
 
-			return JsonSerializer.Deserialize<GlobalConfig>(Data, Options);
+			return JsonSerializer.Deserialize<GlobalConfig>(data, options);
 		}
 
 		// Administrative Server Settings
@@ -127,15 +113,15 @@ namespace Horde.Build.Controllers
 		[Route("/api/v1/config/serversettings")]
 		public async Task<ActionResult<GetServerSettingsResponse>> GetServerSettings()
 		{
-			if (!await AclService.AuthorizeAsync(AclAction.AdminRead, User))
+			if (!await _aclService.AuthorizeAsync(AclAction.AdminRead, User))
 			{
 				return Forbid(AclAction.AdminRead);
 			}
 
-			GetServerSettingsResponse Response = new GetServerSettingsResponse(Settings.CurrentValue);
-			Response.NumServerUpdates = ConfigService.NumUserConfigUpdates;
+			GetServerSettingsResponse response = new GetServerSettingsResponse(_settings.CurrentValue);
+			response.NumServerUpdates = _configService.NumUserConfigUpdates;
 
-			return Response;
+			return response;
 		}
 
 		/// <summary>
@@ -144,28 +130,26 @@ namespace Horde.Build.Controllers
 		/// <returns>Administrative settings for the server</returns>
 		[HttpPut]
 		[Route("/api/v1/config/serversettings")]
-		public async Task<ActionResult<ServerUpdateResponse>> UpdateServerSettings([FromBody] UpdateServerSettingsRequest Request)
+		public async Task<ActionResult<ServerUpdateResponse>> UpdateServerSettings([FromBody] UpdateServerSettingsRequest request)
 		{
-			if (!await AclService.AuthorizeAsync(AclAction.AdminWrite, User))
+			if (!await _aclService.AuthorizeAsync(AclAction.AdminWrite, User))
 			{
 				return Forbid(AclAction.AdminWrite);
 			}
 
-			if (!Settings.CurrentValue.SingleInstance)
+			if (!_settings.CurrentValue.SingleInstance)
 			{
 				return BadRequest("Updating server settings from ConfigController currently only supported in single instance mode");
 			}
 
-			if (Request.Settings == null || Request.Settings.Count == 0)
+			if (request.Settings == null || request.Settings.Count == 0)
 			{
 				return new ServerUpdateResponse() { RestartRequired = false };
 			}
 
-			ServerUpdateResponse Response = await ConfigService.UpdateServerSettings(Request);
+			ServerUpdateResponse response = await _configService.UpdateServerSettings(request);
 
-			return Response;
+			return response;
 		}
-
 	}
-
 }
diff --git a/Engine/Source/Programs/Horde/Horde.Build/Controllers/CredentialsController.cs b/Engine/Source/Programs/Horde/Horde.Build/Controllers/CredentialsController.cs
--- a/Engine/Source/Programs/Horde/Horde.Build/Controllers/CredentialsController.cs
+++ b/Engine/Source/Programs/Horde/Horde.Build/Controllers/CredentialsController.cs
@@ -1,5 +1,7 @@
 ﻿// Copyright Epic Games, Inc. All Rights Reserved.
 
+using System.Collections.Generic;
+using System.Threading.Tasks;
 using Horde.Build.Acls;
 using Horde.Build.Api;
 using Horde.Build.Models;
@@ -8,11 +10,6 @@ using Horde.Build.Utilities;
 using Microsoft.AspNetCore.Authorization;
 using Microsoft.AspNetCore.Mvc;
 using MongoDB.Bson;
-using System;
-using System.Collections.Generic;
-using System.ComponentModel.DataAnnotations;
-using System.Linq;
-using System.Threading.Tasks;
 
 namespace Horde.Build.Controllers
 {
@@ -27,155 +24,155 @@ namespace Horde.Build.Controllers
 		/// <summary>
 		/// Singleton instance of the ACL service
 		/// </summary>
-		private readonly AclService AclService;
+		private readonly AclService _aclService;
 
 		/// <summary>
 		/// Singleton instance of the credential service
 		/// </summary>
-		private readonly CredentialService CredentialService;
+		private readonly CredentialService _credentialService;
 
 		/// <summary>
 		/// Constructor
 		/// </summary>
-		/// <param name="AclService">The ACL service</param>
-		/// <param name="CredentialService">The credential service</param>
-		public CredentialsController(AclService AclService, CredentialService CredentialService)
+		/// <param name="aclService">The ACL service</param>
+		/// <param name="credentialService">The credential service</param>
+		public CredentialsController(AclService aclService, CredentialService credentialService)
 		{
-			this.AclService = AclService;
-			this.CredentialService = CredentialService;
+			_aclService = aclService;
+			_credentialService = credentialService;
 		}
 
 		/// <summary>
 		/// Creates a new credential
 		/// </summary>
-		/// <param name="Create">Parameters for the new credential.</param>
+		/// <param name="create">Parameters for the new credential.</param>
 		/// <returns>Http result code</returns>
 		[HttpPost]
 		[Route("/api/v1/credentials")]
-		public async Task<ActionResult<CreateCredentialResponse>> CreateCredentialAsync([FromBody] CreateCredentialRequest Create)
+		public async Task<ActionResult<CreateCredentialResponse>> CreateCredentialAsync([FromBody] CreateCredentialRequest create)
 		{
-			if(!await AclService.AuthorizeAsync(AclAction.CreateCredential, User))
+			if(!await _aclService.AuthorizeAsync(AclAction.CreateCredential, User))
 			{
 				return Forbid();
 			}
 
-			Credential NewCredential = await CredentialService.CreateCredentialAsync(Create.Name, Create.Properties);
-			return new CreateCredentialResponse(NewCredential.Id.ToString());
+			Credential newCredential = await _credentialService.CreateCredentialAsync(create.Name, create.Properties);
+			return new CreateCredentialResponse(newCredential.Id.ToString());
 		}
 
 		/// <summary>
 		/// Query all the credentials
 		/// </summary>
-		/// <param name="Name">Id of the credential to get information about</param>
-		/// <param name="Filter">Filter for the properties to return</param>
+		/// <param name="name">Id of the credential to get information about</param>
+		/// <param name="filter">Filter for the properties to return</param>
 		/// <returns>Information about all the credentials</returns>
 		[HttpGet]
 		[Route("/api/v1/credentials")]
 		[ProducesResponseType(typeof(List<GetCredentialResponse>), 200)]
-		public async Task<ActionResult<object>> FindCredentialAsync([FromQuery] string? Name = null, [FromQuery] PropertyFilter? Filter = null)
+		public async Task<ActionResult<object>> FindCredentialAsync([FromQuery] string? name = null, [FromQuery] PropertyFilter? filter = null)
 		{
-			if (!await AclService.AuthorizeAsync(AclAction.ListCredentials, User))
+			if (!await _aclService.AuthorizeAsync(AclAction.ListCredentials, User))
 			{
 				return Forbid();
 			}
 
-			List<Credential> Credentials = await CredentialService.FindCredentialsAsync(Name);
-			GlobalPermissionsCache Cache = new GlobalPermissionsCache();
+			List<Credential> credentials = await _credentialService.FindCredentialsAsync(name);
+			GlobalPermissionsCache cache = new GlobalPermissionsCache();
 
-			List<object> Responses = new List<object>();
-			foreach (Credential Credential in Credentials)
+			List<object> responses = new List<object>();
+			foreach (Credential credential in credentials)
 			{
-				if (await CredentialService.AuthorizeAsync(Credential, AclAction.ViewCredential, User, Cache))
+				if (await _credentialService.AuthorizeAsync(credential, AclAction.ViewCredential, User, cache))
 				{
-					bool bIncludeAcl = await CredentialService.AuthorizeAsync(Credential, AclAction.ViewPermissions, User, Cache);
-					Responses.Add(new GetCredentialResponse(Credential, bIncludeAcl).ApplyFilter(Filter));
+					bool bIncludeAcl = await _credentialService.AuthorizeAsync(credential, AclAction.ViewPermissions, User, cache);
+					responses.Add(new GetCredentialResponse(credential, bIncludeAcl).ApplyFilter(filter));
 				}
 			}
-			return Responses;
+			return responses;
 		}
 
 		/// <summary>
 		/// Retrieve information about a specific credential
 		/// </summary>
-		/// <param name="CredentialId">Id of the credential to get information about</param>
-		/// <param name="Filter">Filter for properties to return</param>
+		/// <param name="credentialId">Id of the credential to get information about</param>
+		/// <param name="filter">Filter for properties to return</param>
 		/// <returns>Information about the requested credential</returns>
 		[HttpGet]
-		[Route("/api/v1/credentials/{CredentialId}")]
+		[Route("/api/v1/credentials/{credentialId}")]
 		[ProducesResponseType(typeof(GetCredentialResponse), 200)]
-		public async Task<ActionResult<object>> GetCredentialAsync(string CredentialId, [FromQuery] PropertyFilter? Filter = null)
+		public async Task<ActionResult<object>> GetCredentialAsync(string credentialId, [FromQuery] PropertyFilter? filter = null)
 		{
-			ObjectId ProjectIdValue = CredentialId.ToObjectId();
+			ObjectId projectIdValue = credentialId.ToObjectId();
 
-			Credential? Credential = await CredentialService.GetCredentialAsync(ProjectIdValue);
-			if (Credential == null)
+			Credential? credential = await _credentialService.GetCredentialAsync(projectIdValue);
+			if (credential == null)
 			{
 				return NotFound();
 			}
 
-			GlobalPermissionsCache Cache = new GlobalPermissionsCache();
-			if (!await CredentialService.AuthorizeAsync(Credential, AclAction.ViewCredential, User, Cache))
+			GlobalPermissionsCache cache = new GlobalPermissionsCache();
+			if (!await _credentialService.AuthorizeAsync(credential, AclAction.ViewCredential, User, cache))
 			{
 				return Forbid();
 			}
 
-			bool bIncludeAcl = await CredentialService.AuthorizeAsync(Credential, AclAction.ViewPermissions, User, Cache);
-			return new GetCredentialResponse(Credential, bIncludeAcl).ApplyFilter(Filter);
+			bool bIncludeAcl = await _credentialService.AuthorizeAsync(credential, AclAction.ViewPermissions, User, cache);
+			return new GetCredentialResponse(credential, bIncludeAcl).ApplyFilter(filter);
 		}
 
 		/// <summary>
 		/// Update a credential's properties.
 		/// </summary>
-		/// <param name="CredentialId">Id of the credential to update</param>
-		/// <param name="Update">Items on the credential to update</param>
+		/// <param name="credentialId">Id of the credential to update</param>
+		/// <param name="update">Items on the credential to update</param>
 		/// <returns>Http result code</returns>
 		[HttpPut]
-		[Route("/api/v1/credentials/{CredentialId}")]
-		public async Task<ActionResult> UpdateCredentialAsync(string CredentialId, [FromBody] UpdateCredentialRequest Update)
+		[Route("/api/v1/credentials/{credentialId}")]
+		public async Task<ActionResult> UpdateCredentialAsync(string credentialId, [FromBody] UpdateCredentialRequest update)
 		{
-			ObjectId CredentialIdValue = CredentialId.ToObjectId();
+			ObjectId credentialIdValue = credentialId.ToObjectId();
 
-			Credential? Credential = await CredentialService.GetCredentialAsync(CredentialIdValue);
-			if(Credential == null)
+			Credential? credential = await _credentialService.GetCredentialAsync(credentialIdValue);
+			if(credential == null)
 			{
 				return NotFound();
 			}
 
-			GlobalPermissionsCache Cache = new GlobalPermissionsCache();
-			if (!await CredentialService.AuthorizeAsync(Credential, AclAction.UpdateCredential, User, Cache))
+			GlobalPermissionsCache cache = new GlobalPermissionsCache();
+			if (!await _credentialService.AuthorizeAsync(credential, AclAction.UpdateCredential, User, cache))
 			{
 				return Forbid();
 			}
-			if (Update.Acl != null && !await CredentialService.AuthorizeAsync(Credential, AclAction.ChangePermissions, User, Cache))
+			if (update.Acl != null && !await _credentialService.AuthorizeAsync(credential, AclAction.ChangePermissions, User, cache))
 			{
 				return Forbid();
 			}
 
-			await CredentialService.UpdateCredentialAsync(CredentialIdValue, Update.Name, Update.Properties, Acl.Merge(Credential.Acl, Update.Acl));
+			await _credentialService.UpdateCredentialAsync(credentialIdValue, update.Name, update.Properties, Acl.Merge(credential.Acl, update.Acl));
 			return new OkResult();
 		}
 
 		/// <summary>
 		/// Delete a credential
 		/// </summary>
-		/// <param name="CredentialId">Id of the credential to delete</param>
+		/// <param name="credentialId">Id of the credential to delete</param>
 		/// <returns>Http result code</returns>
 		[HttpDelete]
-		[Route("/api/v1/credentials/{CredentialId}")]
-		public async Task<ActionResult> DeleteCredentialAsync(string CredentialId)
+		[Route("/api/v1/credentials/{credentialId}")]
+		public async Task<ActionResult> DeleteCredentialAsync(string credentialId)
 		{
-			ObjectId CredentialIdValue = CredentialId.ToObjectId();
+			ObjectId credentialIdValue = credentialId.ToObjectId();
 
-			Credential? Credential = await CredentialService.GetCredentialAsync(CredentialIdValue);
-			if (Credential == null)
+			Credential? credential = await _credentialService.GetCredentialAsync(credentialIdValue);
+			if (credential == null)
 			{
 				return NotFound();
 			}
-			if (!await CredentialService.AuthorizeAsync(Credential, AclAction.DeleteCredential, User, null))
+			if (!await _credentialService.AuthorizeAsync(credential, AclAction.DeleteCredential, User, null))
 			{
 				return Forbid();
 			}
-			if (!await CredentialService.DeleteCredentialAsync(CredentialIdValue))
+			if (!await _credentialService.DeleteCredentialAsync(credentialIdValue))
 			{
 				return NotFound();
 			}
diff --git a/Engine/Source/Programs/Horde/Horde.Build/Controllers/DashboardController.cs b/Engine/Source/Programs/Horde/Horde.Build/Controllers/DashboardController.cs
--- a/Engine/Source/Programs/Horde/Horde.Build/Controllers/DashboardController.cs
+++ b/Engine/Source/Programs/Horde/Horde.Build/Controllers/DashboardController.cs
@@ -2,17 +2,15 @@
 
 // Copyright Epic Games, Inc. All Rights Reserved.	
 
-using Horde.Build.Authentication;
+using System;
+using System.Text;
+using System.Threading.Tasks;
+using Horde.Build.Controllers;
 using Microsoft.AspNetCore.Authentication;
 using Microsoft.AspNetCore.Authentication.Cookies;
 using Microsoft.AspNetCore.Authorization;
 using Microsoft.AspNetCore.Mvc;
-using System.Threading.Tasks;
-using Horde.Build;
-using Horde.Build.Controllers;
 using Microsoft.Extensions.Options;
-using System;
-using System.Text;
 
 namespace Horde.Build
 {
@@ -26,15 +24,15 @@ namespace Horde.Build
 		/// <summary>
 		/// Authentication scheme in use
 		/// </summary>
-		string AuthenticationScheme;
+		readonly string _authenticationScheme;
 
 		/// <summary>
 		/// Constructor
 		/// </summary>
-		/// <param name="ServerSettings">Server settings</param>
-		public DashboardChallengeController(IOptionsMonitor<ServerSettings> ServerSettings)
+		/// <param name="serverSettings">Server settings</param>
+		public DashboardChallengeController(IOptionsMonitor<ServerSettings> serverSettings)
 		{
-			AuthenticationScheme = AccountController.GetAuthScheme(ServerSettings.CurrentValue.AuthMethod);
+			_authenticationScheme = AccountController.GetAuthScheme(serverSettings.CurrentValue.AuthMethod);
 		}
 
 		/// <summary>	
@@ -52,36 +50,35 @@ namespace Horde.Build
 		/// <summary>
 		/// Login to server, redirecting to the specified URL on success
 		/// </summary>
-		/// <param name="Redirect"></param>
+		/// <param name="redirect"></param>
 		/// <returns></returns>
 		[HttpGet]
 		[Route("/api/v1/dashboard/login")]
-		public IActionResult Login([FromQuery] string? Redirect)
+		public IActionResult Login([FromQuery] string? redirect)
 		{
-			return new ChallengeResult(AuthenticationScheme, new AuthenticationProperties { RedirectUri = Redirect ?? "/" });
+			return new ChallengeResult(_authenticationScheme, new AuthenticationProperties { RedirectUri = redirect ?? "/" });
 		}
 
 		/// <summary>
 		/// Login to server, redirecting to the specified Base64 encoded URL, which fixes some escaping issues on some auth providers, on success
 		/// </summary>
-		/// <param name="Redirect"></param>
+		/// <param name="redirect"></param>
 		/// <returns></returns>
 		[HttpGet]
 		[Route("/api/v2/dashboard/login")]
-		public IActionResult LoginV2([FromQuery] string? Redirect)
+		public IActionResult LoginV2([FromQuery] string? redirect)
 		{
-			string? RedirectUri = null;
+			string? redirectUri = null;
 
-			if (Redirect != null)
+			if (redirect != null)
 			{
-				byte[] Data = Convert.FromBase64String(Redirect);
-				RedirectUri = Encoding.UTF8.GetString(Data);
+				byte[] data = Convert.FromBase64String(redirect);
+				redirectUri = Encoding.UTF8.GetString(data);
 			}
 
-			return new ChallengeResult(AuthenticationScheme, new AuthenticationProperties { RedirectUri = RedirectUri ?? "/index" });
+			return new ChallengeResult(_authenticationScheme, new AuthenticationProperties { RedirectUri = redirectUri ?? "/index" });
 		}
 
-
 		/// <summary>
 		/// Logout of the current account
 		/// </summary>
@@ -93,17 +90,13 @@ namespace Horde.Build
 			await HttpContext.SignOutAsync(CookieAuthenticationDefaults.AuthenticationScheme);
 			try
 			{
-				await HttpContext.SignOutAsync(AuthenticationScheme);
+				await HttpContext.SignOutAsync(_authenticationScheme);
 			}
-#pragma warning disable CA1031 // Do not catch general exception types
 			catch
-#pragma warning restore CA1031
 			{
 			}
 
 			return Ok();
 		}
-
-
 	}
 }
diff --git a/Engine/Source/Programs/Horde/Horde.Build/Controllers/DebugController.cs b/Engine/Source/Programs/Horde/Horde.Build/Controllers/DebugController.cs
--- a/Engine/Source/Programs/Horde/Horde.Build/Controllers/DebugController.cs
+++ b/Engine/Source/Programs/Horde/Horde.Build/Controllers/DebugController.cs
@@ -3,55 +3,14 @@
 //#define ENABLE_PUBLIC_DEBUG_CONTROLLER
 //#define ENABLE_SECURE_DEBUG_CONTROLLER
 
-using EpicGames.Core;
-using Horde.Build.Api;
-using HordeCommon;
+#if ENABLE_PUBLIC_DEBUG_CONTROLLER
+
 using Horde.Build.Collections;
 using Horde.Build.Models;
-using Horde.Build.Services;
-using Horde.Build.Tasks.Impl;
 using Horde.Build.Utilities;
-using Microsoft.AspNetCore.Authorization;
-using Microsoft.AspNetCore.Mvc;
-using Microsoft.Extensions.Primitives;
-using MongoDB.Bson;
-using System;
-using System.Collections.Generic;
-using System.Diagnostics;
-using System.Globalization;
-using System.IO;
-using System.Linq;
-using System.Net;
-using System.Runtime;
-using System.Runtime.InteropServices;
-using System.Text;
-using System.Threading.Tasks;
-using System.Web;
-
-using MongoDB.Driver;
-using System.Text.Json;
-using System.Threading;
-using Horde.Build.Services.Impl;
-using Microsoft.Net.Http.Headers;
-using Horde.Build.Storage;
-using Microsoft.Extensions.Hosting;
-using Microsoft.Extensions.Options;
-using StatsdClient;
-
-using Amazon.S3;
-using JetBrains.Profiler.SelfApi;
 
 namespace Horde.Build.Controllers
 {
-	using LogId = ObjectId<ILogFile>;
-	using PoolId = StringId<IPool>;
-	using ProjectId = StringId<IProject>;
-	using TemplateRefId = StringId<TemplateRef>;
-	using AgentSoftwareVersion = StringId<IAgentSoftwareCollection>;
-	using IStream = Horde.Build.Models.IStream;
-	using P4Debugging = Perforce.P4.P4Debugging;
-
-#if ENABLE_PUBLIC_DEBUG_CONTROLLER
 	/// <summary>
 	/// Public endpoints for the debug controller
 	/// </summary>
@@ -75,9 +34,9 @@ namespace Horde.Build.Controllers
 		/// <param name="DogStatsd"></param>
 		public PublicDebugController(RequestTrackerService RequestTrackerService, IHostApplicationLifetime ApplicationLifetime, IDogStatsd DogStatsd)
 		{
-			this.RequestTrackerService = RequestTrackerService;
-			this.ApplicationLifetime = ApplicationLifetime;
-			this.DogStatsd = DogStatsd;
+			RequestTrackerService = RequestTrackerService;
+			ApplicationLifetime = ApplicationLifetime;
+			DogStatsd = DogStatsd;
 		}
 
 		/// <summary>
@@ -365,22 +324,22 @@ namespace Horde.Build.Controllers
 		/// <param name="Settings">Settings</param>
 		public SecureDebugController(AclService AclService, DatabaseService DatabaseService, JobTaskSource JobTaskSource, ISingletonDocument<Globals> GlobalsSingleton, IPoolCollection PoolCollection, IProjectCollection ProjectCollection, IAgentCollection AgentCollection, ISessionCollection SessionCollection, ILeaseCollection LeaseCollection, ITemplateCollection TemplateCollection, IStreamCollection StreamCollection, IGraphCollection GraphCollection, ILogFileCollection LogFileCollection, IPerforceService Perforce, IFleetManager FleetManager, IOptionsMonitor<ServerSettings> Settings)
 		{
-			this.AclService = AclService;
-			this.DatabaseService = DatabaseService;
-			this.JobTaskSource = JobTaskSource;
-			this.GlobalsSingleton = GlobalsSingleton;
-			this.PoolCollection = PoolCollection;
-			this.ProjectCollection = ProjectCollection;
-			this.AgentCollection = AgentCollection;
-			this.SessionCollection = SessionCollection;
-			this.LeaseCollection = LeaseCollection;
-			this.TemplateCollection = TemplateCollection;
-			this.StreamCollection = StreamCollection;
-			this.GraphCollection = GraphCollection;
-			this.LogFileCollection = LogFileCollection;
-			this.Perforce = Perforce;
-			this.FleetManager = FleetManager;
-			this.Settings = Settings;
+			AclService = AclService;
+			DatabaseService = DatabaseService;
+			JobTaskSource = JobTaskSource;
+			GlobalsSingleton = GlobalsSingleton;
+			PoolCollection = PoolCollection;
+			ProjectCollection = ProjectCollection;
+			AgentCollection = AgentCollection;
+			SessionCollection = SessionCollection;
+			LeaseCollection = LeaseCollection;
+			TemplateCollection = TemplateCollection;
+			StreamCollection = StreamCollection;
+			GraphCollection = GraphCollection;
+			LogFileCollection = LogFileCollection;
+			Perforce = Perforce;
+			FleetManager = FleetManager;
+			Settings = Settings;
 		}
 
 		/// <summary>
@@ -931,5 +890,6 @@ namespace Horde.Build.Controllers
 			return PhysicalFile(SnapshotZipFile, "application/zip", Path.GetFileName(SnapshotZipFile));
 		}
 	}
-#endif
 }
+
+#endif
