Horde: Fix coding conventions within Horde.Build.

diff --git a/Engine/Source/Programs/Horde/Horde.Build/Controllers/DevicesController.cs b/Engine/Source/Programs/Horde/Horde.Build/Controllers/DevicesController.cs
--- a/Engine/Source/Programs/Horde/Horde.Build/Controllers/DevicesController.cs
+++ b/Engine/Source/Programs/Horde/Horde.Build/Controllers/DevicesController.cs
@@ -1,5 +1,10 @@
-﻿// Copyright Epic Games, Inc. All Rights Reserved.
+// Copyright Epic Games, Inc. All Rights Reserved.
 
+using System;
+using System.Collections.Generic;
+using System.Linq;
+using System.Text.Json;
+using System.Threading.Tasks;
 using Horde.Build.Acls;
 using Horde.Build.Api;
 using Horde.Build.Collections;
@@ -11,11 +16,6 @@ using Microsoft.AspNetCore.Mvc;
 using Microsoft.Extensions.Logging;
 using MongoDB.Bson;
 using MongoDB.Driver;
-using System;
-using System.Collections.Generic;
-using System.Linq;
-using System.Text.Json;
-using System.Threading.Tasks;
 
 namespace Horde.Build.Controllers
 {
@@ -33,33 +33,32 @@ namespace Horde.Build.Controllers
 		/// <summary>
 		/// The acl service singleton
 		/// </summary>
-		AclService AclService;
+		readonly AclService _aclService;
 
 		/// <summary>
 		/// The user collection instance
 		/// </summary>
 		IUserCollection UserCollection { get; set; }
 
-
 		/// <summary>
 		/// Singleton instance of the device service
 		/// </summary>
-		DeviceService DeviceService;
+		readonly DeviceService _deviceService;
 
 		/// <summary>
 		///  Logger for controller
 		/// </summary>
-		private readonly ILogger<DevicesController> Logger;
+		private readonly ILogger<DevicesController> _logger;
 
 		/// <summary>
 		/// Constructor
 		/// </summary>
-		public DevicesController(DeviceService DeviceService, AclService AclService, IUserCollection UserCollection, ILogger<DevicesController> Logger)
+		public DevicesController(DeviceService deviceService, AclService aclService, IUserCollection userCollection, ILogger<DevicesController> logger)
 		{
-			this.UserCollection = UserCollection;
-			this.DeviceService = DeviceService;
-			this.Logger = Logger;
-			this.AclService = AclService;
+			UserCollection = userCollection;
+			_deviceService = deviceService;
+			_logger = logger;
+			_aclService = aclService;
 		}
 
 		// DEVICES
@@ -70,60 +69,59 @@ namespace Horde.Build.Controllers
 		[HttpPost]
 		[Authorize]
 		[Route("/api/v2/devices")]
-		public async Task<ActionResult<CreateDeviceResponse>> CreateDeviceAsync([FromBody] CreateDeviceRequest DeviceRequest)
+		public async Task<ActionResult<CreateDeviceResponse>> CreateDeviceAsync([FromBody] CreateDeviceRequest deviceRequest)
 		{
 
-			DevicePoolAuthorization? PoolAuth = await DeviceService.GetUserPoolAuthorizationAsync(new DevicePoolId(DeviceRequest.PoolId!), User);
+			DevicePoolAuthorization? poolAuth = await _deviceService.GetUserPoolAuthorizationAsync(new DevicePoolId(deviceRequest.PoolId!), User);
 
-			if (PoolAuth == null || !PoolAuth.Write)
+			if (poolAuth == null || !poolAuth.Write)
 			{
 				return Forbid();
 			}
 
-			IUser? InternalUser = await UserCollection.GetUserAsync(User);
-			if (InternalUser == null)
+			IUser? internalUser = await UserCollection.GetUserAsync(User);
+			if (internalUser == null)
 			{
 				return NotFound();
 			}
 
-			IDevicePlatform? Platform = await DeviceService.GetPlatformAsync(new DevicePlatformId(DeviceRequest.PlatformId!));
+			IDevicePlatform? platform = await _deviceService.GetPlatformAsync(new DevicePlatformId(deviceRequest.PlatformId!));
 
-			if (Platform == null)
+			if (platform == null)
 			{
-				return BadRequest($"Bad platform id {DeviceRequest.PlatformId} on request");
+				return BadRequest($"Bad platform id {deviceRequest.PlatformId} on request");
 			}
 
-			IDevicePool? Pool = await DeviceService.GetPoolAsync(new DevicePoolId(DeviceRequest.PoolId!));
+			IDevicePool? pool = await _deviceService.GetPoolAsync(new DevicePoolId(deviceRequest.PoolId!));
 
-			if (Pool == null)
+			if (pool == null)
 			{
-				return BadRequest($"Bad pool id {DeviceRequest.PoolId} on request");
+				return BadRequest($"Bad pool id {deviceRequest.PoolId} on request");
 			}
 
+			string? modelId = null;
 
-			string? ModelId = null;
-
-			if (!string.IsNullOrEmpty(DeviceRequest.ModelId))
+			if (!String.IsNullOrEmpty(deviceRequest.ModelId))
 			{
-				ModelId = new string(DeviceRequest.ModelId);
+				modelId = new string(deviceRequest.ModelId);
 
-				if (Platform.Models?.FirstOrDefault(x => x == ModelId) == null)
+				if (platform.Models?.FirstOrDefault(x => x == modelId) == null)
 				{
-					return BadRequest($"Bad model id {ModelId} for platform {Platform.Id} on request");
+					return BadRequest($"Bad model id {modelId} for platform {platform.Id} on request");
 				}
 			}
 
-			string Name = DeviceRequest.Name.Trim();
-			string? Address = DeviceRequest.Address?.Trim();
+			string name = deviceRequest.Name.Trim();
+			string? address = deviceRequest.Address?.Trim();
 
-			IDevice? Device = await DeviceService.TryCreateDeviceAsync(DeviceId.Sanitize(Name), Name, Platform.Id, Pool.Id, DeviceRequest.Enabled, Address, ModelId, InternalUser.Id);
+			IDevice? device = await _deviceService.TryCreateDeviceAsync(DeviceId.Sanitize(name), name, platform.Id, pool.Id, deviceRequest.Enabled, address, modelId, internalUser.Id);
 
-			if (Device == null)
+			if (device == null)
 			{
 				return BadRequest($"Unable to create device");
 			}
 
-			return new CreateDeviceResponse(Device.Id.ToString());
+			return new CreateDeviceResponse(device.Id.ToString());
 		}
 
 		/// <summary>
@@ -136,25 +134,25 @@ namespace Horde.Build.Controllers
 		public async Task<ActionResult<List<object>>> GetDevicesAsync()
 		{
 
-			List<DevicePoolAuthorization> PoolAuth = await DeviceService.GetUserPoolAuthorizationsAsync(User);
+			List<DevicePoolAuthorization> poolAuth = await _deviceService.GetUserPoolAuthorizationsAsync(User);
 
-			List<IDevice> Devices = await DeviceService.GetDevicesAsync();
+			List<IDevice> devices = await _deviceService.GetDevicesAsync();
 
-			List<object> Responses = new List<object>();
+			List<object> responses = new List<object>();
 
-			foreach (IDevice Device in Devices)
+			foreach (IDevice device in devices)
 			{
-				DevicePoolAuthorization? Auth = PoolAuth.Where(x => x.Pool.Id == Device.PoolId).FirstOrDefault();
+				DevicePoolAuthorization? auth = poolAuth.Where(x => x.Pool.Id == device.PoolId).FirstOrDefault();
 
-				if (Auth == null || !Auth.Read)
+				if (auth == null || !auth.Read)
 				{
 					continue;
 				}
 
-				Responses.Add(new GetDeviceResponse(Device.Id.ToString(), Device.PlatformId.ToString(), Device.PoolId.ToString(), Device.Name, Device.Enabled, Device.Address, Device.ModelId?.ToString(), Device.ModifiedByUser, Device.Notes, Device.ProblemTimeUtc, Device.MaintenanceTimeUtc, Device.Utilization, Device.CheckedOutByUser, Device.CheckOutTime));
+				responses.Add(new GetDeviceResponse(device.Id.ToString(), device.PlatformId.ToString(), device.PoolId.ToString(), device.Name, device.Enabled, device.Address, device.ModelId?.ToString(), device.ModifiedByUser, device.Notes, device.ProblemTimeUtc, device.MaintenanceTimeUtc, device.Utilization, device.CheckedOutByUser, device.CheckOutTime));
 			}
 
-			return Responses;
+			return responses;
 		}
 
 		/// <summary>
@@ -162,29 +160,28 @@ namespace Horde.Build.Controllers
 		/// </summary>
 		[HttpGet]
 		[Authorize]
-		[Route("/api/v2/devices/{DeviceId}")]
+		[Route("/api/v2/devices/{deviceId}")]
 		[ProducesResponseType(typeof(GetDeviceResponse), 200)]
-		public async Task<ActionResult<object>> GetDeviceAsync(string DeviceId)
+		public async Task<ActionResult<object>> GetDeviceAsync(string deviceId)
 		{
 
-			DeviceId DeviceIdValue = new DeviceId(DeviceId);
+			DeviceId deviceIdValue = new DeviceId(deviceId);
 
-			IDevice? Device = await DeviceService.GetDeviceAsync(DeviceIdValue);
+			IDevice? device = await _deviceService.GetDeviceAsync(deviceIdValue);
 
-			if (Device == null)
+			if (device == null)
 			{
-				return BadRequest($"Unable to find device with id {DeviceId}");
+				return BadRequest($"Unable to find device with id {deviceId}");
 			}
 
-			DevicePoolAuthorization? PoolAuth = await DeviceService.GetUserPoolAuthorizationAsync(Device.PoolId, User);
+			DevicePoolAuthorization? poolAuth = await _deviceService.GetUserPoolAuthorizationAsync(device.PoolId, User);
 
-			if (PoolAuth == null || !PoolAuth.Write)
+			if (poolAuth == null || !poolAuth.Write)
 			{
 				return Forbid();
 			}
 
-
-			return new GetDeviceResponse(Device.Id.ToString(), Device.PlatformId.ToString(), Device.PoolId.ToString(), Device.Name, Device.Enabled, Device.Address, Device.ModelId?.ToString(), Device.ModifiedByUser?.ToString(), Device.Notes, Device.ProblemTimeUtc, Device.MaintenanceTimeUtc, Device.Utilization, Device.CheckedOutByUser, Device.CheckOutTime);
+			return new GetDeviceResponse(device.Id.ToString(), device.PlatformId.ToString(), device.PoolId.ToString(), device.Name, device.Enabled, device.Address, device.ModelId?.ToString(), device.ModifiedByUser?.ToString(), device.Notes, device.ProblemTimeUtc, device.MaintenanceTimeUtc, device.Utilization, device.CheckedOutByUser, device.CheckOutTime);
 		}
 
 		/// <summary>
@@ -192,63 +189,62 @@ namespace Horde.Build.Controllers
 		/// </summary>
 		[HttpPut]
 		[Authorize]
-		[Route("/api/v2/devices/{DeviceId}")]
+		[Route("/api/v2/devices/{deviceId}")]
 		[ProducesResponseType(typeof(List<GetDeviceResponse>), 200)]
-		public async Task<ActionResult> UpdateDeviceAsync(string DeviceId, [FromBody] UpdateDeviceRequest Update)
+		public async Task<ActionResult> UpdateDeviceAsync(string deviceId, [FromBody] UpdateDeviceRequest update)
 		{
-			IUser? InternalUser = await UserCollection.GetUserAsync(User);
-			if (InternalUser == null)
+			IUser? internalUser = await UserCollection.GetUserAsync(User);
+			if (internalUser == null)
 			{
 				return NotFound();
 			}
 
-			DeviceId DeviceIdValue = new DeviceId(DeviceId);
+			DeviceId deviceIdValue = new DeviceId(deviceId);
 
-			IDevice? Device = await DeviceService.GetDeviceAsync(DeviceIdValue);
+			IDevice? device = await _deviceService.GetDeviceAsync(deviceIdValue);
 
-			if (Device == null)
+			if (device == null)
 			{
-				return BadRequest($"Device with id ${DeviceId} does not exist");
+				return BadRequest($"Device with id ${deviceId} does not exist");
 			}
 
-			DevicePoolAuthorization? PoolAuth = await DeviceService.GetUserPoolAuthorizationAsync(Device.PoolId, User);
+			DevicePoolAuthorization? poolAuth = await _deviceService.GetUserPoolAuthorizationAsync(device.PoolId, User);
 
-			if (PoolAuth == null || !PoolAuth.Write)
+			if (poolAuth == null || !poolAuth.Write)
 			{
 				return Forbid();
 			}
 
-			IDevicePlatform? Platform = await DeviceService.GetPlatformAsync(Device.PlatformId);
+			IDevicePlatform? platform = await _deviceService.GetPlatformAsync(device.PlatformId);
 
-			if (Platform == null)
+			if (platform == null)
 			{
-				return BadRequest($"Platform id ${Device.PlatformId} does not exist");
+				return BadRequest($"Platform id ${device.PlatformId} does not exist");
 			}
 
-			DevicePoolId? PoolIdValue = null;
+			DevicePoolId? poolIdValue = null;
 
-			if (Update.PoolId != null)
+			if (update.PoolId != null)
 			{
-				PoolIdValue = new DevicePoolId(Update.PoolId);
+				poolIdValue = new DevicePoolId(update.PoolId);
 			}
 
-			string? ModelIdValue = null;
+			string? modelIdValue = null;
 
-			if (Update.ModelId != null)
+			if (update.ModelId != null)
 			{
-				ModelIdValue = new string(Update.ModelId);
+				modelIdValue = new string(update.ModelId);
 
-				if (Platform.Models?.FirstOrDefault(x => x == ModelIdValue) == null)
+				if (platform.Models?.FirstOrDefault(x => x == modelIdValue) == null)
 				{
-					return BadRequest($"Bad model id {Update.ModelId} for platform {Platform.Id} on request");
+					return BadRequest($"Bad model id {update.ModelId} for platform {platform.Id} on request");
 				}
-
 			}
 
-			string? Name = Update.Name?.Trim();
-			string? Address = Update.Address?.Trim();
+			string? name = update.Name?.Trim();
+			string? address = update.Address?.Trim();
 
-            await DeviceService.UpdateDeviceAsync(DeviceIdValue, PoolIdValue, Name, Address, ModelIdValue, Update.Notes, Update.Enabled, Update.Problem, Update.Maintenance, InternalUser.Id);
+            await _deviceService.UpdateDeviceAsync(deviceIdValue, poolIdValue, name, address, modelIdValue, update.Notes, update.Enabled, update.Problem, update.Maintenance, internalUser.Id);
 
 			return Ok();
 		}
@@ -258,82 +254,80 @@ namespace Horde.Build.Controllers
 		/// </summary>
 		[HttpPut]
 		[Authorize]
-		[Route("/api/v2/devices/{DeviceId}/checkout")]
+		[Route("/api/v2/devices/{deviceId}/checkout")]
 		[ProducesResponseType(typeof(List<GetDeviceResponse>), 200)]
-		public async Task<ActionResult> CheckoutDeviceAsync(string DeviceId, [FromBody] CheckoutDeviceRequest Request)
+		public async Task<ActionResult> CheckoutDeviceAsync(string deviceId, [FromBody] CheckoutDeviceRequest request)
 		{
 
-			IUser? InternalUser = await UserCollection.GetUserAsync(User);
-			if (InternalUser == null)
+			IUser? internalUser = await UserCollection.GetUserAsync(User);
+			if (internalUser == null)
 			{
 				return NotFound();
 			}
 
-			DeviceId DeviceIdValue = new DeviceId(DeviceId);
+			DeviceId deviceIdValue = new DeviceId(deviceId);
 
-			IDevice? Device = await DeviceService.GetDeviceAsync(DeviceIdValue);
+			IDevice? device = await _deviceService.GetDeviceAsync(deviceIdValue);
 
-			if (Device == null)
+			if (device == null)
 			{
-				return BadRequest($"Device with id ${DeviceId} does not exist");
+				return BadRequest($"Device with id ${deviceId} does not exist");
 			}
 
-			DevicePoolAuthorization? PoolAuth = await DeviceService.GetUserPoolAuthorizationAsync(Device.PoolId, User);
+			DevicePoolAuthorization? poolAuth = await _deviceService.GetUserPoolAuthorizationAsync(device.PoolId, User);
 
-			if (PoolAuth == null || !PoolAuth.Write)
+			if (poolAuth == null || !poolAuth.Write)
 			{
 				return Forbid();
 			}
 
-			if (Request.Checkout)
+			if (request.Checkout)
 			{
-				if (!string.IsNullOrEmpty(Device.CheckedOutByUser))
+				if (!String.IsNullOrEmpty(device.CheckedOutByUser))
 				{
-					return BadRequest($"Already checked out by user {Device.CheckedOutByUser}");
+					return BadRequest($"Already checked out by user {device.CheckedOutByUser}");
 				}
 
-                await DeviceService.CheckoutDeviceAsync(DeviceIdValue, InternalUser.Id);
+                await _deviceService.CheckoutDeviceAsync(deviceIdValue, internalUser.Id);
 
             }
 			else
 			{
-				await DeviceService.CheckoutDeviceAsync(DeviceIdValue, null);
+				await _deviceService.CheckoutDeviceAsync(deviceIdValue, null);
 			}
 
 			return Ok();
 		}
 
-
 		/// <summary>
 		/// Delete a specific device
 		/// </summary>
 		[HttpDelete]
 		[Authorize]
-		[Route("/api/v2/devices/{DeviceId}")]
-		public async Task<ActionResult> DeleteDeviceAsync(string DeviceId)
+		[Route("/api/v2/devices/{deviceId}")]
+		public async Task<ActionResult> DeleteDeviceAsync(string deviceId)
 		{
-			if (!await DeviceService.AuthorizeAsync(AclAction.DeviceWrite, User))
+			if (!await _deviceService.AuthorizeAsync(AclAction.DeviceWrite, User))
 			{
 				return Forbid();
 			}
 
-			DeviceId DeviceIdValue = new DeviceId(DeviceId);
+			DeviceId deviceIdValue = new DeviceId(deviceId);
 
-			IDevice? Device = await DeviceService.GetDeviceAsync(DeviceIdValue);
-			if (Device == null)
+			IDevice? device = await _deviceService.GetDeviceAsync(deviceIdValue);
+			if (device == null)
 			{
 				return NotFound();
 			}
 
-			DevicePoolAuthorization? PoolAuth = await DeviceService.GetUserPoolAuthorizationAsync(Device.PoolId, User);
+			DevicePoolAuthorization? poolAuth = await _deviceService.GetUserPoolAuthorizationAsync(device.PoolId, User);
 
-			if (PoolAuth == null || !PoolAuth.Write)
+			if (poolAuth == null || !poolAuth.Write)
 			{
 				return Forbid();
 			}
 
-
-			await DeviceService.DeleteDeviceAsync(DeviceIdValue);
+			await _deviceService.DeleteDeviceAsync(deviceIdValue);
 			return Ok();
 		}
 
@@ -345,25 +339,23 @@ namespace Horde.Build.Controllers
 		[HttpPost]
 		[Authorize]
 		[Route("/api/v2/devices/platforms")]
-		public async Task<ActionResult<CreateDevicePlatformResponse>> CreatePlatformAsync([FromBody] CreateDevicePlatformRequest Request)
+		public async Task<ActionResult<CreateDevicePlatformResponse>> CreatePlatformAsync([FromBody] CreateDevicePlatformRequest request)
 		{
-			if (!await DeviceService.AuthorizeAsync(AclAction.DeviceWrite, User))
+			if (!await _deviceService.AuthorizeAsync(AclAction.DeviceWrite, User))
 			{
 				return Forbid();
 			}
 
-			string Name = Request.Name.Trim();
+			string name = request.Name.Trim();
 
-			IDevicePlatform? Platform = await DeviceService.TryCreatePlatformAsync(DevicePlatformId.Sanitize(Name), Name);
+			IDevicePlatform? platform = await _deviceService.TryCreatePlatformAsync(DevicePlatformId.Sanitize(name), name);
 
-			if (Platform == null)
+			if (platform == null)
 			{
-				return BadRequest($"Unable to create platform for {Name}");
+				return BadRequest($"Unable to create platform for {name}");
 			}
 
-			return new CreateDevicePlatformResponse(Platform.Id.ToString());
-
-
+			return new CreateDevicePlatformResponse(platform.Id.ToString());
 		}
 
 		/// <summary>
@@ -371,20 +363,17 @@ namespace Horde.Build.Controllers
 		/// </summary>
 		[HttpPut]
 		[Authorize]
-		[Route("/api/v2/devices/platforms/{PlatformId}")]
-		public async Task<ActionResult<CreateDevicePlatformResponse>> UpdatePlatformAsync(string PlatformId, [FromBody] UpdateDevicePlatformRequest Request)
+		[Route("/api/v2/devices/platforms/{platformId}")]
+		public async Task<ActionResult<CreateDevicePlatformResponse>> UpdatePlatformAsync(string platformId, [FromBody] UpdateDevicePlatformRequest request)
 		{
-			if (!await DeviceService.AuthorizeAsync(AclAction.DeviceWrite, User))
+			if (!await _deviceService.AuthorizeAsync(AclAction.DeviceWrite, User))
 			{
 				return Forbid();
 			}
 
-
-			await DeviceService.UpdatePlatformAsync(new DevicePlatformId(PlatformId), Request.ModelIds);
+			await _deviceService.UpdatePlatformAsync(new DevicePlatformId(platformId), request.ModelIds);
 
 			return Ok();
-
-
 		}
 
 		/// <summary>
@@ -397,22 +386,22 @@ namespace Horde.Build.Controllers
 		public async Task<ActionResult<List<object>>> GetDevicePlatformsAsync()
 		{
 
-			if (!await DeviceService.AuthorizeAsync(AclAction.DeviceRead, User))
+			if (!await _deviceService.AuthorizeAsync(AclAction.DeviceRead, User))
 			{
 				return Forbid();
 			}
 
-			List<IDevicePlatform> Platforms = await DeviceService.GetPlatformsAsync();
+			List<IDevicePlatform> platforms = await _deviceService.GetPlatformsAsync();
 
-			List<object> Responses = new List<object>();
+			List<object> responses = new List<object>();
 
-			foreach (IDevicePlatform Platform in Platforms)
+			foreach (IDevicePlatform platform in platforms)
 			{
 				// @todo: ACL per platform
-				Responses.Add(new GetDevicePlatformResponse(Platform.Id.ToString(), Platform.Name, Platform.Models?.ToArray() ?? Array.Empty<string>()));
+				responses.Add(new GetDevicePlatformResponse(platform.Id.ToString(), platform.Name, platform.Models?.ToArray() ?? Array.Empty<string>()));
 			}
 
-			return Responses;
+			return responses;
 		}
 
 		// POOLS
@@ -423,23 +412,23 @@ namespace Horde.Build.Controllers
 		[HttpPost]
 		[Authorize]
 		[Route("/api/v2/devices/pools")]
-		public async Task<ActionResult<CreateDevicePoolResponse>> CreatePoolAsync([FromBody] CreateDevicePoolRequest Request)
+		public async Task<ActionResult<CreateDevicePoolResponse>> CreatePoolAsync([FromBody] CreateDevicePoolRequest request)
 		{
-			if (!await DeviceService.AuthorizeAsync(AclAction.DeviceWrite, User))
+			if (!await _deviceService.AuthorizeAsync(AclAction.DeviceWrite, User))
 			{
 				return Forbid();
 			}
 
-			string Name = Request.Name.Trim();
+			string name = request.Name.Trim();
 
-			IDevicePool? Pool = await DeviceService.TryCreatePoolAsync(DevicePoolId.Sanitize(Name), Name, Request.PoolType, Request.ProjectIds?.Select(x => new ProjectId(x)).ToList());
+			IDevicePool? pool = await _deviceService.TryCreatePoolAsync(DevicePoolId.Sanitize(name), name, request.PoolType, request.ProjectIds?.Select(x => new ProjectId(x)).ToList());
 
-			if (Pool == null)
+			if (pool == null)
 			{
-				return BadRequest($"Unable to create pool for {Request.Name}");
+				return BadRequest($"Unable to create pool for {request.Name}");
 			}
 
-			return new CreateDevicePoolResponse(Pool.Id.ToString());
+			return new CreateDevicePoolResponse(pool.Id.ToString());
 
 		}
 
@@ -449,29 +438,28 @@ namespace Horde.Build.Controllers
 		[HttpPut]
 		[Authorize]
 		[Route("/api/v2/devices/pools")]
-		public async Task<ActionResult> UpdatePoolAsync([FromBody] UpdateDevicePoolRequest Request)
+		public async Task<ActionResult> UpdatePoolAsync([FromBody] UpdateDevicePoolRequest request)
 		{
 
-			DevicePoolAuthorization? PoolAuth = await DeviceService.GetUserPoolAuthorizationAsync(new DevicePoolId(Request.Id), User);
+			DevicePoolAuthorization? poolAuth = await _deviceService.GetUserPoolAuthorizationAsync(new DevicePoolId(request.Id), User);
 
-			if (PoolAuth == null || !PoolAuth.Write)
+			if (poolAuth == null || !poolAuth.Write)
 			{
 				return Forbid();
 			}
 
-			List<ProjectId>? ProjectIds = null;
+			List<ProjectId>? projectIds = null;
 
-			if (Request.ProjectIds != null)
+			if (request.ProjectIds != null)
 			{
-				ProjectIds = Request.ProjectIds.Select(x => new ProjectId(x)).ToList();
+				projectIds = request.ProjectIds.Select(x => new ProjectId(x)).ToList();
 			}
 
-			await DeviceService.UpdatePoolAsync(new DevicePoolId(Request.Id), ProjectIds);
+			await _deviceService.UpdatePoolAsync(new DevicePoolId(request.Id), projectIds);
 			
 			return Ok();
 		}
 
-
 		/// <summary>
 		/// Get a list of existing device pools
 		/// </summary>
@@ -482,25 +470,25 @@ namespace Horde.Build.Controllers
 		public async Task<ActionResult<List<object>>> GetDevicePoolsAsync()
 		{
 
-			List<DevicePoolAuthorization> PoolAuth = await DeviceService.GetUserPoolAuthorizationsAsync(User);
+			List<DevicePoolAuthorization> poolAuth = await _deviceService.GetUserPoolAuthorizationsAsync(User);
 
-			List<IDevicePool> Pools = await DeviceService.GetPoolsAsync();
+			List<IDevicePool> pools = await _deviceService.GetPoolsAsync();
 
-			List<object> Responses = new List<object>();
+			List<object> responses = new List<object>();
 
-			foreach (IDevicePool Pool in Pools)
+			foreach (IDevicePool pool in pools)
 			{
-				DevicePoolAuthorization? Auth = PoolAuth.Where(x => x.Pool.Id == Pool.Id).FirstOrDefault();
+				DevicePoolAuthorization? auth = poolAuth.Where(x => x.Pool.Id == pool.Id).FirstOrDefault();
 
-				if (Auth == null || !Auth.Read)
+				if (auth == null || !auth.Read)
 				{
 					continue;
 				}
 				
-				Responses.Add(new GetDevicePoolResponse(Pool.Id.ToString(), Pool.Name, Pool.PoolType, Auth.Write));
+				responses.Add(new GetDevicePoolResponse(pool.Id.ToString(), pool.Name, pool.PoolType, auth.Write));
 			}
 
-			return Responses;
+			return responses;
 		}
 
 		// RESERVATIONS
@@ -512,86 +500,80 @@ namespace Horde.Build.Controllers
 		[Authorize]
 		[Route("/api/v2/devices/reservations")]
 		[ProducesResponseType(typeof(CreateDeviceReservationResponse), 200)]
-		public async Task<ActionResult<CreateDeviceReservationResponse>> CreateDeviceReservation([FromBody] CreateDeviceReservationRequest Request)
+		public async Task<ActionResult<CreateDeviceReservationResponse>> CreateDeviceReservation([FromBody] CreateDeviceReservationRequest request)
 		{
 
-			if (!await DeviceService.AuthorizeAsync(AclAction.DeviceWrite, User))
+			if (!await _deviceService.AuthorizeAsync(AclAction.DeviceWrite, User))
 			{
 				return Forbid();
 			}
 
-			List<IDevicePool> Pools = await DeviceService.GetPoolsAsync();
-			List<IDevicePlatform> Platforms = await DeviceService.GetPlatformsAsync();
+			List<IDevicePool> pools = await _deviceService.GetPoolsAsync();
+			List<IDevicePlatform> platforms = await _deviceService.GetPlatformsAsync();
 
-			DevicePoolId PoolIdValue = new DevicePoolId(Request.PoolId);
-			IDevicePool? Pool = Pools.FirstOrDefault(x => x.Id == PoolIdValue);
-			if (Pool == null)
+			DevicePoolId poolIdValue = new DevicePoolId(request.PoolId);
+			IDevicePool? pool = pools.FirstOrDefault(x => x.Id == poolIdValue);
+			if (pool == null)
 			{
-				return BadRequest($"Unknown pool {Request.PoolId} on device reservation request");
+				return BadRequest($"Unknown pool {request.PoolId} on device reservation request");
 			}
 
-			List<DeviceRequestData> RequestedDevices = new List<DeviceRequestData>();
+			List<DeviceRequestData> requestedDevices = new List<DeviceRequestData>();
 
-			foreach (DeviceReservationRequest DeviceRequest in Request.Devices)
+			foreach (DeviceReservationRequest deviceRequest in request.Devices)
 			{
-				DevicePlatformId PlatformIdValue = new DevicePlatformId(DeviceRequest.PlatformId);
-				IDevicePlatform? Platform = Platforms.FirstOrDefault(x => x.Id == PlatformIdValue);
+				DevicePlatformId platformIdValue = new DevicePlatformId(deviceRequest.PlatformId);
+				IDevicePlatform? platform = platforms.FirstOrDefault(x => x.Id == platformIdValue);
 
-				if (Platform == null)
+				if (platform == null)
 				{
-					return BadRequest($"Unknown platform {DeviceRequest.PlatformId} on device reservation request");
+					return BadRequest($"Unknown platform {deviceRequest.PlatformId} on device reservation request");
 				}
 
-				if (DeviceRequest.IncludeModels != null)
+				if (deviceRequest.IncludeModels != null)
 				{
-					foreach (string Model in DeviceRequest.IncludeModels)
+					foreach (string model in deviceRequest.IncludeModels)
 					{
-						if (Platform.Models?.FirstOrDefault(x => x == Model) == null)
+						if (platform.Models?.FirstOrDefault(x => x == model) == null)
 						{
-							return BadRequest($"Unknown model {Model} for platform {DeviceRequest.PlatformId} on device reservation request");
+							return BadRequest($"Unknown model {model} for platform {deviceRequest.PlatformId} on device reservation request");
 						}
-
 					}
-
 				}
 
-				if (DeviceRequest.ExcludeModels != null)
+				if (deviceRequest.ExcludeModels != null)
 				{
-					foreach (string Model in DeviceRequest.ExcludeModels)
+					foreach (string model in deviceRequest.ExcludeModels)
 					{
-						if (Platform.Models?.FirstOrDefault(x => x == Model) == null)
+						if (platform.Models?.FirstOrDefault(x => x == model) == null)
 						{
-							return BadRequest($"Unknown model {Model} for platform {DeviceRequest.PlatformId} on device reservation request");
+							return BadRequest($"Unknown model {model} for platform {deviceRequest.PlatformId} on device reservation request");
 						}
-
 					}
-
 				}					
 
-				RequestedDevices.Add(new DeviceRequestData(PlatformIdValue, PlatformIdValue.ToString(), DeviceRequest.IncludeModels, DeviceRequest.ExcludeModels));
-
+				requestedDevices.Add(new DeviceRequestData(platformIdValue, platformIdValue.ToString(), deviceRequest.IncludeModels, deviceRequest.ExcludeModels));
 			}
 
-			IDeviceReservation? Reservation = await DeviceService.TryCreateReservationAsync(PoolIdValue, RequestedDevices);
+			IDeviceReservation? reservation = await _deviceService.TryCreateReservationAsync(poolIdValue, requestedDevices);
 
-			if (Reservation == null)
+			if (reservation == null)
 			{
 				return Conflict("Unable to allocated devices for reservation");
 			}
 
-			List<IDevice> Devices = await DeviceService.GetDevicesAsync(Reservation.Devices);
+			List<IDevice> devices = await _deviceService.GetDevicesAsync(reservation.Devices);
 
+			CreateDeviceReservationResponse response = new CreateDeviceReservationResponse();
 
-			CreateDeviceReservationResponse Response = new CreateDeviceReservationResponse();
+			response.Id = reservation.Id.ToString();
 
-			Response.Id = Reservation.Id.ToString();
-
-			foreach (IDevice Device in Devices)
+			foreach (IDevice device in devices)
 			{
-				Response.Devices.Add(new GetDeviceResponse(Device.Id.ToString(), Device.PlatformId.ToString(), Device.PoolId.ToString(), Device.Name, Device.Enabled, Device.Address, Device.ModelId?.ToString(), Device.ModifiedByUser, Device.Notes, Device.ProblemTimeUtc, Device.MaintenanceTimeUtc, Device.Utilization));
+				response.Devices.Add(new GetDeviceResponse(device.Id.ToString(), device.PlatformId.ToString(), device.PoolId.ToString(), device.Name, device.Enabled, device.Address, device.ModelId?.ToString(), device.ModifiedByUser, device.Notes, device.ProblemTimeUtc, device.MaintenanceTimeUtc, device.Utilization));
 			}
 
-			return Response;
+			return response;
 		}
 
 		/// <summary>
@@ -603,36 +585,36 @@ namespace Horde.Build.Controllers
 		[ProducesResponseType(typeof(List<GetDeviceReservationResponse>), 200)]
 		public async Task<ActionResult<List<GetDeviceReservationResponse>>> GetDeviceReservations()
 		{
-			if (!await DeviceService.AuthorizeAsync(AclAction.DeviceRead, User))
+			if (!await _deviceService.AuthorizeAsync(AclAction.DeviceRead, User))
 			{
 				return Forbid();
 			}
 
-			List<IDeviceReservation> Reservations = await DeviceService.GetReservationsAsync();
+			List<IDeviceReservation> reservations = await _deviceService.GetReservationsAsync();
 
-			List<GetDeviceReservationResponse> Response = new List<GetDeviceReservationResponse>();
+			List<GetDeviceReservationResponse> response = new List<GetDeviceReservationResponse>();
 
-			foreach (IDeviceReservation Reservation in Reservations)
+			foreach (IDeviceReservation reservation in reservations)
 			{
-				GetDeviceReservationResponse ReservationResponse = new GetDeviceReservationResponse()
+				GetDeviceReservationResponse reservationResponse = new GetDeviceReservationResponse()
 				{
-					Id = Reservation.Id.ToString(),
-					PoolId = Reservation.PoolId.ToString(),
-					Devices = Reservation.Devices.Select(x => x.ToString()).ToList(),
-					JobId = Reservation.JobId?.ToString(),
-					StepId = Reservation.StepId?.ToString(),
-					UserId = Reservation.UserId?.ToString(),
-					Hostname = Reservation.Hostname,
-					ReservationDetails = Reservation.ReservationDetails,
-					CreateTimeUtc = Reservation.CreateTimeUtc,
-					LegacyGuid = Reservation.LegacyGuid
+					Id = reservation.Id.ToString(),
+					PoolId = reservation.PoolId.ToString(),
+					Devices = reservation.Devices.Select(x => x.ToString()).ToList(),
+					JobId = reservation.JobId?.ToString(),
+					StepId = reservation.StepId?.ToString(),
+					UserId = reservation.UserId?.ToString(),
+					Hostname = reservation.Hostname,
+					ReservationDetails = reservation.ReservationDetails,
+					CreateTimeUtc = reservation.CreateTimeUtc,
+					LegacyGuid = reservation.LegacyGuid
 				};
 
-				Response.Add(ReservationResponse);
+				response.Add(reservationResponse);
 
 			}
 
-			return Response;
+			return response;
 		}
 
 		/// <summary>
@@ -640,19 +622,19 @@ namespace Horde.Build.Controllers
 		/// </summary>
 		[HttpPut]
 		[Authorize]
-		[Route("/api/v2/devices/reservations/{ReservationId}")]
-		public async Task<ActionResult> UpdateReservationAsync(string ReservationId)
+		[Route("/api/v2/devices/reservations/{reservationId}")]
+		public async Task<ActionResult> UpdateReservationAsync(string reservationId)
 		{
-			if (!await DeviceService.AuthorizeAsync(AclAction.DeviceWrite, User))
+			if (!await _deviceService.AuthorizeAsync(AclAction.DeviceWrite, User))
 			{
 				return Forbid();
 			}
 
-			ObjectId ReservationIdValue = new ObjectId(ReservationId);
+			ObjectId reservationIdValue = new ObjectId(reservationId);
 
-			bool Updated = await DeviceService.TryUpdateReservationAsync(ReservationIdValue);
+			bool updated = await _deviceService.TryUpdateReservationAsync(reservationIdValue);
 
-			if (!Updated)
+			if (!updated)
 			{
 				return BadRequest("Failed to update reservation");
 			}
@@ -665,19 +647,19 @@ namespace Horde.Build.Controllers
 		/// </summary>
 		[HttpDelete]
 		[Authorize]
-		[Route("/api/v2/devices/reservations/{ReservationId}")]
-		public async Task<ActionResult> DeleteReservationAsync(string ReservationId)
+		[Route("/api/v2/devices/reservations/{reservationId}")]
+		public async Task<ActionResult> DeleteReservationAsync(string reservationId)
 		{
-			if (!await DeviceService.AuthorizeAsync(AclAction.DeviceWrite, User))
+			if (!await _deviceService.AuthorizeAsync(AclAction.DeviceWrite, User))
 			{
 				return Forbid();
 			}
 
-			ObjectId ReservationIdValue = new ObjectId(ReservationId);
+			ObjectId reservationIdValue = new ObjectId(reservationId);
 
-			bool Deleted = await DeviceService.DeleteReservationAsync(ReservationIdValue);
+			bool deleted = await _deviceService.DeleteReservationAsync(reservationIdValue);
 
-			if (!Deleted)
+			if (!deleted)
 			{
 				return BadRequest("Failed to delete reservation");
 			}
@@ -701,133 +683,131 @@ namespace Horde.Build.Controllers
 		[HttpPost]
 		[Route("/api/v1/reservations")]
 		[ProducesResponseType(typeof(GetLegacyReservationResponse), 200)]
-		public async Task<ActionResult<GetLegacyReservationResponse>> CreateDeviceReservationV1Async([FromBody] LegacyCreateReservationRequest Request)
+		public async Task<ActionResult<GetLegacyReservationResponse>> CreateDeviceReservationV1Async([FromBody] LegacyCreateReservationRequest request)
 		{
 
-			List<IDevicePool> Pools = await DeviceService.GetPoolsAsync();
-			List<IDevicePlatform> Platforms = await DeviceService.GetPlatformsAsync();
+			List<IDevicePool> pools = await _deviceService.GetPoolsAsync();
+			List<IDevicePlatform> platforms = await _deviceService.GetPlatformsAsync();
 
-			string? PoolId = Request.PoolId;
+			string? poolId = request.PoolId;
 
 			// @todo: Remove this once all streams are updated to provide jobid
-			string Details = "";
-			if ((String.IsNullOrEmpty(Request.JobId) || String.IsNullOrEmpty(Request.StepId)))
+			string details = "";
+			if ((String.IsNullOrEmpty(request.JobId) || String.IsNullOrEmpty(request.StepId)))
 			{
-				if (!string.IsNullOrEmpty(Request.ReservationDetails))
+				if (!String.IsNullOrEmpty(request.ReservationDetails))
 				{
-					Details = $" - {Request.ReservationDetails}";
+					details = $" - {request.ReservationDetails}";
 				}
 				else
 				{
-					Details = $" - Host {Request.Hostname}";
+					details = $" - Host {request.Hostname}";
 				}
 			}
 
-			if (string.IsNullOrEmpty(PoolId))
+			if (String.IsNullOrEmpty(poolId))
 			{
-				Logger.LogError("No pool specified, defaulting to UE4 {Details} JobId: {JobId}, StepId: {StepId}", Details, Request.JobId, Request.StepId);
-				string Message = $"No pool specified, defaulting to UE4" + Details;
-				await DeviceService.NotifyDeviceServiceAsync(Message, null, Request.JobId, Request.StepId);
-				PoolId = "ue4";
+				_logger.LogError("No pool specified, defaulting to UE4 {Details} JobId: {JobId}, StepId: {StepId}", details, request.JobId, request.StepId);
+				string message = $"No pool specified, defaulting to UE4" + details;
+				await _deviceService.NotifyDeviceServiceAsync(message, null, request.JobId, request.StepId);
+				poolId = "ue4";
 				//return BadRequest(Message);
 			}
 
-			DevicePoolId PoolIdValue = DevicePoolId.Sanitize(PoolId);
-			IDevicePool? Pool = Pools.FirstOrDefault(x => x.Id == PoolIdValue);
-			if (Pool == null)
+			DevicePoolId poolIdValue = DevicePoolId.Sanitize(poolId);
+			IDevicePool? pool = pools.FirstOrDefault(x => x.Id == poolIdValue);
+			if (pool == null)
 			{
-				Logger.LogError("Unknown pool {PoolId} {Details}", PoolId, Details);
-				string Message = $"Unknown pool {PoolId} " + Details;
-				await DeviceService.NotifyDeviceServiceAsync(Message, null, Request.JobId, Request.StepId);
-				return BadRequest(Message);
+				_logger.LogError("Unknown pool {PoolId} {Details}", poolId, details);
+				string message = $"Unknown pool {poolId} " + details;
+				await _deviceService.NotifyDeviceServiceAsync(message, null, request.JobId, request.StepId);
+				return BadRequest(message);
 			}
 
-			List<DeviceRequestData> RequestedDevices = new List<DeviceRequestData>();
+			List<DeviceRequestData> requestedDevices = new List<DeviceRequestData>();
 
-			List<string> PerfSpecs = new List<string>();
+			List<string> perfSpecs = new List<string>();
 
-			foreach (string DeviceType in Request.DeviceTypes)
+			foreach (string deviceType in request.DeviceTypes)
 			{
 
-				string PlatformName = DeviceType;
-				string PerfSpecName = "Unspecified";
+				string platformName = deviceType;
+				string perfSpecName = "Unspecified";
 
-				if (DeviceType.Contains(':', StringComparison.Ordinal))
+				if (deviceType.Contains(':', StringComparison.Ordinal))
 				{
-					string[] Tokens = DeviceType.Split(":");
-					PlatformName = Tokens[0];
-					PerfSpecName = Tokens[1];
+					string[] tokens = deviceType.Split(":");
+					platformName = tokens[0];
+					perfSpecName = tokens[1];
 				}
 
-				PerfSpecs.Add(PerfSpecName);
+				perfSpecs.Add(perfSpecName);
 
-				DevicePlatformId PlatformId;
+				DevicePlatformId platformId;
 
-				DevicePlatformMapV1 MapV1 = await DeviceService.GetPlatformMapV1();
+				DevicePlatformMapV1 mapV1 = await _deviceService.GetPlatformMapV1();
 
-				if (!MapV1.PlatformMap.TryGetValue(PlatformName, out PlatformId))
+				if (!mapV1.PlatformMap.TryGetValue(platformName, out platformId))
 				{
-					string Message = $"Unknown platform {PlatformName}" + Details;
-					Logger.LogError("Unknown platform {PlatformName} {Details}", PlatformName, Details);
-					await DeviceService.NotifyDeviceServiceAsync(Message, null, Request.JobId, Request.StepId);
-					return BadRequest(Message);
+					string message = $"Unknown platform {platformName}" + details;
+					_logger.LogError("Unknown platform {PlatformName} {Details}", platformName, details);
+					await _deviceService.NotifyDeviceServiceAsync(message, null, request.JobId, request.StepId);
+					return BadRequest(message);
 				}
 
-				IDevicePlatform? Platform = Platforms.FirstOrDefault(x => x.Id == PlatformId);
+				IDevicePlatform? platform = platforms.FirstOrDefault(x => x.Id == platformId);
 
-				if (Platform == null)
+				if (platform == null)
 				{
-					string Message = $"Unknown platform {PlatformId}" + Details;
-					Logger.LogError("Unknown platform {PlatformId} {Details}", PlatformId, Details);
-					await DeviceService.NotifyDeviceServiceAsync(Message, null, Request.JobId, Request.StepId);
-					return BadRequest(Message);
+					string message = $"Unknown platform {platformId}" + details;
+					_logger.LogError("Unknown platform {PlatformId} {Details}", platformId, details);
+					await _deviceService.NotifyDeviceServiceAsync(message, null, request.JobId, request.StepId);
+					return BadRequest(message);
 				}
 
-				List<string> IncludeModels = new List<string>();
-				List<string> ExcludeModels = new List<string>();
+				List<string> includeModels = new List<string>();
+				List<string> excludeModels = new List<string>();
 
-				if (PerfSpecName == "High")
+				if (perfSpecName == "High")
 				{
-					string? Model = null;
-					if (MapV1.PerfSpecHighMap.TryGetValue(PlatformId, out Model))
+					string? model = null;
+					if (mapV1.PerfSpecHighMap.TryGetValue(platformId, out model))
 					{
-						IncludeModels.Add(Model);
+						includeModels.Add(model);
 					}
 				}
 
-				if (PerfSpecName == "Minimum" || PerfSpecName == "Recommended")
+				if (perfSpecName == "Minimum" || perfSpecName == "Recommended")
 				{
-					string? Model = null;
-					if (MapV1.PerfSpecHighMap.TryGetValue(PlatformId, out Model))
+					string? model = null;
+					if (mapV1.PerfSpecHighMap.TryGetValue(platformId, out model))
 					{
-						ExcludeModels.Add(Model);
+						excludeModels.Add(model);
 					}
 				}
 
-
-				RequestedDevices.Add(new DeviceRequestData(PlatformId, PlatformName, IncludeModels, ExcludeModels));
-
+				requestedDevices.Add(new DeviceRequestData(platformId, platformName, includeModels, excludeModels));
 			}
 
-			IDeviceReservation? Reservation = await DeviceService.TryCreateReservationAsync(PoolIdValue, RequestedDevices, Request.Hostname, Request.ReservationDetails, Request.JobId, Request.StepId);
+			IDeviceReservation? reservation = await _deviceService.TryCreateReservationAsync(poolIdValue, requestedDevices, request.Hostname, request.ReservationDetails, request.JobId, request.StepId);
 
-			if (Reservation == null)
+			if (reservation == null)
 			{
 				return Conflict("Unable to allocated devices for reservation");
 			}
 
-			List<IDevice> Devices = await DeviceService.GetDevicesAsync(Reservation.Devices);
+			List<IDevice> devices = await _deviceService.GetDevicesAsync(reservation.Devices);
 
-			GetLegacyReservationResponse Response = new GetLegacyReservationResponse();
+			GetLegacyReservationResponse response = new GetLegacyReservationResponse();
 
-			Response.Guid = Reservation.LegacyGuid;
-			Response.DeviceNames = Devices.Select(x => x.Name).ToArray();
-			Response.DevicePerfSpecs = PerfSpecs.ToArray();
-			Response.HostName = Request.Hostname;
-			Response.StartDateTime = Reservation.CreateTimeUtc.ToString("O", System.Globalization.CultureInfo.InvariantCulture);
-			Response.Duration = $"{Request.Duration}";
+			response.Guid = reservation.LegacyGuid;
+			response.DeviceNames = devices.Select(x => x.Name).ToArray();
+			response.DevicePerfSpecs = perfSpecs.ToArray();
+			response.HostName = request.Hostname;
+			response.StartDateTime = reservation.CreateTimeUtc.ToString("O", System.Globalization.CultureInfo.InvariantCulture);
+			response.Duration = $"{request.Duration}";
 
-			return new JsonResult(Response, new JsonSerializerOptions() { PropertyNamingPolicy = null });
+			return new JsonResult(response, new JsonSerializerOptions() { PropertyNamingPolicy = null });
 
 		}
 
@@ -835,56 +815,56 @@ namespace Horde.Build.Controllers
 		/// Renew a reservation
 		/// </summary>
 		[HttpPut]
-		[Route("/api/v1/reservations/{ReservationGuid}")]
+		[Route("/api/v1/reservations/{reservationGuid}")]
 		[ProducesResponseType(typeof(GetLegacyReservationResponse), 200)]
-		public async Task<ActionResult<GetLegacyReservationResponse>> UpdateReservationV1Async(string ReservationGuid /* [FromBody] string Duration */)
+		public async Task<ActionResult<GetLegacyReservationResponse>> UpdateReservationV1Async(string reservationGuid /* [FromBody] string Duration */)
 		{
-			IDeviceReservation? Reservation = await DeviceService.TryGetReservationFromLegacyGuidAsync(ReservationGuid);
+			IDeviceReservation? reservation = await _deviceService.TryGetReservationFromLegacyGuidAsync(reservationGuid);
 
-			if (Reservation == null)
+			if (reservation == null)
 			{
-				Logger.LogError("Unable to find reservation for legacy guid {ReservationGuid}", ReservationGuid);
+				_logger.LogError("Unable to find reservation for legacy guid {ReservationGuid}", reservationGuid);
 				return BadRequest();
 			}
 
-			bool Updated = await DeviceService.TryUpdateReservationAsync(Reservation.Id);
+			bool updated = await _deviceService.TryUpdateReservationAsync(reservation.Id);
 
-			if (!Updated)
+			if (!updated)
 			{
-				Logger.LogError("Unable to find reservation for reservation {ReservationId}", Reservation.Id);
+				_logger.LogError("Unable to find reservation for reservation {ReservationId}", reservation.Id);
 				return BadRequest();
 			}
 
-			List<IDevice> Devices = await DeviceService.GetDevicesAsync(Reservation.Devices);
+			List<IDevice> devices = await _deviceService.GetDevicesAsync(reservation.Devices);
 
-			GetLegacyReservationResponse Response = new GetLegacyReservationResponse();
+			GetLegacyReservationResponse response = new GetLegacyReservationResponse();
 
-			Response.Guid = Reservation.LegacyGuid;
-			Response.DeviceNames = Reservation.Devices.Select(DeviceId => Devices.First(D => D.Id == DeviceId).Name).ToArray();
-			Response.HostName = Reservation.Hostname ?? "";
-			Response.StartDateTime = Reservation.CreateTimeUtc.ToString("O", System.Globalization.CultureInfo.InvariantCulture);
-			Response.Duration = "00:10:00"; // matches gauntlet duration
+			response.Guid = reservation.LegacyGuid;
+			response.DeviceNames = reservation.Devices.Select(deviceId => devices.First(d => d.Id == deviceId).Name).ToArray();
+			response.HostName = reservation.Hostname ?? "";
+			response.StartDateTime = reservation.CreateTimeUtc.ToString("O", System.Globalization.CultureInfo.InvariantCulture);
+			response.Duration = "00:10:00"; // matches gauntlet duration
 
-			return new JsonResult(Response, new JsonSerializerOptions() { PropertyNamingPolicy = null });
+			return new JsonResult(response, new JsonSerializerOptions() { PropertyNamingPolicy = null });
 		}
 
 		/// <summary>
 		/// Delete a reservation
 		/// </summary>
 		[HttpDelete]
-		[Route("/api/v1/reservations/{ReservationGuid}")]
-		public async Task<ActionResult> DeleteReservationV1Async(string ReservationGuid)
+		[Route("/api/v1/reservations/{reservationGuid}")]
+		public async Task<ActionResult> DeleteReservationV1Async(string reservationGuid)
 		{
-			IDeviceReservation? Reservation = await DeviceService.TryGetReservationFromLegacyGuidAsync(ReservationGuid);
+			IDeviceReservation? reservation = await _deviceService.TryGetReservationFromLegacyGuidAsync(reservationGuid);
 
-			if (Reservation == null)
+			if (reservation == null)
 			{
-				return BadRequest($"Unable to find reservation for guid {ReservationGuid}");
+				return BadRequest($"Unable to find reservation for guid {reservationGuid}");
 			}
 
-			bool Deleted = await DeviceService.DeleteReservationAsync(Reservation.Id);
+			bool deleted = await _deviceService.DeleteReservationAsync(reservation.Id);
 
-			if (!Deleted)
+			if (!deleted)
 			{
 				return BadRequest("Failed to delete reservation");
 			}
@@ -892,125 +872,122 @@ namespace Horde.Build.Controllers
 			return Ok();
 		}
 
-
 		/// <summary>
 		/// Get device info for a reserved device
 		/// </summary>
 		[HttpGet]
-		[Route("/api/v1/devices/{DeviceName}")]
+		[Route("/api/v1/devices/{deviceName}")]
 		[ProducesResponseType(typeof(GetLegacyDeviceResponse), 200)]
-		public async Task<ActionResult<GetLegacyDeviceResponse>> GetDeviceV1Async(string DeviceName)
+		public async Task<ActionResult<GetLegacyDeviceResponse>> GetDeviceV1Async(string deviceName)
 		{
-			IDevice? Device = await DeviceService.GetDeviceByNameAsync(DeviceName);
+			IDevice? device = await _deviceService.GetDeviceByNameAsync(deviceName);
 
-			if (Device == null)
+			if (device == null)
 			{
-				return BadRequest($"Unknown device {DeviceName}");
+				return BadRequest($"Unknown device {deviceName}");
 			}
 
-			IDeviceReservation? Reservation = await DeviceService.TryGetDeviceReservation(Device.Id);
+			IDeviceReservation? reservation = await _deviceService.TryGetDeviceReservation(device.Id);
 
-			string? PlatformName = null;
+			string? platformName = null;
 
-			if (Reservation != null && Reservation.Devices.Count == Reservation.RequestedDevicePlatforms.Count)
+			if (reservation != null && reservation.Devices.Count == reservation.RequestedDevicePlatforms.Count)
 			{
-				for (int i = 0; i < Reservation.Devices.Count; i++)
+				for (int i = 0; i < reservation.Devices.Count; i++)
 				{
-					if (Reservation.Devices[i] == Device.Id)
+					if (reservation.Devices[i] == device.Id)
 					{
-						PlatformName = Reservation.RequestedDevicePlatforms[i];
+						platformName = reservation.RequestedDevicePlatforms[i];
 					}
 				}
 			}
 
-			DevicePlatformMapV1 MapV1 = await DeviceService.GetPlatformMapV1();
+			DevicePlatformMapV1 mapV1 = await _deviceService.GetPlatformMapV1();
 
-			if (string.IsNullOrEmpty(PlatformName))
+			if (String.IsNullOrEmpty(platformName))
 			{				
-				if (!MapV1.PlatformReverseMap.TryGetValue(Device.PlatformId, out PlatformName))
+				if (!mapV1.PlatformReverseMap.TryGetValue(device.PlatformId, out platformName))
 				{
-					return BadRequest($"Unable to map platform for {DeviceName} : {Device.PlatformId}");
+					return BadRequest($"Unable to map platform for {deviceName} : {device.PlatformId}");
 				}
 			}
 
-			if (string.IsNullOrEmpty(PlatformName)) {
-				return BadRequest($"Unable to get platform for {DeviceName} from reservation or mapping : {Device.PlatformId}");
+			if (String.IsNullOrEmpty(platformName)) 
+			{
+				return BadRequest($"Unable to get platform for {deviceName} from reservation or mapping : {device.PlatformId}");
 			}
 
-			GetLegacyDeviceResponse Response = new GetLegacyDeviceResponse();
+			GetLegacyDeviceResponse response = new GetLegacyDeviceResponse();
 
-			Response.Name = Device.Name;
-			Response.Type = PlatformName;
-			Response.IPOrHostName = Device.Address ?? "";
-			Response.AvailableStartTime = "00:00:00";
-			Response.AvailableEndTime = "00:00:00";
-			Response.Enabled = true;
-			Response.DeviceData = "";
+			response.Name = device.Name;
+			response.Type = platformName;
+			response.IpOrHostName = device.Address ?? "";
+			response.AvailableStartTime = "00:00:00";
+			response.AvailableEndTime = "00:00:00";
+			response.Enabled = true;
+			response.DeviceData = "";
 
-			Response.PerfSpec = "Minimum";
+			response.PerfSpec = "Minimum";
 
-			if (Device.ModelId != null)
+			if (device.ModelId != null)
 			{
-				string? Model = null;
-				if (MapV1.PerfSpecHighMap.TryGetValue(Device.PlatformId, out Model))
+				string? model;
+				if (mapV1.PerfSpecHighMap.TryGetValue(device.PlatformId, out model))
 				{
-					if (Model == Device.ModelId)
+					if (model == device.ModelId)
 					{
-						Response.PerfSpec = "High";
+						response.PerfSpec = "High";
 					}
-
 				}
 			}
 
-			return new JsonResult(Response, new JsonSerializerOptions() { PropertyNamingPolicy = null });
-
+			return new JsonResult(response, new JsonSerializerOptions() { PropertyNamingPolicy = null });
 		}
 
 		/// <summary>
 		/// Mark a problem device
 		/// </summary>
 		[HttpPut]
-		[Route("/api/v1/deviceerror/{DeviceName}")]
-		public async Task<ActionResult> PutDeviceErrorAsync(string DeviceName)
+		[Route("/api/v1/deviceerror/{deviceName}")]
+		public async Task<ActionResult> PutDeviceErrorAsync(string deviceName)
 		{
+			IDevice? device = await _deviceService.GetDeviceByNameAsync(deviceName);
 
-			IDevice? Device = await DeviceService.GetDeviceByNameAsync(DeviceName);
-
-			if (Device == null)
+			if (device == null)
 			{
-				Logger.LogError("Device error reported for unknown device {DeviceName}", DeviceName);
-				return BadRequest($"Unknown device {DeviceName}");
+				_logger.LogError("Device error reported for unknown device {DeviceName}", deviceName);
+				return BadRequest($"Unknown device {deviceName}");
 			}
 
-			await DeviceService.UpdateDeviceAsync(Device.Id, NewProblem: true);
+			await _deviceService.UpdateDeviceAsync(device.Id, newProblem: true);
 
-			string Message = $"Device problem, {Device.Name} : {Device.PoolId.ToString().ToUpperInvariant()}";
+			string message = $"Device problem, {device.Name} : {device.PoolId.ToString().ToUpperInvariant()}";
 
-			IDeviceReservation? Reservation = await DeviceService.TryGetDeviceReservation(Device.Id);
+			IDeviceReservation? reservation = await _deviceService.TryGetDeviceReservation(device.Id);
 
-			string? JobId = null;
-			string? StepId = null;
-			if (Reservation != null)
+			string? jobId = null;
+			string? stepId = null;
+			if (reservation != null)
 			{
-				JobId = !string.IsNullOrEmpty(Reservation.JobId) ? Reservation.JobId : null;
-				StepId = !string.IsNullOrEmpty(Reservation.StepId) ? Reservation.StepId : null;
+				jobId = !String.IsNullOrEmpty(reservation.JobId) ? reservation.JobId : null;
+				stepId = !String.IsNullOrEmpty(reservation.StepId) ? reservation.StepId : null;
 
-				if ((JobId == null || StepId == null))
+				if ((jobId == null || stepId == null))
 				{
-					if (!string.IsNullOrEmpty(Reservation.ReservationDetails))
+					if (!String.IsNullOrEmpty(reservation.ReservationDetails))
 					{
-						Message += $" - {Reservation.ReservationDetails}";
+						message += $" - {reservation.ReservationDetails}";
 					}
 					else
 					{
-						Message += $" - Host {Reservation.Hostname}";
+						message += $" - Host {reservation.Hostname}";
 					}
 				}
 			}
 
-			Logger.LogError("{DeviceError}", Message);
+			_logger.LogError("{DeviceError}", message);
 
-			await DeviceService.NotifyDeviceServiceAsync(Message, Device.Id, JobId, StepId);
+			await _deviceService.NotifyDeviceServiceAsync(message, device.Id, jobId, stepId);
 
 			return Ok();
 
@@ -1021,28 +998,27 @@ namespace Horde.Build.Controllers
 		/// </summary>
 		[HttpPut]
 		[Route("/api/v1/devices/platformmap")]
-		public async Task<ActionResult> UpdatePlatformMapV1([FromBody] UpdatePlatformMapRequest Request)
+		public async Task<ActionResult> UpdatePlatformMapV1([FromBody] UpdatePlatformMapRequest request)
 		{
-			if (!await AclService.AuthorizeAsync(AclAction.AdminWrite, User))
+			if (!await _aclService.AuthorizeAsync(AclAction.AdminWrite, User))
 			{
 				return Forbid();
 			}
 
-			bool Result = false;
-
+			bool result;
 			try
 			{
-				Result = await DeviceService.UpdatePlatformMapAsync(Request);
+				result = await _deviceService.UpdatePlatformMapAsync(request);
 			}
-			catch (Exception Ex)
+			catch (Exception ex)
 			{
-				Logger.LogError(Ex, "Error updating device platform map {Message}", Ex.Message);
+				_logger.LogError(ex, "Error updating device platform map {Message}", ex.Message);
 				throw;
 			}
 
-			if (!Result)
+			if (!result)
 			{
-				Logger.LogError("Unable to update device platform mapping");
+				_logger.LogError("Unable to update device platform mapping");
 				return BadRequest();
 			}
 
diff --git a/Engine/Source/Programs/Horde/Horde.Build/Controllers/ErrorController.cs b/Engine/Source/Programs/Horde/Horde.Build/Controllers/ErrorController.cs
--- a/Engine/Source/Programs/Horde/Horde.Build/Controllers/ErrorController.cs
+++ b/Engine/Source/Programs/Horde/Horde.Build/Controllers/ErrorController.cs
@@ -23,12 +23,12 @@ namespace Horde.Build.Controllers
 		[Route("/api/v1/error")]
 		public IActionResult Error()
 		{
-			IExceptionHandlerFeature? Context = HttpContext.Features.Get<IExceptionHandlerFeature>();
-			if (Context == null)
+			IExceptionHandlerFeature? context = HttpContext.Features.Get<IExceptionHandlerFeature>();
+			if (context == null)
 			{
 				return NoContent();
 			}
-			return StatusCode(StatusCodes.Status500InternalServerError, LogEvent.Create(LogLevel.Error, KnownLogEvents.None, Context.Error, Context.Error.Message));
+			return StatusCode(StatusCodes.Status500InternalServerError, LogEvent.Create(LogLevel.Error, KnownLogEvents.None, context.Error, context.Error.Message));
 		}
 	}
 }
diff --git a/Engine/Source/Programs/Horde/Horde.Build/Controllers/IssuesController.cs b/Engine/Source/Programs/Horde/Horde.Build/Controllers/IssuesController.cs
--- a/Engine/Source/Programs/Horde/Horde.Build/Controllers/IssuesController.cs
+++ b/Engine/Source/Programs/Horde/Horde.Build/Controllers/IssuesController.cs
@@ -1,24 +1,19 @@
-﻿// Copyright Epic Games, Inc. All Rights Reserved.
+// Copyright Epic Games, Inc. All Rights Reserved.
 
-using EpicGames.Core;
+using System;
+using System.Collections.Generic;
+using System.Linq;
+using System.Threading.Tasks;
 using Horde.Build.Acls;
 using Horde.Build.Api;
 using Horde.Build.Collections;
-using HordeCommon;
 using Horde.Build.Models;
 using Horde.Build.Services;
 using Horde.Build.Utilities;
+using HordeCommon;
 using Microsoft.AspNetCore.Authorization;
 using Microsoft.AspNetCore.Mvc;
-using Microsoft.AspNetCore.Server.Kestrel.Core.Features;
-using Microsoft.IdentityModel.Tokens;
 using MongoDB.Bson;
-using System;
-using System.Collections.Generic;
-using System.ComponentModel.DataAnnotations;
-using System.Linq;
-using System.Text;
-using System.Threading.Tasks;
 
 namespace Horde.Build.Controllers
 {
@@ -35,126 +30,124 @@ namespace Horde.Build.Controllers
 	[Route("[controller]")]
 	public class IssuesController : HordeControllerBase
 	{
-		private readonly IIssueCollection IssueCollection;
-		private readonly IIssueService IssueService;
-		private readonly JobService JobService;
-		private readonly StreamService StreamService;
-		private readonly IUserCollection UserCollection;
-		private readonly ILogEventCollection LogEventCollection;
-		private readonly ILogFileService LogFileService;
+		private readonly IIssueCollection _issueCollection;
+		private readonly IIssueService _issueService;
+		private readonly JobService _jobService;
+		private readonly StreamService _streamService;
+		private readonly IUserCollection _userCollection;
+		private readonly ILogFileService _logFileService;
 
 		/// <summary>
 		/// Constructor
 		/// </summary>
-		public IssuesController(IIssueCollection IssueCollection, IIssueService IssueService, JobService JobService, StreamService StreamService, IUserCollection UserCollection, ILogEventCollection LogEventCollection, ILogFileService LogFileService)
+		public IssuesController(IIssueCollection issueCollection, IIssueService issueService, JobService jobService, StreamService streamService, IUserCollection userCollection, ILogFileService logFileService)
 		{
-			this.IssueCollection = IssueCollection;
-			this.IssueService = IssueService;
-			this.JobService = JobService;
-			this.StreamService = StreamService;
-			this.UserCollection = UserCollection;
-			this.LogEventCollection = LogEventCollection;
-			this.LogFileService = LogFileService;
+			_issueCollection = issueCollection;
+			_issueService = issueService;
+			_jobService = jobService;
+			_streamService = streamService;
+			_userCollection = userCollection;
+			_logFileService = logFileService;
 		}
 
 		/// <summary>
 		/// Retrieve information about a specific issue
 		/// </summary>
-		/// <param name="Ids">Set of issue ids to find</param>
-		/// <param name="StreamId">The stream to query for</param>
-		/// <param name="MinChange">The minimum changelist range to query, inclusive</param>
-		/// <param name="MaxChange">The minimum changelist range to query, inclusive</param>
-		/// <param name="Resolved">Whether to include resolved issues</param>
-		/// <param name="Index">Starting offset of the window of results to return</param>
-		/// <param name="Count">Number of results to return</param>
-		/// <param name="Filter">Filter for the properties to return</param>
+		/// <param name="ids">Set of issue ids to find</param>
+		/// <param name="streamId">The stream to query for</param>
+		/// <param name="minChange">The minimum changelist range to query, inclusive</param>
+		/// <param name="maxChange">The minimum changelist range to query, inclusive</param>
+		/// <param name="resolved">Whether to include resolved issues</param>
+		/// <param name="index">Starting offset of the window of results to return</param>
+		/// <param name="count">Number of results to return</param>
+		/// <param name="filter">Filter for the properties to return</param>
 		/// <returns>List of matching agents</returns>
 		[HttpGet]
 		[Route("/api/v2/issues")]
 		[ProducesResponseType(typeof(List<GetIssueResponse>), 200)]
-		public async Task<ActionResult<object>> FindIssuesV2Async([FromQuery(Name = "Id")] int[]? Ids = null, [FromQuery] StreamId? StreamId = null, [FromQuery] int? MinChange = null, [FromQuery] int? MaxChange = null, [FromQuery] bool? Resolved = null, [FromQuery] int Index = 0, [FromQuery] int Count = 10, [FromQuery] PropertyFilter? Filter = null)
+		public async Task<ActionResult<object>> FindIssuesV2Async([FromQuery(Name = "Id")] int[]? ids = null, [FromQuery] StreamId? streamId = null, [FromQuery] int? minChange = null, [FromQuery] int? maxChange = null, [FromQuery] bool? resolved = null, [FromQuery] int index = 0, [FromQuery] int count = 10, [FromQuery] PropertyFilter? filter = null)
 		{
-			if (Ids != null && Ids.Length == 0)
+			if (ids != null && ids.Length == 0)
 			{
-				Ids = null;
+				ids = null;
 			}
 
-			List<object> Responses = new List<object>();
-			if (StreamId != null)
+			List<object> responses = new List<object>();
+			if (streamId != null)
 			{
-				if (!await StreamService.AuthorizeAsync(StreamId.Value, AclAction.ViewStream, User, new StreamPermissionsCache()))
+				if (!await _streamService.AuthorizeAsync(streamId.Value, AclAction.ViewStream, User, new StreamPermissionsCache()))
 				{
 					return Forbid();
 				}
 
-				List<IIssueSpan> Spans = await IssueCollection.FindSpansAsync(null, Ids, StreamId.Value, MinChange, MaxChange, Resolved);
-				if(Spans.Count > 0)
+				List<IIssueSpan> spans = await _issueCollection.FindSpansAsync(null, ids, streamId.Value, minChange, maxChange, resolved);
+				if(spans.Count > 0)
 				{
 					// Group all the spans by their issue id
-					Dictionary<int, List<IIssueSpan>> IssueIdToSpans = new Dictionary<int, List<IIssueSpan>>();
-					foreach (IIssueSpan Span in Spans)
+					Dictionary<int, List<IIssueSpan>> issueIdToSpans = new Dictionary<int, List<IIssueSpan>>();
+					foreach (IIssueSpan span in spans)
 					{
-						List<IIssueSpan>? SpansForIssue;
-						if (!IssueIdToSpans.TryGetValue(Span.IssueId, out SpansForIssue))
+						List<IIssueSpan>? spansForIssue;
+						if (!issueIdToSpans.TryGetValue(span.IssueId, out spansForIssue))
 						{
-							SpansForIssue = new List<IIssueSpan>();
-							IssueIdToSpans.Add(Span.IssueId, SpansForIssue);
+							spansForIssue = new List<IIssueSpan>();
+							issueIdToSpans.Add(span.IssueId, spansForIssue);
 						}
-						SpansForIssue.Add(Span);
+						spansForIssue.Add(span);
 					}
 
 					// Find the matching issues
-					List<IIssue> Issues = await IssueCollection.FindIssuesAsync(IssueIdToSpans.Keys, Index: Index, Count: Count);
+					List<IIssue> issues = await _issueCollection.FindIssuesAsync(issueIdToSpans.Keys, index: index, count: count);
 
 					// Create the corresponding responses
-					foreach (IIssue Issue in Issues.OrderByDescending(x => x.Id))
+					foreach (IIssue issue in issues.OrderByDescending(x => x.Id))
 					{
-						IssueSeverity StreamSeverity = IssueSeverity.Unspecified;
+						IssueSeverity streamSeverity = IssueSeverity.Unspecified;
 
-						List<FindIssueSpanResponse> SpanResponses = new List<FindIssueSpanResponse>();
-						if (IssueIdToSpans.TryGetValue(Issue.Id, out List<IIssueSpan>? SpansForIssue))
+						List<FindIssueSpanResponse> spanResponses = new List<FindIssueSpanResponse>();
+						if (issueIdToSpans.TryGetValue(issue.Id, out List<IIssueSpan>? spansForIssue))
 						{
 							// Filter issues on resolved state
-							if (Resolved != null && (Resolved.Value != SpansForIssue.All(x => x.NextSuccess != null)))
+							if (resolved != null && (resolved.Value != spansForIssue.All(x => x.NextSuccess != null)))
 							{
 								continue;
 							}
 
 							// Find the current severity in the stream
-							DateTime LastStepTime = DateTime.MinValue;
-							foreach (IIssueSpan Span in SpansForIssue)
+							DateTime lastStepTime = DateTime.MinValue;
+							foreach (IIssueSpan span in spansForIssue)
 							{
-								if (Span.LastFailure != null && Span.LastFailure.StepTime > LastStepTime)
+								if (span.LastFailure != null && span.LastFailure.StepTime > lastStepTime)
 								{
-									LastStepTime = Span.LastFailure.StepTime;
-									StreamSeverity = Span.LastFailure.Severity;
+									lastStepTime = span.LastFailure.StepTime;
+									streamSeverity = span.LastFailure.Severity;
 								}
 							}
 
-							SpanResponses.AddRange(SpansForIssue.Select(x => new FindIssueSpanResponse(x)));
+							spanResponses.AddRange(spansForIssue.Select(x => new FindIssueSpanResponse(x)));
 						}
 					
-						IUser? Owner = null;
-						IUser? NominatedBy = null;
-						IUser? ResolvedBy = null;
+						IUser? owner = null;
+						IUser? nominatedBy = null;
+						IUser? resolvedBy = null;
 
-						if (Issue.OwnerId != null)
+						if (issue.OwnerId != null)
 						{
-							Owner = await UserCollection.GetCachedUserAsync(Issue.OwnerId.Value);
+							owner = await _userCollection.GetCachedUserAsync(issue.OwnerId.Value);
 						}
 						
-						if (Issue.NominatedById != null)
+						if (issue.NominatedById != null)
 						{
-							NominatedBy = await UserCollection.GetCachedUserAsync(Issue.NominatedById.Value);
+							nominatedBy = await _userCollection.GetCachedUserAsync(issue.NominatedById.Value);
 						}
 						
-						if (Issue.ResolvedById != null)
+						if (issue.ResolvedById != null)
 						{
-							ResolvedBy = await UserCollection.GetCachedUserAsync(Issue.ResolvedById.Value);
+							resolvedBy = await _userCollection.GetCachedUserAsync(issue.ResolvedById.Value);
 						}						
 
-						FindIssueResponse Response = new FindIssueResponse(Issue, Owner, NominatedBy, ResolvedBy, StreamSeverity, SpanResponses);
-						Responses.Add(PropertyFilter.Apply(Response, Filter));
+						FindIssueResponse response = new FindIssueResponse(issue, owner, nominatedBy, resolvedBy, streamSeverity, spanResponses);
+						responses.Add(PropertyFilter.Apply(response, filter));
 					}
 				}
 			}
@@ -163,326 +156,326 @@ namespace Horde.Build.Controllers
 				return BadRequest("Missing StreamId on request");
 			}
 
-			return Responses;
+			return responses;
 		}
 
 		/// <summary>
 		/// Retrieve information about a specific issue
 		/// </summary>
-		/// <param name="Ids">Set of issue ids to find</param>
-		/// <param name="StreamId">The stream to query for</param>
-		/// <param name="Change">The changelist to query</param>
-		/// <param name="MinChange">The minimum changelist range to query, inclusive</param>
-		/// <param name="MaxChange">The minimum changelist range to query, inclusive</param>
-		/// <param name="JobId">Job id to filter by</param>
-		/// <param name="BatchId">The batch to filter by</param>
-		/// <param name="StepId">The step to filter by</param>
-		/// <param name="LabelIdx">The label within the job to filter by</param>
-		/// <param name="UserId">User to filter issues for</param>
-		/// <param name="Resolved">Whether to include resolved issues</param>
-		/// <param name="Promoted">Whether to include promoted issues</param>
-		/// <param name="Index">Starting offset of the window of results to return</param>
-		/// <param name="Count">Number of results to return</param>
-		/// <param name="Filter">Filter for the properties to return</param>
+		/// <param name="ids">Set of issue ids to find</param>
+		/// <param name="streamId">The stream to query for</param>
+		/// <param name="change">The changelist to query</param>
+		/// <param name="minChange">The minimum changelist range to query, inclusive</param>
+		/// <param name="maxChange">The minimum changelist range to query, inclusive</param>
+		/// <param name="jobId">Job id to filter by</param>
+		/// <param name="batchId">The batch to filter by</param>
+		/// <param name="stepId">The step to filter by</param>
+		/// <param name="labelIdx">The label within the job to filter by</param>
+		/// <param name="userId">User to filter issues for</param>
+		/// <param name="resolved">Whether to include resolved issues</param>
+		/// <param name="promoted">Whether to include promoted issues</param>
+		/// <param name="index">Starting offset of the window of results to return</param>
+		/// <param name="count">Number of results to return</param>
+		/// <param name="filter">Filter for the properties to return</param>
 		/// <returns>List of matching agents</returns>
 		[HttpGet]
 		[Route("/api/v1/issues")]
 		[ProducesResponseType(typeof(List<GetIssueResponse>), 200)]
-		public async Task<ActionResult<object>> FindIssuesAsync([FromQuery(Name = "Id")] int[]? Ids = null, [FromQuery] string? StreamId = null, [FromQuery] int? Change = null, [FromQuery] int? MinChange = null, [FromQuery] int? MaxChange = null, [FromQuery] JobId? JobId = null, [FromQuery] string? BatchId = null, [FromQuery] string? StepId = null, [FromQuery(Name = "label")] int? LabelIdx = null, [FromQuery] string? UserId = null, [FromQuery] bool? Resolved = null, [FromQuery] bool? Promoted = null, [FromQuery] int Index = 0, [FromQuery] int Count = 10, [FromQuery] PropertyFilter? Filter = null)
+		public async Task<ActionResult<object>> FindIssuesAsync([FromQuery(Name = "Id")] int[]? ids = null, [FromQuery] string? streamId = null, [FromQuery] int? change = null, [FromQuery] int? minChange = null, [FromQuery] int? maxChange = null, [FromQuery] JobId? jobId = null, [FromQuery] string? batchId = null, [FromQuery] string? stepId = null, [FromQuery(Name = "label")] int? labelIdx = null, [FromQuery] string? userId = null, [FromQuery] bool? resolved = null, [FromQuery] bool? promoted = null, [FromQuery] int index = 0, [FromQuery] int count = 10, [FromQuery] PropertyFilter? filter = null)
 		{
-			if(Ids != null && Ids.Length == 0)
+			if(ids != null && ids.Length == 0)
 			{
-				Ids = null;
+				ids = null;
 			}
 
-			UserId? UserIdValue = null;
-			if (UserId != null)
+			UserId? userIdValue = null;
+			if (userId != null)
 			{
-				UserIdValue = new UserId(UserId);
+				userIdValue = new UserId(userId);
 			}
 
-			List<IIssue> Issues;
-			if (JobId == null)
+			List<IIssue> issues;
+			if (jobId == null)
 			{
-				StreamId? StreamIdValue = null;
-				if (StreamId != null)
+				StreamId? streamIdValue = null;
+				if (streamId != null)
 				{
-					StreamIdValue = new StreamId(StreamId);
+					streamIdValue = new StreamId(streamId);
 				}
 
-				Issues = await IssueService.FindIssuesAsync(Ids, UserIdValue, StreamIdValue, MinChange ?? Change, MaxChange ?? Change, Resolved, Promoted, Index, Count);
+				issues = await _issueService.FindIssuesAsync(ids, userIdValue, streamIdValue, minChange ?? change, maxChange ?? change, resolved, promoted, index, count);
 			}
 			else
 			{
-				IJob? Job = await JobService.GetJobAsync(JobId.Value);
-				if (Job == null)
+				IJob? job = await _jobService.GetJobAsync(jobId.Value);
+				if (job == null)
 				{
-					return NotFound(JobId.Value);
+					return NotFound(jobId.Value);
 				}
-				if(!await JobService.AuthorizeAsync(Job, AclAction.ViewJob, this.User, null))
+				if(!await _jobService.AuthorizeAsync(job, AclAction.ViewJob, User, null))
 				{
-					return Forbid(AclAction.ViewJob, JobId.Value);
+					return Forbid(AclAction.ViewJob, jobId.Value);
 				}
 
-				IGraph Graph = await JobService.GetGraphAsync(Job);
-				Issues = await IssueService.FindIssuesForJobAsync(Ids, Job, Graph, StepId?.ToSubResourceId(), BatchId?.ToSubResourceId(), LabelIdx, UserIdValue, Resolved, Promoted, Index, Count);
+				IGraph graph = await _jobService.GetGraphAsync(job);
+				issues = await _issueService.FindIssuesForJobAsync(ids, job, graph, stepId?.ToSubResourceId(), batchId?.ToSubResourceId(), labelIdx, userIdValue, resolved, promoted, index, count);
 			}
 
-			StreamPermissionsCache PermissionsCache = new StreamPermissionsCache();
+			StreamPermissionsCache permissionsCache = new StreamPermissionsCache();
 
-			List<object> Responses = new List<object>();
-			foreach (IIssue Issue in Issues)
+			List<object> responses = new List<object>();
+			foreach (IIssue issue in issues)
 			{
-				IIssueDetails Details = await IssueService.GetIssueDetailsAsync(Issue);
-				if (await AuthorizeIssue(Details, PermissionsCache))
+				IIssueDetails details = await _issueService.GetIssueDetailsAsync(issue);
+				if (await AuthorizeIssue(details, permissionsCache))
 				{
-					bool bShowDesktopAlerts = IssueService.ShowDesktopAlertsForIssue(Issue, Details.Spans);
-					GetIssueResponse Response = await CreateIssueResponseAsync(Details, bShowDesktopAlerts);
-					Responses.Add(PropertyFilter.Apply(Response, Filter));
+					bool bShowDesktopAlerts = _issueService.ShowDesktopAlertsForIssue(issue, details.Spans);
+					GetIssueResponse response = await CreateIssueResponseAsync(details, bShowDesktopAlerts);
+					responses.Add(PropertyFilter.Apply(response, filter));
 				}
 			}
-			return Responses;
+			return responses;
 		}
 
 		/// <summary>
 		/// Retrieve information about a specific issue
 		/// </summary>
-		/// <param name="IssueId">Id of the issue to get information about</param>
-		/// <param name="Filter">Filter for the properties to return</param>
+		/// <param name="issueId">Id of the issue to get information about</param>
+		/// <param name="filter">Filter for the properties to return</param>
 		/// <returns>List of matching agents</returns>
 		[HttpGet]
-		[Route("/api/v1/issues/{IssueId}")]
+		[Route("/api/v1/issues/{issueId}")]
 		[ProducesResponseType(typeof(GetIssueResponse), 200)]
-		public async Task<ActionResult<object>> GetIssueAsync(int IssueId, [FromQuery] PropertyFilter? Filter = null)
+		public async Task<ActionResult<object>> GetIssueAsync(int issueId, [FromQuery] PropertyFilter? filter = null)
 		{
-			IIssueDetails? Details = await IssueService.GetIssueDetailsAsync(IssueId);
-			if (Details == null)
+			IIssueDetails? details = await _issueService.GetIssueDetailsAsync(issueId);
+			if (details == null)
 			{
 				return NotFound();
 			}
-			if (!await AuthorizeIssue(Details, null))
+			if (!await AuthorizeIssue(details, null))
 			{
 				return Forbid();
 			}
 
-			bool bShowDesktopAlerts = IssueService.ShowDesktopAlertsForIssue(Details.Issue, Details.Spans);
-			return PropertyFilter.Apply(await CreateIssueResponseAsync(Details, bShowDesktopAlerts), Filter);
+			bool bShowDesktopAlerts = _issueService.ShowDesktopAlertsForIssue(details.Issue, details.Spans);
+			return PropertyFilter.Apply(await CreateIssueResponseAsync(details, bShowDesktopAlerts), filter);
 		}
 
 		/// <summary>
 		/// Retrieve historical information about a specific issue
 		/// </summary>
-		/// <param name="IssueId">Id of the agent to get information about</param>
-		/// <param name="MinTime">Minimum time for records to return</param>
-		/// <param name="MaxTime">Maximum time for records to return</param>
-		/// <param name="Index">Offset of the first result</param>
-		/// <param name="Count">Number of records to return</param>
+		/// <param name="issueId">Id of the agent to get information about</param>
+		/// <param name="minTime">Minimum time for records to return</param>
+		/// <param name="maxTime">Maximum time for records to return</param>
+		/// <param name="index">Offset of the first result</param>
+		/// <param name="count">Number of records to return</param>
 		/// <returns>Information about the requested agent</returns>
 		[HttpGet]
-		[Route("/api/v1/issues/{IssueId}/history")]
-		public async Task GetAgentHistoryAsync(int IssueId, [FromQuery] DateTime? MinTime = null, [FromQuery] DateTime? MaxTime = null, [FromQuery] int Index = 0, [FromQuery] int Count = 50)
+		[Route("/api/v1/issues/{issueId}/history")]
+		public async Task GetAgentHistoryAsync(int issueId, [FromQuery] DateTime? minTime = null, [FromQuery] DateTime? maxTime = null, [FromQuery] int index = 0, [FromQuery] int count = 50)
 		{
 			Response.ContentType = "application/json";
 			Response.StatusCode = 200;
 			await Response.StartAsync();
-			await IssueCollection.GetLogger(IssueId).FindAsync(Response.BodyWriter, MinTime, MaxTime, Index, Count);
+			await _issueCollection.GetLogger(issueId).FindAsync(Response.BodyWriter, minTime, maxTime, index, count);
 		}
 
 		/// <summary>
 		/// Create an issue response object
 		/// </summary>
-		/// <param name="Details"></param>
-		/// <param name="ShowDesktopAlerts"></param>
+		/// <param name="details"></param>
+		/// <param name="showDesktopAlerts"></param>
 		/// <returns></returns>
-		async Task<GetIssueResponse> CreateIssueResponseAsync(IIssueDetails Details, bool ShowDesktopAlerts)
+		async Task<GetIssueResponse> CreateIssueResponseAsync(IIssueDetails details, bool showDesktopAlerts)
 		{
-			List<GetIssueAffectedStreamResponse> AffectedStreams = new List<GetIssueAffectedStreamResponse>();
-			foreach (IGrouping<StreamId, IIssueSpan> StreamSpans in Details.Spans.GroupBy(x => x.StreamId))
+			List<GetIssueAffectedStreamResponse> affectedStreams = new List<GetIssueAffectedStreamResponse>();
+			foreach (IGrouping<StreamId, IIssueSpan> streamSpans in details.Spans.GroupBy(x => x.StreamId))
 			{
-				IStream? Stream = await StreamService.GetCachedStream(StreamSpans.Key);
-				AffectedStreams.Add(new GetIssueAffectedStreamResponse(Details, Stream, StreamSpans));
+				IStream? stream = await _streamService.GetCachedStream(streamSpans.Key);
+				affectedStreams.Add(new GetIssueAffectedStreamResponse(details, stream, streamSpans));
 			}
-			return new GetIssueResponse(Details, AffectedStreams, ShowDesktopAlerts);
+			return new GetIssueResponse(details, affectedStreams, showDesktopAlerts);
 		}
 
 		/// <summary>
 		/// Retrieve events for a specific issue
 		/// </summary>
-		/// <param name="IssueId">Id of the issue to get information about</param>
-		/// <param name="Filter">Filter for the properties to return</param>
+		/// <param name="issueId">Id of the issue to get information about</param>
+		/// <param name="filter">Filter for the properties to return</param>
 		/// <returns>List of matching agents</returns>
 		[HttpGet]
-		[Route("/api/v1/issues/{IssueId}/streams")]
+		[Route("/api/v1/issues/{issueId}/streams")]
 		[ProducesResponseType(typeof(List<GetIssueStreamResponse>), 200)]
-		public async Task<ActionResult<List<object>>> GetIssueStreamsAsync(int IssueId, [FromQuery] PropertyFilter? Filter = null)
+		public async Task<ActionResult<List<object>>> GetIssueStreamsAsync(int issueId, [FromQuery] PropertyFilter? filter = null)
 		{
-			IIssueDetails? Issue = await IssueService.GetIssueDetailsAsync(IssueId);
-			if (Issue == null)
+			IIssueDetails? issue = await _issueService.GetIssueDetailsAsync(issueId);
+			if (issue == null)
 			{
 				return NotFound();
 			}
 
-			StreamPermissionsCache Cache = new StreamPermissionsCache();
-			if (!await AuthorizeIssue(Issue, Cache))
+			StreamPermissionsCache cache = new StreamPermissionsCache();
+			if (!await AuthorizeIssue(issue, cache))
 			{
 				return Forbid();
 			}
 
-			List<object> Responses = new List<object>();
-			foreach (IGrouping<StreamId, IIssueSpan> SpanGroup in Issue.Spans.GroupBy(x => x.StreamId))
+			List<object> responses = new List<object>();
+			foreach (IGrouping<StreamId, IIssueSpan> spanGroup in issue.Spans.GroupBy(x => x.StreamId))
 			{
-				if (await StreamService.AuthorizeAsync(SpanGroup.Key, AclAction.ViewStream, User, Cache))
+				if (await _streamService.AuthorizeAsync(spanGroup.Key, AclAction.ViewStream, User, cache))
 				{
-					HashSet<ObjectId> SpanIds = new HashSet<ObjectId>(SpanGroup.Select(x => x.Id));
-					List<IIssueStep> Steps = Issue.Steps.Where(x => SpanIds.Contains(x.SpanId)).ToList();
-					Responses.Add(PropertyFilter.Apply(new GetIssueStreamResponse(SpanGroup.Key, SpanGroup.ToList(), Steps), Filter));
+					HashSet<ObjectId> spanIds = new HashSet<ObjectId>(spanGroup.Select(x => x.Id));
+					List<IIssueStep> steps = issue.Steps.Where(x => spanIds.Contains(x.SpanId)).ToList();
+					responses.Add(PropertyFilter.Apply(new GetIssueStreamResponse(spanGroup.Key, spanGroup.ToList(), steps), filter));
 				}
 			}
-			return Responses;
+			return responses;
 		}
 
 		/// <summary>
 		/// Retrieve events for a specific issue
 		/// </summary>
-		/// <param name="IssueId">Id of the issue to get information about</param>
-		/// <param name="StreamId">The stream id</param>
-		/// <param name="Filter">Filter for the properties to return</param>
+		/// <param name="issueId">Id of the issue to get information about</param>
+		/// <param name="streamId">The stream id</param>
+		/// <param name="filter">Filter for the properties to return</param>
 		/// <returns>List of matching agents</returns>
 		[HttpGet]
-		[Route("/api/v1/issues/{IssueId}/streams/{StreamId}")]
+		[Route("/api/v1/issues/{issueId}/streams/{streamId}")]
 		[ProducesResponseType(typeof(List<GetIssueStreamResponse>), 200)]
-		public async Task<ActionResult<object>> GetIssueStreamAsync(int IssueId, string StreamId, [FromQuery] PropertyFilter? Filter = null)
+		public async Task<ActionResult<object>> GetIssueStreamAsync(int issueId, string streamId, [FromQuery] PropertyFilter? filter = null)
 		{
-			IIssueDetails? Details = await IssueService.GetIssueDetailsAsync(IssueId);
-			if (Details == null)
+			IIssueDetails? details = await _issueService.GetIssueDetailsAsync(issueId);
+			if (details == null)
 			{
 				return NotFound();
 			}
 
-			StreamId StreamIdValue = new StreamId(StreamId);
-			if (!await StreamService.AuthorizeAsync(StreamIdValue, AclAction.ViewStream, User, null))
+			StreamId streamIdValue = new StreamId(streamId);
+			if (!await _streamService.AuthorizeAsync(streamIdValue, AclAction.ViewStream, User, null))
 			{
 				return Forbid();
 			}
 
-			List<IIssueSpan> Spans = Details.Spans.Where(x => x.StreamId == StreamIdValue).ToList();
-			if(Spans.Count == 0)
+			List<IIssueSpan> spans = details.Spans.Where(x => x.StreamId == streamIdValue).ToList();
+			if(spans.Count == 0)
 			{
 				return NotFound();
 			}
 
-			HashSet<ObjectId> SpanIds = new HashSet<ObjectId>(Spans.Select(x => x.Id));
-			List<IIssueStep> Steps = Details.Steps.Where(x => SpanIds.Contains(x.SpanId)).ToList();
+			HashSet<ObjectId> spanIds = new HashSet<ObjectId>(spans.Select(x => x.Id));
+			List<IIssueStep> steps = details.Steps.Where(x => spanIds.Contains(x.SpanId)).ToList();
 
-			return PropertyFilter.Apply(new GetIssueStreamResponse(StreamIdValue, Spans, Steps), Filter);
+			return PropertyFilter.Apply(new GetIssueStreamResponse(streamIdValue, spans, steps), filter);
 		}
 
 		/// <summary>
 		/// Retrieve events for a specific issue
 		/// </summary>
-		/// <param name="IssueId">Id of the issue to get information about</param>
-		/// <param name="JobId">The job id to filter for</param>
-		/// <param name="BatchId">The batch to filter by</param>
-		/// <param name="StepId">The step to filter by</param>
-		/// <param name="LabelIdx">The label within the job to filter by</param>
-		/// <param name="LogIds">List of log ids to return issues for</param>
-		/// <param name="Index">Index of the first event</param>
-		/// <param name="Count">Number of events to return</param>
-		/// <param name="Filter">Filter for the properties to return</param>
+		/// <param name="issueId">Id of the issue to get information about</param>
+		/// <param name="jobId">The job id to filter for</param>
+		/// <param name="batchId">The batch to filter by</param>
+		/// <param name="stepId">The step to filter by</param>
+		/// <param name="labelIdx">The label within the job to filter by</param>
+		/// <param name="logIds">List of log ids to return issues for</param>
+		/// <param name="index">Index of the first event</param>
+		/// <param name="count">Number of events to return</param>
+		/// <param name="filter">Filter for the properties to return</param>
 		/// <returns>List of matching agents</returns>
 		[HttpGet]
-		[Route("/api/v1/issues/{IssueId}/events")]
+		[Route("/api/v1/issues/{issueId}/events")]
 		[ProducesResponseType(typeof(List<GetLogEventResponse>), 200)]
-		public async Task<ActionResult<List<object>>> GetIssueEventsAsync(int IssueId, [FromQuery] JobId? JobId = null, [FromQuery] string? BatchId = null, [FromQuery] string? StepId = null, [FromQuery(Name = "label")] int? LabelIdx = null, [FromQuery] string[]? LogIds = null, [FromQuery] int Index = 0, [FromQuery] int Count = 10, [FromQuery] PropertyFilter? Filter = null)
+		public async Task<ActionResult<List<object>>> GetIssueEventsAsync(int issueId, [FromQuery] JobId? jobId = null, [FromQuery] string? batchId = null, [FromQuery] string? stepId = null, [FromQuery(Name = "label")] int? labelIdx = null, [FromQuery] string[]? logIds = null, [FromQuery] int index = 0, [FromQuery] int count = 10, [FromQuery] PropertyFilter? filter = null)
 		{
-			HashSet<LogId> LogIdValues = new HashSet<LogId>();
-			if(JobId != null)
+			HashSet<LogId> logIdValues = new HashSet<LogId>();
+			if(jobId != null)
 			{
-				IJob? Job = await JobService.GetJobAsync(JobId.Value);
-				if(Job == null)
+				IJob? job = await _jobService.GetJobAsync(jobId.Value);
+				if(job == null)
 				{
 					return NotFound();
 				}
 
-				if (StepId != null)
+				if (stepId != null)
 				{
-					IJobStep? Step;
-					if (Job.TryGetStep(StepId.ToSubResourceId(), out Step) && Step.Outcome != JobStepOutcome.Success && Step.LogId != null)
+					IJobStep? step;
+					if (job.TryGetStep(stepId.ToSubResourceId(), out step) && step.Outcome != JobStepOutcome.Success && step.LogId != null)
 					{
-						LogIdValues.Add(Step.LogId.Value);
+						logIdValues.Add(step.LogId.Value);
 					}
 				}
-				else if (BatchId != null)
+				else if (batchId != null)
 				{
-					IJobStepBatch? Batch;
-					if (Job.TryGetBatch(BatchId.ToSubResourceId(), out Batch))
+					IJobStepBatch? batch;
+					if (job.TryGetBatch(batchId.ToSubResourceId(), out batch))
 					{
-						LogIdValues.UnionWith(Batch.Steps.Where(x => x.Outcome != JobStepOutcome.Success && x.LogId != null).Select(x => x.LogId!.Value));
+						logIdValues.UnionWith(batch.Steps.Where(x => x.Outcome != JobStepOutcome.Success && x.LogId != null).Select(x => x.LogId!.Value));
 					}
 				}
-				else if (LabelIdx != null)
+				else if (labelIdx != null)
 				{
-					IGraph Graph = await JobService.GetGraphAsync(Job);
+					IGraph graph = await _jobService.GetGraphAsync(job);
 
-					HashSet<NodeRef> IncludedNodes = new HashSet<NodeRef>(Graph.Labels[LabelIdx.Value].IncludedNodes);
+					HashSet<NodeRef> includedNodes = new HashSet<NodeRef>(graph.Labels[labelIdx.Value].IncludedNodes);
 
-					foreach (IJobStepBatch Batch in Job.Batches)
+					foreach (IJobStepBatch batch in job.Batches)
 					{
-						foreach (IJobStep Step in Batch.Steps)
+						foreach (IJobStep step in batch.Steps)
 						{
-							NodeRef NodeRef = new NodeRef(Batch.GroupIdx, Step.NodeIdx);
-							if (Step.Outcome != JobStepOutcome.Success && Step.LogId != null && IncludedNodes.Contains(NodeRef))
+							NodeRef nodeRef = new NodeRef(batch.GroupIdx, step.NodeIdx);
+							if (step.Outcome != JobStepOutcome.Success && step.LogId != null && includedNodes.Contains(nodeRef))
 							{
-								LogIdValues.Add(Step.LogId.Value);
+								logIdValues.Add(step.LogId.Value);
 							}
 						}
 					}
 				}
 				else
 				{
-					LogIdValues.UnionWith(Job.Batches.SelectMany(x => x.Steps).Where(x => x.Outcome != JobStepOutcome.Success && x.LogId != null).Select(x => x.LogId!.Value));
+					logIdValues.UnionWith(job.Batches.SelectMany(x => x.Steps).Where(x => x.Outcome != JobStepOutcome.Success && x.LogId != null).Select(x => x.LogId!.Value));
 				}
 			}
-			if(LogIds != null)
+			if(logIds != null)
 			{
-				LogIdValues.UnionWith(LogIds.Select(x => new LogId(x)));
+				logIdValues.UnionWith(logIds.Select(x => new LogId(x)));
 			}
 
-			List<ILogEvent> Events = await IssueService.FindEventsForIssueAsync(IssueId, LogIdValues.ToArray(), Index, Count);
+			List<ILogEvent> events = await _issueService.FindEventsForIssueAsync(issueId, logIdValues.ToArray(), index, count);
 
-			JobPermissionsCache PermissionsCache = new JobPermissionsCache();
-			Dictionary<LogId, ILogFile?> LogFiles = new Dictionary<LogId, ILogFile?>();
+			JobPermissionsCache permissionsCache = new JobPermissionsCache();
+			Dictionary<LogId, ILogFile?> logFiles = new Dictionary<LogId, ILogFile?>();
 
-			List<object> Responses = new List<object>();
-			foreach (ILogEvent Event in Events)
+			List<object> responses = new List<object>();
+			foreach (ILogEvent logEvent in events)
 			{
-				ILogFile? LogFile;
-				if (!LogFiles.TryGetValue(Event.LogId, out LogFile))
+				ILogFile? logFile;
+				if (!logFiles.TryGetValue(logEvent.LogId, out logFile))
 				{
-					LogFile = await LogFileService.GetLogFileAsync(Event.LogId);
-					LogFiles[Event.LogId] = LogFile;
+					logFile = await _logFileService.GetLogFileAsync(logEvent.LogId);
+					logFiles[logEvent.LogId] = logFile;
 				}
-				if (LogFile != null && await JobService.AuthorizeAsync(LogFile.JobId, AclAction.ViewLog, User, PermissionsCache))
+				if (logFile != null && await _jobService.AuthorizeAsync(logFile.JobId, AclAction.ViewLog, User, permissionsCache))
 				{
-					ILogEventData Data = await LogFileService.GetEventDataAsync(LogFile, Event.LineIndex, Event.LineCount);
-					GetLogEventResponse Response = new GetLogEventResponse(Event, Data, IssueId);
-					Responses.Add(PropertyFilter.Apply(Response, Filter));
+					ILogEventData data = await _logFileService.GetEventDataAsync(logFile, logEvent.LineIndex, logEvent.LineCount);
+					GetLogEventResponse response = new GetLogEventResponse(logEvent, data, issueId);
+					responses.Add(PropertyFilter.Apply(response, filter));
 				}
 			}
-			return Responses;
+			return responses;
 		}
 
 		/// <summary>
 		/// Authorize the current user to see an issue
 		/// </summary>
-		/// <param name="Issue">The issue to authorize</param>
-		/// <param name="PermissionsCache">Cache of permissions</param>
+		/// <param name="issue">The issue to authorize</param>
+		/// <param name="permissionsCache">Cache of permissions</param>
 		/// <returns>True if the user is authorized to see the issue</returns>
-		private async Task<bool> AuthorizeIssue(IIssueDetails Issue, StreamPermissionsCache? PermissionsCache)
+		private async Task<bool> AuthorizeIssue(IIssueDetails issue, StreamPermissionsCache? permissionsCache)
 		{
-			foreach (StreamId StreamId in Issue.Spans.Select(x => x.StreamId).Distinct())
+			foreach (StreamId streamId in issue.Spans.Select(x => x.StreamId).Distinct())
 			{
-				if (await StreamService.AuthorizeAsync(StreamId, AclAction.ViewStream, User, PermissionsCache))
+				if (await _streamService.AuthorizeAsync(streamId, AclAction.ViewStream, User, permissionsCache))
 				{
 					return true;
 				}
@@ -493,50 +486,50 @@ namespace Horde.Build.Controllers
 		/// <summary>
 		/// Update an issue
 		/// </summary>
-		/// <param name="IssueId">Id of the issue to get information about</param>
-		/// <param name="Request">The update information</param>
+		/// <param name="issueId">Id of the issue to get information about</param>
+		/// <param name="request">The update information</param>
 		/// <returns>List of matching agents</returns>
 		[HttpPut]
-		[Route("/api/v1/issues/{IssueId}")]
-		public async Task<ActionResult> UpdateIssueAsync(int IssueId, [FromBody] UpdateIssueRequest Request)
+		[Route("/api/v1/issues/{issueId}")]
+		public async Task<ActionResult> UpdateIssueAsync(int issueId, [FromBody] UpdateIssueRequest request)
 		{
-			UserId? NewOwnerId = null;
-			if (Request.OwnerId != null)
+			UserId? newOwnerId = null;
+			if (request.OwnerId != null)
 			{
-				NewOwnerId = Request.OwnerId.Length == 0 ? UserId.Empty : new UserId(Request.OwnerId);
+				newOwnerId = request.OwnerId.Length == 0 ? UserId.Empty : new UserId(request.OwnerId);
 			}
 
-			UserId? NewNominatedById = null;
-			if (Request.NominatedById != null)
+			UserId? newNominatedById = null;
+			if (request.NominatedById != null)
 			{
-				NewNominatedById = new UserId(Request.NominatedById);
+				newNominatedById = new UserId(request.NominatedById);
 			}
 
-			UserId? NewDeclinedById = null;
-			if (Request.Declined ?? false)
+			UserId? newDeclinedById = null;
+			if (request.Declined ?? false)
 			{
-				NewDeclinedById = User.GetUserId();
+				newDeclinedById = User.GetUserId();
 			}
 
-			UserId? NewResolvedById = null;
-			if (Request.Resolved.HasValue)
+			UserId? newResolvedById = null;
+			if (request.Resolved.HasValue)
 			{
-				NewResolvedById = Request.Resolved.Value ? User.GetUserId() : UserId.Empty;
+				newResolvedById = request.Resolved.Value ? User.GetUserId() : UserId.Empty;
 			}
 
-			List<ObjectId>? AddSpans = null;
-			if (Request.AddSpans != null && Request.AddSpans.Count > 0)
+			List<ObjectId>? addSpans = null;
+			if (request.AddSpans != null && request.AddSpans.Count > 0)
 			{
-				AddSpans = Request.AddSpans.ConvertAll(x => ObjectId.Parse(x));
+				addSpans = request.AddSpans.ConvertAll(x => ObjectId.Parse(x));
 			}
 
-			List<ObjectId>? RemoveSpans = null;
-			if (Request.RemoveSpans != null && Request.RemoveSpans.Count > 0)
+			List<ObjectId>? removeSpans = null;
+			if (request.RemoveSpans != null && request.RemoveSpans.Count > 0)
 			{
-				RemoveSpans = Request.RemoveSpans.ConvertAll(x => ObjectId.Parse(x));
+				removeSpans = request.RemoveSpans.ConvertAll(x => ObjectId.Parse(x));
 			}
 
-			if (!await IssueService.UpdateIssueAsync(IssueId, Request.Summary, Request.Description, Request.Promoted, NewOwnerId, NewNominatedById, Request.Acknowledged, NewDeclinedById, Request.FixChange, NewResolvedById, AddSpans, RemoveSpans))
+			if (!await _issueService.UpdateIssueAsync(issueId, request.Summary, request.Description, request.Promoted, newOwnerId, newNominatedById, request.Acknowledged, newDeclinedById, request.FixChange, newResolvedById, addSpans, removeSpans))
 			{
 				return NotFound();
 			}
diff --git a/Engine/Source/Programs/Horde/Horde.Build/Controllers/JiraController.cs b/Engine/Source/Programs/Horde/Horde.Build/Controllers/JiraController.cs
--- a/Engine/Source/Programs/Horde/Horde.Build/Controllers/JiraController.cs
+++ b/Engine/Source/Programs/Horde/Horde.Build/Controllers/JiraController.cs
@@ -1,15 +1,14 @@
-﻿// Copyright Epic Games, Inc. All Rights Reserved.
+// Copyright Epic Games, Inc. All Rights Reserved.
 
-using Microsoft.AspNetCore.Authorization;
-using Microsoft.AspNetCore.Mvc;
+using System.Collections.Generic;
 using System.Threading.Tasks;
 using Horde.Build.Acls;
-using Horde.Build.Services;
-using Microsoft.Extensions.Logging;
 using Horde.Build.Api;
-using System.Collections.Generic;
 using Horde.Build.Models;
+using Horde.Build.Services;
 using Horde.Build.Utilities;
+using Microsoft.AspNetCore.Authorization;
+using Microsoft.AspNetCore.Mvc;
 
 namespace Horde.Build
 {
@@ -25,26 +24,20 @@ namespace Horde.Build
 		/// <summary>
 		/// Singleton instance of the jira service
 		/// </summary>
-		JiraService JiraService;
+		readonly JiraService _jiraService;
 
 		/// <summary>
 		/// Singleton instance of the stream service
 		/// </summary>
-		StreamService StreamService;
-
-		/// <summary>
-		///  Logger for controller
-		/// </summary>
-		private readonly ILogger<JiraController> Logger;
+		readonly StreamService _streamService;
 
 		/// <summary>
 		/// Constructor
 		/// </summary>
-		public JiraController(JiraService JiraService, StreamService StreamService, ILogger<JiraController> Logger)
+		public JiraController(JiraService jiraService, StreamService streamService)
 		{
-			this.JiraService = JiraService;
-			this.Logger = Logger;
-			this.StreamService = StreamService;
+			_jiraService = jiraService;
+			_streamService = streamService;
 		}
 
 		/// <summary>
@@ -53,33 +46,30 @@ namespace Horde.Build
 		[HttpGet]
 		[Authorize]
 		[Route("/api/v1/jira")]
-		public async Task<ActionResult<List<GetJiraIssueResponse>>> GetJiraIssuesAsync([FromQuery] string StreamId, [FromQuery] string[] JiraKeys)
+		public async Task<ActionResult<List<GetJiraIssueResponse>>> GetJiraIssuesAsync([FromQuery] string streamId, [FromQuery] string[] jiraKeys)
 		{
 
-			StreamId StreamIdValue = new StreamId(StreamId);
-			StreamPermissionsCache Cache = new StreamPermissionsCache();
+			StreamId streamIdValue = new StreamId(streamId);
+			StreamPermissionsCache cache = new StreamPermissionsCache();
 
-			if (!await StreamService.AuthorizeAsync(StreamIdValue, AclAction.ViewStream, User, Cache))
+			if (!await _streamService.AuthorizeAsync(streamIdValue, AclAction.ViewStream, User, cache))
 			{
 				return Forbid();
 			}
 
-			List<GetJiraIssueResponse> Response = new List<GetJiraIssueResponse>();
+			List<GetJiraIssueResponse> response = new List<GetJiraIssueResponse>();
 
-			if (JiraKeys.Length != 0)
+			if (jiraKeys.Length != 0)
 			{
-				List<JiraIssue> Issues = await this.JiraService.GetJiraIssuesAsync(JiraKeys);
+				List<JiraIssue> issues = await _jiraService.GetJiraIssuesAsync(jiraKeys);
 
-				for (int i = 0; i < Issues.Count; i++)
+				for (int i = 0; i < issues.Count; i++)
 				{
-					Response.Add(new GetJiraIssueResponse(Issues[i]));
+					response.Add(new GetJiraIssueResponse(issues[i]));
 				}
-
 			}			
 
-			return Response;
+			return response;
 		}
-
-
 	}
 }
