Horde: Fix coding conventions within Horde.Build.

diff --git a/Engine/Source/Programs/Horde/Horde.Build/Controllers/JobsController.cs b/Engine/Source/Programs/Horde/Horde.Build/Controllers/JobsController.cs
--- a/Engine/Source/Programs/Horde/Horde.Build/Controllers/JobsController.cs
+++ b/Engine/Source/Programs/Horde/Horde.Build/Controllers/JobsController.cs
@@ -1,31 +1,26 @@
-﻿// Copyright Epic Games, Inc. All Rights Reserved.
+// Copyright Epic Games, Inc. All Rights Reserved.
 
+using System;
+using System.Collections.Generic;
+using System.Linq;
+using System.Threading.Tasks;
 using EpicGames.Core;
+using EpicGames.Perforce;
 using Horde.Build.Acls;
-using Horde.Build.Utilities;
 using Horde.Build.Api;
 using Horde.Build.Collections;
-using HordeCommon;
+using Horde.Build.Jobs;
 using Horde.Build.Models;
+using Horde.Build.Notifications;
 using Horde.Build.Services;
+using Horde.Build.Utilities;
+using HordeCommon;
 using Microsoft.AspNetCore.Authorization;
 using Microsoft.AspNetCore.Mvc;
-using MongoDB.Bson;
-using System;
-using System.Collections.Generic;
-using System.ComponentModel.DataAnnotations;
-using System.Globalization;
-using System.Linq;
-using System.Security.Claims;
-using System.Text.RegularExpressions;
-using System.Threading.Tasks;
-
 using Microsoft.Extensions.Logging;
-using Horde.Build.Notifications;
-using OpenTracing.Util;
+using MongoDB.Bson;
 using OpenTracing;
-using Horde.Build.Jobs;
-using EpicGames.Perforce;
+using OpenTracing.Util;
 
 namespace Horde.Build.Controllers
 {
@@ -42,216 +37,214 @@ namespace Horde.Build.Controllers
 	[Route("[controller]")]
 	public class JobsController : HordeControllerBase
 	{
-		private readonly AclService AclService;
-		private readonly IGraphCollection Graphs;
-		private readonly IPerforceService Perforce;
-		private readonly StreamService StreamService;
-		private readonly JobService JobService;
-		private readonly ITemplateCollection TemplateCollection;
-		private readonly IArtifactCollection ArtifactCollection;
-		private readonly IUserCollection UserCollection;
-		private readonly INotificationService NotificationService;
-		private readonly AgentService AgentService;
-		private ILogger<JobsController> Logger;
+		private readonly IGraphCollection _graphs;
+		private readonly IPerforceService _perforce;
+		private readonly StreamService _streamService;
+		private readonly JobService _jobService;
+		private readonly ITemplateCollection _templateCollection;
+		private readonly IArtifactCollection _artifactCollection;
+		private readonly IUserCollection _userCollection;
+		private readonly INotificationService _notificationService;
+		private readonly AgentService _agentService;
+		private readonly ILogger<JobsController> _logger;
 
 		/// <summary>
 		/// Constructor
 		/// </summary>
-		public JobsController(AclService AclService, IGraphCollection Graphs, IPerforceService Perforce, StreamService StreamService, JobService JobService, ITemplateCollection TemplateCollection, IArtifactCollection ArtifactCollection, IUserCollection UserCollection, INotificationService NotificationService, AgentService AgentService, ILogger<JobsController> Logger)
+		public JobsController(IGraphCollection graphs, IPerforceService perforce, StreamService streamService, JobService jobService, ITemplateCollection templateCollection, IArtifactCollection artifactCollection, IUserCollection userCollection, INotificationService notificationService, AgentService agentService, ILogger<JobsController> logger)
 		{
-			this.AclService = AclService;
-			this.Graphs = Graphs;
-			this.Perforce = Perforce;
-			this.StreamService = StreamService;
-			this.JobService = JobService;
-			this.TemplateCollection = TemplateCollection;
-			this.ArtifactCollection = ArtifactCollection;
-			this.UserCollection = UserCollection;
-			this.NotificationService = NotificationService;
-			this.AgentService = AgentService;
-			this.Logger = Logger;
+			_graphs = graphs;
+			_perforce = perforce;
+			_streamService = streamService;
+			_jobService = jobService;
+			_templateCollection = templateCollection;
+			_artifactCollection = artifactCollection;
+			_userCollection = userCollection;
+			_notificationService = notificationService;
+			_agentService = agentService;
+			_logger = logger;
 		}
 
 		/// <summary>
 		/// Creates a new job
 		/// </summary>
-		/// <param name="Create">Properties of the new job</param>
+		/// <param name="create">Properties of the new job</param>
 		/// <returns>Id of the new job</returns>
 		[HttpPost]
 		[Route("/api/v1/jobs")]
-		public async Task<ActionResult<CreateJobResponse>> CreateJobAsync([FromBody] CreateJobRequest Create)
+		public async Task<ActionResult<CreateJobResponse>> CreateJobAsync([FromBody] CreateJobRequest create)
 		{
-			IStream? Stream = await StreamService.GetStreamAsync(new StreamId(Create.StreamId));
-			if (Stream == null)
+			IStream? stream = await _streamService.GetStreamAsync(new StreamId(create.StreamId));
+			if (stream == null)
 			{
-				return NotFound(Create.StreamId);
+				return NotFound(create.StreamId);
 			}
-			if (!await StreamService.AuthorizeAsync(Stream, AclAction.CreateJob, User, null))
+			if (!await _streamService.AuthorizeAsync(stream, AclAction.CreateJob, User, null))
 			{
-				return Forbid(AclAction.CreateJob, Stream.Id);
+				return Forbid(AclAction.CreateJob, stream.Id);
 			}
 
 			// Get the name of the template ref
-			TemplateRefId TemplateRefId = new TemplateRefId(Create.TemplateId);
+			TemplateRefId templateRefId = new TemplateRefId(create.TemplateId);
 
 			// Augment the request with template properties
-			TemplateRef? TemplateRef;
-			if (!Stream.Templates.TryGetValue(TemplateRefId, out TemplateRef))
+			TemplateRef? templateRef;
+			if (!stream.Templates.TryGetValue(templateRefId, out templateRef))
 			{
-				return BadRequest($"Invalid {Create.TemplateId} parameter");
+				return BadRequest($"Invalid {create.TemplateId} parameter");
 			}
-			if (!await StreamService.AuthorizeAsync(Stream, TemplateRef, AclAction.CreateJob, User, null))
+			if (!await _streamService.AuthorizeAsync(stream, templateRef, AclAction.CreateJob, User, null))
 			{
-				return Forbid(AclAction.CreateJob, Stream.Id);
+				return Forbid(AclAction.CreateJob, stream.Id);
 			}
 
-			ITemplate? Template = await TemplateCollection.GetAsync(TemplateRef.Hash);
-			if (Template == null)
+			ITemplate? template = await _templateCollection.GetAsync(templateRef.Hash);
+			if (template == null)
 			{
-				return BadRequest("Missing template referenced by {TemplateId}", Create.TemplateId);
+				return BadRequest("Missing template referenced by {TemplateId}", create.TemplateId);
 			}
-			if (!Template.AllowPreflights && Create.PreflightChange > 0)
+			if (!template.AllowPreflights && create.PreflightChange > 0)
 			{
-				return BadRequest("Template {TemplateId} does not allow preflights", Create.TemplateId);
+				return BadRequest("Template {TemplateId} does not allow preflights", create.TemplateId);
 			}
 
 			// Get the name of the new job
-			string Name = Create.Name ?? Template.Name;
-			if (Create.TemplateId.Equals("stage-to-marketplace", StringComparison.Ordinal) && Create.Arguments != null)
+			string name = create.Name ?? template.Name;
+			if (create.TemplateId.Equals("stage-to-marketplace", StringComparison.Ordinal) && create.Arguments != null)
 			{
-				foreach (string Argument in Create.Arguments)
+				foreach (string argument in create.Arguments)
 				{
 					const string Prefix = "-set:UserContentItems=";
-					if (Argument.StartsWith(Prefix, StringComparison.Ordinal))
+					if (argument.StartsWith(Prefix, StringComparison.Ordinal))
 					{
-						Name += $" - {Argument.Substring(Prefix.Length)}";
+						name += $" - {argument.Substring(Prefix.Length)}";
 						break;
 					}
 				}
 			}
 
 			// Get the priority of the new job
-			Priority Priority = Create.Priority ?? Template.Priority ?? Priority.Normal;
+			Priority priority = create.Priority ?? template.Priority ?? Priority.Normal;
 
 			// New groups for the job
-			IGraph Graph = await Graphs.AddAsync(Template);
+			IGraph graph = await _graphs.AddAsync(template);
 
 			// Get the change to build
-			int Change;
-			if (Create.Change.HasValue)
+			int change;
+			if (create.Change.HasValue)
 			{
-				Change = Create.Change.Value;
+				change = create.Change.Value;
 			}
-			else if (Create.ChangeQuery != null)
+			else if (create.ChangeQuery != null)
 			{
-				Change = await ExecuteChangeQueryAsync(Stream, new TemplateRefId(Create.ChangeQuery.TemplateId ?? Create.TemplateId), Create.ChangeQuery.Target, Create.ChangeQuery.Outcomes ?? new List<JobStepOutcome> { JobStepOutcome.Success });
+				change = await ExecuteChangeQueryAsync(stream, new TemplateRefId(create.ChangeQuery.TemplateId ?? create.TemplateId), create.ChangeQuery.Target, create.ChangeQuery.Outcomes ?? new List<JobStepOutcome> { JobStepOutcome.Success });
 			}
-			else if (Create.PreflightChange == null && Template.SubmitNewChange != null)
+			else if (create.PreflightChange == null && template.SubmitNewChange != null)
 			{
-				Change = await Perforce.CreateNewChangeForTemplateAsync(Stream, Template);
+				change = await _perforce.CreateNewChangeForTemplateAsync(stream, template);
 			}
 			else
 			{
-				Change = await Perforce.GetLatestChangeAsync(Stream.ClusterName, Stream.Name, null);
+				change = await _perforce.GetLatestChangeAsync(stream.ClusterName, stream.Name, null);
 			}
 
 			// And get the matching code changelist
-			int CodeChange = await Perforce.GetCodeChangeAsync(Stream.ClusterName, Stream.Name, Change);
+			int codeChange = await _perforce.GetCodeChangeAsync(stream.ClusterName, stream.Name, change);
 
 			// New properties for the job
-			List<string> Arguments = Create.Arguments ?? Template.GetDefaultArguments();
+			List<string> arguments = create.Arguments ?? template.GetDefaultArguments();
 
 			// Special handling for internal graph evaluation
-			if (Arguments.Any(x => x.Equals("-InternalParser", StringComparison.OrdinalIgnoreCase)))
+			if (arguments.Any(x => x.Equals("-InternalParser", StringComparison.OrdinalIgnoreCase)))
 			{
-				IPerforceConnection? PerforceConnection = await Perforce.GetServiceUserConnection(Stream.ClusterName);
-				if (PerforceConnection != null)
+				IPerforceConnection? perforceConnection = await _perforce.GetServiceUserConnection(stream.ClusterName);
+				if (perforceConnection != null)
 				{
-					Graph = await Graphs.AddAsync(PerforceConnection, Stream, Change, CodeChange, Arguments);
+					graph = await _graphs.AddAsync(perforceConnection, stream, change, codeChange, arguments);
 				}
 			}
 
 			// Check the preflight change is valid
-			if (Create.PreflightChange != null)
+			if (create.PreflightChange != null)
 			{
-				CheckShelfResult Result = await Perforce.CheckShelfAsync(Stream.ClusterName, Stream.Name, Create.PreflightChange.Value, null);
-				switch (Result)
+				CheckShelfResult result = await _perforce.CheckShelfAsync(stream.ClusterName, stream.Name, create.PreflightChange.Value, null);
+				switch (result)
 				{
 					case CheckShelfResult.Ok:
 						break;
 					case CheckShelfResult.NoChange:
-						return BadRequest(KnownLogEvents.Horde_InvalidPreflight, "CL {Change} does not exist", Create.PreflightChange);
+						return BadRequest(KnownLogEvents.Horde_InvalidPreflight, "CL {Change} does not exist", create.PreflightChange);
 					case CheckShelfResult.NoShelvedFiles:
-						return BadRequest(KnownLogEvents.Horde_InvalidPreflight, "CL {Change} does not contain any shelved files", Create.PreflightChange);
+						return BadRequest(KnownLogEvents.Horde_InvalidPreflight, "CL {Change} does not contain any shelved files", create.PreflightChange);
 					case CheckShelfResult.WrongStream:
-						return BadRequest(KnownLogEvents.Horde_InvalidPreflight, "CL {Change} does not contain files in {Stream}", Create.PreflightChange, Stream.Name);
+						return BadRequest(KnownLogEvents.Horde_InvalidPreflight, "CL {Change} does not contain files in {Stream}", create.PreflightChange, stream.Name);
 					case CheckShelfResult.MixedStream:
-						return BadRequest(KnownLogEvents.Horde_InvalidPreflight, "CL {Change} contains files from multiple streams", Create.PreflightChange);
+						return BadRequest(KnownLogEvents.Horde_InvalidPreflight, "CL {Change} contains files from multiple streams", create.PreflightChange);
 					default:
-						return BadRequest(KnownLogEvents.Horde_InvalidPreflight, "CL {Change} cannot be preflighted ({Result})", Create.PreflightChange, Result);
+						return BadRequest(KnownLogEvents.Horde_InvalidPreflight, "CL {Change} cannot be preflighted ({Result})", create.PreflightChange, result);
 				}
 			}
 
-			bool? UpdateIssues = null;
-			if (Template.UpdateIssues)
+			bool? updateIssues = null;
+			if (template.UpdateIssues)
 			{
-				UpdateIssues = true;
+				updateIssues = true;
 			}
-			else if (Create.UpdateIssues.HasValue)
+			else if (create.UpdateIssues.HasValue)
 			{
-				UpdateIssues = Create.UpdateIssues.Value;
+				updateIssues = create.UpdateIssues.Value;
 			}
 
 			// Create the job
-			IJob Job = await JobService.CreateJobAsync(null, Stream, TemplateRefId, Template.Id, Graph, Name, Change, CodeChange, Create.PreflightChange, null, User.GetUserId(), Priority, Create.AutoSubmit, UpdateIssues, false, TemplateRef.ChainedJobs, TemplateRef.ShowUgsBadges, TemplateRef.ShowUgsAlerts, TemplateRef.NotificationChannel, TemplateRef.NotificationChannelFilter, Arguments);
-			await UpdateNotificationsAsync(Job.Id, new UpdateNotificationsRequest { Slack = true });
-			return new CreateJobResponse(Job.Id.ToString());
+			IJob job = await _jobService.CreateJobAsync(null, stream, templateRefId, template.Id, graph, name, change, codeChange, create.PreflightChange, null, User.GetUserId(), priority, create.AutoSubmit, updateIssues, false, templateRef.ChainedJobs, templateRef.ShowUgsBadges, templateRef.ShowUgsAlerts, templateRef.NotificationChannel, templateRef.NotificationChannelFilter, arguments);
+			await UpdateNotificationsAsync(job.Id, new UpdateNotificationsRequest { Slack = true });
+			return new CreateJobResponse(job.Id.ToString());
 		}
 
 		/// <summary>
 		/// Evaluate a change query to determine which CL to run a job at
 		/// </summary>
-		/// <param name="Stream"></param>
-		/// <param name="TemplateId"></param>
-		/// <param name="Target"></param>
-		/// <param name="Outcomes"></param>
+		/// <param name="stream"></param>
+		/// <param name="templateId"></param>
+		/// <param name="target"></param>
+		/// <param name="outcomes"></param>
 		/// <returns></returns>
-		async Task<int> ExecuteChangeQueryAsync(IStream Stream, TemplateRefId TemplateId, string? Target, List<JobStepOutcome> Outcomes)
+		async Task<int> ExecuteChangeQueryAsync(IStream stream, TemplateRefId templateId, string? target, List<JobStepOutcome> outcomes)
 		{
-			IList<IJob> Jobs = await JobService.FindJobsAsync(StreamId: Stream.Id, Templates: new[] { TemplateId }, Target: Target, State: new[] { JobStepState.Completed }, Outcome: Outcomes.ToArray(), Count: 1, ExcludeUserJobs: true);
-			if (Jobs.Count == 0)
+			IList<IJob> jobs = await _jobService.FindJobsAsync(streamId: stream.Id, templates: new[] { templateId }, target: target, state: new[] { JobStepState.Completed }, outcome: outcomes.ToArray(), count: 1, excludeUserJobs: true);
+			if (jobs.Count == 0)
 			{
-				Logger.LogInformation("Unable to find successful build of {TemplateId} target {Target}. Using latest change instead", TemplateId, Target);
-				return await Perforce.GetLatestChangeAsync(Stream.ClusterName, Stream.Name, null);
+				_logger.LogInformation("Unable to find successful build of {TemplateId} target {Target}. Using latest change instead", templateId, target);
+				return await _perforce.GetLatestChangeAsync(stream.ClusterName, stream.Name, null);
 			}
 			else
 			{
-				Logger.LogInformation("Last successful build of {TemplateId} target {Target} was job {JobId} at change {Change}", TemplateId, Target, Jobs[0].Id, Jobs[0].Change);
-				return Jobs[0].Change;
+				_logger.LogInformation("Last successful build of {TemplateId} target {Target} was job {JobId} at change {Change}", templateId, target, jobs[0].Id, jobs[0].Change);
+				return jobs[0].Change;
 			}
 		}
 
 		/// <summary>
 		/// Deletes a specific job.
 		/// </summary>
-		/// <param name="JobId">Id of the job to delete</param>
+		/// <param name="jobId">Id of the job to delete</param>
 		/// <returns>Async task</returns>
 		[HttpDelete]
-		[Route("/api/v1/jobs/{JobId}")]
-		public async Task<ActionResult> DeleteJobAsync(JobId JobId)
+		[Route("/api/v1/jobs/{jobId}")]
+		public async Task<ActionResult> DeleteJobAsync(JobId jobId)
 		{
-			IJob? Job = await JobService.GetJobAsync(JobId);
-			if (Job == null)
+			IJob? job = await _jobService.GetJobAsync(jobId);
+			if (job == null)
 			{
-				return NotFound(JobId);
+				return NotFound(jobId);
 			}
-			if (!await JobService.AuthorizeAsync(Job, AclAction.DeleteJob, User, null))
+			if (!await _jobService.AuthorizeAsync(job, AclAction.DeleteJob, User, null))
 			{
-				return Forbid(AclAction.DeleteJob, JobId);
+				return Forbid(AclAction.DeleteJob, jobId);
 			}
-			if (!await JobService.DeleteJobAsync(Job))
+			if (!await _jobService.DeleteJobAsync(job))
 			{
-				return NotFound(JobId);
+				return NotFound(jobId);
 			}
 			return Ok();
 		}
@@ -259,46 +252,46 @@ namespace Horde.Build.Controllers
 		/// <summary>
 		/// Updates a specific job.
 		/// </summary>
-		/// <param name="JobId">Id of the job to find</param>
-		/// <param name="Request">Settings to update in the job</param>
+		/// <param name="jobId">Id of the job to find</param>
+		/// <param name="request">Settings to update in the job</param>
 		/// <returns>Async task</returns>
 		[HttpPut]
-		[Route("/api/v1/jobs/{JobId}")]
-		public async Task<ActionResult> UpdateJobAsync(JobId JobId, [FromBody] UpdateJobRequest Request)
+		[Route("/api/v1/jobs/{jobId}")]
+		public async Task<ActionResult> UpdateJobAsync(JobId jobId, [FromBody] UpdateJobRequest request)
 		{
-			IJob? Job = await JobService.GetJobAsync(JobId);
-			if (Job == null)
+			IJob? job = await _jobService.GetJobAsync(jobId);
+			if (job == null)
 			{
-				return NotFound(JobId);
+				return NotFound(jobId);
 			}
 
-			StreamPermissionsCache PermissionsCache = new StreamPermissionsCache();
-			if (!await JobService.AuthorizeAsync(Job, AclAction.UpdateJob, User, PermissionsCache))
+			StreamPermissionsCache permissionsCache = new StreamPermissionsCache();
+			if (!await _jobService.AuthorizeAsync(job, AclAction.UpdateJob, User, permissionsCache))
 			{
-				return Forbid(AclAction.UpdateJob, JobId);
+				return Forbid(AclAction.UpdateJob, jobId);
 			}
-			if (Request.Acl != null && !await JobService.AuthorizeAsync(Job, AclAction.ChangePermissions, User, PermissionsCache))
+			if (request.Acl != null && !await _jobService.AuthorizeAsync(job, AclAction.ChangePermissions, User, permissionsCache))
 			{
-				return Forbid(AclAction.ChangePermissions, JobId);
+				return Forbid(AclAction.ChangePermissions, jobId);
 			}
 
 			// Convert legacy behavior of clearing out the argument to setting the aborted flag
-			if (Request.Arguments != null && Request.Arguments.Count == 0)
+			if (request.Arguments != null && request.Arguments.Count == 0)
 			{
-				Request.Aborted = true;
-				Request.Arguments = null;
+				request.Aborted = true;
+				request.Arguments = null;
 			}
 
-			UserId? AbortedByUserId = null;
-			if (Request.Aborted ?? false)
+			UserId? abortedByUserId = null;
+			if (request.Aborted ?? false)
 			{
-				AbortedByUserId = User.GetUserId();
+				abortedByUserId = User.GetUserId();
 			}
 
-			IJob? NewJob = await JobService.UpdateJobAsync(Job, Name: Request.Name, Priority: Request.Priority, AutoSubmit: Request.AutoSubmit, AbortedByUserId: AbortedByUserId, Arguments: Request.Arguments);
-			if (NewJob == null)
+			IJob? newJob = await _jobService.UpdateJobAsync(job, name: request.Name, priority: request.Priority, autoSubmit: request.AutoSubmit, abortedByUserId: abortedByUserId, arguments: request.Arguments);
+			if (newJob == null)
 			{
-				return NotFound(JobId);
+				return NotFound(jobId);
 			}
 			return Ok();
 		}
@@ -306,612 +299,610 @@ namespace Horde.Build.Controllers
 		/// <summary>
 		/// Updates notifications for a specific job.
 		/// </summary>
-		/// <param name="JobId">Id of the job to find</param>
-		/// <param name="Request">The notification request</param>
+		/// <param name="jobId">Id of the job to find</param>
+		/// <param name="request">The notification request</param>
 		/// <returns>Information about the requested job</returns>
 		[HttpPut]
-		[Route("/api/v1/jobs/{JobId}/notifications")]
-		public async Task<ActionResult> UpdateNotificationsAsync(JobId JobId, [FromBody] UpdateNotificationsRequest Request)
+		[Route("/api/v1/jobs/{jobId}/notifications")]
+		public async Task<ActionResult> UpdateNotificationsAsync(JobId jobId, [FromBody] UpdateNotificationsRequest request)
 		{
-			IJob? Job = await JobService.GetJobAsync(JobId);
-			if (Job == null)
+			IJob? job = await _jobService.GetJobAsync(jobId);
+			if (job == null)
 			{
-				return NotFound(JobId);
+				return NotFound(jobId);
 			}
-			if (!await JobService.AuthorizeAsync(Job, AclAction.CreateSubscription, User, null))
+			if (!await _jobService.AuthorizeAsync(job, AclAction.CreateSubscription, User, null))
 			{
-				return Forbid(AclAction.CreateSubscription, JobId);
+				return Forbid(AclAction.CreateSubscription, jobId);
 			}
 
-			ObjectId TriggerId = Job.NotificationTriggerId ?? ObjectId.GenerateNewId();
+			ObjectId triggerId = job.NotificationTriggerId ?? ObjectId.GenerateNewId();
 
-			Job = await JobService.UpdateJobAsync(Job, null, null, null, null, TriggerId, null, null);
-			if (Job == null)
+			job = await _jobService.UpdateJobAsync(job, null, null, null, null, triggerId, null, null);
+			if (job == null)
 			{
-				return NotFound(JobId);
+				return NotFound(jobId);
 			}
 
-			await NotificationService.UpdateSubscriptionsAsync(TriggerId, User, Request.Email, Request.Slack);
+			await _notificationService.UpdateSubscriptionsAsync(triggerId, User, request.Email, request.Slack);
 			return Ok();
 		}
 
 		/// <summary>
 		/// Gets information about a specific job.
 		/// </summary>
-		/// <param name="JobId">Id of the job to find</param>
+		/// <param name="jobId">Id of the job to find</param>
 		/// <returns>Information about the requested job</returns>
 		[HttpGet]
-		[Route("/api/v1/jobs/{JobId}/notifications")]
-		public async Task<ActionResult<GetNotificationResponse>> GetNotificationsAsync(JobId JobId)
+		[Route("/api/v1/jobs/{jobId}/notifications")]
+		public async Task<ActionResult<GetNotificationResponse>> GetNotificationsAsync(JobId jobId)
 		{
-			IJob? Job = await JobService.GetJobAsync(JobId);
-			if (Job == null)
+			IJob? job = await _jobService.GetJobAsync(jobId);
+			if (job == null)
 			{
-				return NotFound(JobId);
+				return NotFound(jobId);
 			}
-			if (!await JobService.AuthorizeAsync(Job, AclAction.CreateSubscription, User, null))
+			if (!await _jobService.AuthorizeAsync(job, AclAction.CreateSubscription, User, null))
 			{
-				return Forbid(AclAction.CreateSubscription, JobId);
+				return Forbid(AclAction.CreateSubscription, jobId);
 			}
 
-			INotificationSubscription? Subscription;
-			if (Job.NotificationTriggerId == null)
+			INotificationSubscription? subscription;
+			if (job.NotificationTriggerId == null)
 			{
-				Subscription = null;
+				subscription = null;
 			}
 			else
 			{
-				Subscription = await NotificationService.GetSubscriptionsAsync(Job.NotificationTriggerId.Value, User);
+				subscription = await _notificationService.GetSubscriptionsAsync(job.NotificationTriggerId.Value, User);
 			}
-			return new GetNotificationResponse(Subscription);
+			return new GetNotificationResponse(subscription);
 		}
 
 		/// <summary>
 		/// Gets information about a specific job.
 		/// </summary>
-		/// <param name="JobId">Id of the job to find</param>
-		/// <param name="ModifiedAfter">If specified, returns an empty response unless the job's update time is equal to or less than the given value</param>
-		/// <param name="Filter">Filter for the fields to return</param>
+		/// <param name="jobId">Id of the job to find</param>
+		/// <param name="modifiedAfter">If specified, returns an empty response unless the job's update time is equal to or less than the given value</param>
+		/// <param name="filter">Filter for the fields to return</param>
 		/// <returns>Information about the requested job</returns>
 		[HttpGet]
-		[Route("/api/v1/jobs/{JobId}")]
+		[Route("/api/v1/jobs/{jobId}")]
 		[ProducesResponseType(typeof(GetJobResponse), 200)]
-		public async Task<ActionResult<object>> GetJobAsync(JobId JobId, [FromQuery] DateTimeOffset? ModifiedAfter = null, [FromQuery] PropertyFilter? Filter = null)
+		public async Task<ActionResult<object>> GetJobAsync(JobId jobId, [FromQuery] DateTimeOffset? modifiedAfter = null, [FromQuery] PropertyFilter? filter = null)
 		{
-			IJob? Job = await JobService.GetJobAsync(JobId);
-			if (Job == null)
+			IJob? job = await _jobService.GetJobAsync(jobId);
+			if (job == null)
 			{
-				return NotFound(JobId);
+				return NotFound(jobId);
 			}
 
-			StreamPermissionsCache Cache = new StreamPermissionsCache();
-			if (!await JobService.AuthorizeAsync(Job, AclAction.ViewJob, User, Cache))
+			StreamPermissionsCache cache = new StreamPermissionsCache();
+			if (!await _jobService.AuthorizeAsync(job, AclAction.ViewJob, User, cache))
 			{
-				return Forbid(AclAction.ViewJob, JobId);
+				return Forbid(AclAction.ViewJob, jobId);
 			}
-			if (ModifiedAfter != null && Job.UpdateTimeUtc <= ModifiedAfter.Value)
+			if (modifiedAfter != null && job.UpdateTimeUtc <= modifiedAfter.Value)
 			{
 				return new Dictionary<string, object>();
 			}
 
-			IGraph Graph = await JobService.GetGraphAsync(Job);
-			bool bIncludeAcl = await JobService.AuthorizeAsync(Job, AclAction.ViewPermissions, User, Cache);
-			bool bIncludeCosts = await JobService.AuthorizeAsync(Job, AclAction.ViewCosts, User, Cache);
-			return await CreateJobResponseAsync(Job, Graph, bIncludeAcl, bIncludeCosts, Filter);
+			IGraph graph = await _jobService.GetGraphAsync(job);
+			bool bIncludeAcl = await _jobService.AuthorizeAsync(job, AclAction.ViewPermissions, User, cache);
+			bool bIncludeCosts = await _jobService.AuthorizeAsync(job, AclAction.ViewCosts, User, cache);
+			return await CreateJobResponseAsync(job, graph, bIncludeAcl, bIncludeCosts, filter);
 		}
 
 		/// <summary>
 		/// Creates a response containing information about this object
 		/// </summary>
-		/// <param name="Job">The job document</param>
-		/// <param name="Graph">The graph for this job</param>
+		/// <param name="job">The job document</param>
+		/// <param name="graph">The graph for this job</param>
 		/// <param name="bIncludeAcl">Whether to include the ACL in the response</param>
 		/// <param name="bIncludeCosts">Whether to include costs in the response</param>
-		/// <param name="Filter">Filter for the properties to return</param>
+		/// <param name="filter">Filter for the properties to return</param>
 		/// <returns>Object containing the requested properties</returns>
-		async Task<object> CreateJobResponseAsync(IJob Job, IGraph Graph, bool bIncludeAcl, bool bIncludeCosts, PropertyFilter? Filter)
+		async Task<object> CreateJobResponseAsync(IJob job, IGraph graph, bool bIncludeAcl, bool bIncludeCosts, PropertyFilter? filter)
 		{
-			if (Filter == null)
+			if (filter == null)
 			{
-				return await CreateJobResponseAsync(Job, Graph, true, true, bIncludeAcl, bIncludeCosts);
+				return await CreateJobResponseAsync(job, graph, true, true, bIncludeAcl, bIncludeCosts);
 			}
 			else
 			{
-				return Filter.ApplyTo(await CreateJobResponseAsync(Job, Graph, Filter.Includes(nameof(GetJobResponse.Batches)), Filter.Includes(nameof(GetJobResponse.Labels)) || Filter.Includes(nameof(GetJobResponse.DefaultLabel)), bIncludeAcl, bIncludeCosts));
+				return filter.ApplyTo(await CreateJobResponseAsync(job, graph, filter.Includes(nameof(GetJobResponse.Batches)), filter.Includes(nameof(GetJobResponse.Labels)) || filter.Includes(nameof(GetJobResponse.DefaultLabel)), bIncludeAcl, bIncludeCosts));
 			}
 		}
 
 		/// <summary>
 		/// Creates a response containing information about this object
 		/// </summary>
-		/// <param name="Job">The job document</param>
-		/// <param name="Graph">The graph definition</param>
+		/// <param name="job">The job document</param>
+		/// <param name="graph">The graph definition</param>
 		/// <param name="bIncludeBatches">Whether to include the job batches in the response</param>
 		/// <param name="bIncludeLabels">Whether to include the job aggregates in the response</param>
 		/// <param name="bIncludeAcl">Whether to include the ACL in the response</param>
 		/// <param name="bIncludeCosts">Whether to include costs of running particular agents</param>
 		/// <returns>The response object</returns>
-		async ValueTask<GetJobResponse> CreateJobResponseAsync(IJob Job, IGraph Graph, bool bIncludeBatches, bool bIncludeLabels, bool bIncludeAcl, bool bIncludeCosts)
+		async ValueTask<GetJobResponse> CreateJobResponseAsync(IJob job, IGraph graph, bool bIncludeBatches, bool bIncludeLabels, bool bIncludeAcl, bool bIncludeCosts)
 		{
-			GetThinUserInfoResponse? StartedByUserInfo = null;
-			if (Job.StartedByUserId != null)
+			GetThinUserInfoResponse? startedByUserInfo = null;
+			if (job.StartedByUserId != null)
 			{
-				StartedByUserInfo = new GetThinUserInfoResponse(await UserCollection.GetCachedUserAsync(Job.StartedByUserId.Value));
+				startedByUserInfo = new GetThinUserInfoResponse(await _userCollection.GetCachedUserAsync(job.StartedByUserId.Value));
 			}
 
-			GetThinUserInfoResponse? AbortedByUserInfo = null;
-			if (Job.AbortedByUserId != null)
+			GetThinUserInfoResponse? abortedByUserInfo = null;
+			if (job.AbortedByUserId != null)
 			{
-				AbortedByUserInfo = new GetThinUserInfoResponse(await UserCollection.GetCachedUserAsync(Job.AbortedByUserId.Value));
+				abortedByUserInfo = new GetThinUserInfoResponse(await _userCollection.GetCachedUserAsync(job.AbortedByUserId.Value));
 			}
 
-			GetAclResponse? AclResponse = null;
-			if (bIncludeAcl && Job.Acl != null)
+			GetAclResponse? aclResponse = null;
+			if (bIncludeAcl && job.Acl != null)
 			{
-				AclResponse = new GetAclResponse(Job.Acl);
+				aclResponse = new GetAclResponse(job.Acl);
 			}
 
-			GetJobResponse Response = new GetJobResponse(Job, StartedByUserInfo, AbortedByUserInfo, AclResponse);
+			GetJobResponse response = new GetJobResponse(job, startedByUserInfo, abortedByUserInfo, aclResponse);
 			if (bIncludeBatches || bIncludeLabels)
 			{
 				if (bIncludeBatches)
 				{
-					Response.Batches = new List<GetBatchResponse>();
-					foreach (IJobStepBatch Batch in Job.Batches)
+					response.Batches = new List<GetBatchResponse>();
+					foreach (IJobStepBatch batch in job.Batches)
 					{
-						Response.Batches.Add(await CreateBatchResponseAsync(Batch, bIncludeCosts));
+						response.Batches.Add(await CreateBatchResponseAsync(batch, bIncludeCosts));
 					}
 				}
 				if (bIncludeLabels)
 				{
-					Response.Labels = new List<GetLabelStateResponse>();
-					Response.DefaultLabel = Job.GetLabelStateResponses(Graph, Response.Labels);
+					response.Labels = new List<GetLabelStateResponse>();
+					response.DefaultLabel = job.GetLabelStateResponses(graph, response.Labels);
 				}
 			}
-			return Response;
+			return response;
 		}
 
 		/// <summary>
 		/// Get the response object for a batch
 		/// </summary>
-		/// <param name="Batch"></param>
+		/// <param name="batch"></param>
 		/// <param name="bIncludeCosts"></param>
 		/// <returns></returns>
-		async ValueTask<GetBatchResponse> CreateBatchResponseAsync(IJobStepBatch Batch, bool bIncludeCosts)
+		async ValueTask<GetBatchResponse> CreateBatchResponseAsync(IJobStepBatch batch, bool bIncludeCosts)
 		{
-			List<GetStepResponse> Steps = new List<GetStepResponse>();
-			foreach (IJobStep Step in Batch.Steps)
+			List<GetStepResponse> steps = new List<GetStepResponse>();
+			foreach (IJobStep step in batch.Steps)
 			{
-				Steps.Add(await CreateStepResponseAsync(Step));
+				steps.Add(await CreateStepResponseAsync(step));
 			}
 
-			double? AgentRate = null;
-			if (Batch.AgentId != null && bIncludeCosts)
+			double? agentRate = null;
+			if (batch.AgentId != null && bIncludeCosts)
 			{
-				AgentRate = await AgentService.GetRateAsync(Batch.AgentId.Value);
+				agentRate = await _agentService.GetRateAsync(batch.AgentId.Value);
 			}
 
-			return new GetBatchResponse(Batch, Steps, AgentRate);
+			return new GetBatchResponse(batch, steps, agentRate);
 		}
 
 		/// <summary>
 		/// Get the response object for a step
 		/// </summary>
-		/// <param name="Step"></param>
+		/// <param name="step"></param>
 		/// <returns></returns>
-		async ValueTask<GetStepResponse> CreateStepResponseAsync(IJobStep Step)
+		async ValueTask<GetStepResponse> CreateStepResponseAsync(IJobStep step)
 		{
-			GetThinUserInfoResponse? AbortedByUserInfo = null;
-			if (Step.AbortedByUserId != null)
+			GetThinUserInfoResponse? abortedByUserInfo = null;
+			if (step.AbortedByUserId != null)
 			{
-				AbortedByUserInfo = new GetThinUserInfoResponse(await UserCollection.GetCachedUserAsync(Step.AbortedByUserId.Value));
+				abortedByUserInfo = new GetThinUserInfoResponse(await _userCollection.GetCachedUserAsync(step.AbortedByUserId.Value));
 			}
 
-			GetThinUserInfoResponse? RetriedByUserInfo = null;
-			if (Step.RetriedByUserId != null)
+			GetThinUserInfoResponse? retriedByUserInfo = null;
+			if (step.RetriedByUserId != null)
 			{
-				RetriedByUserInfo = new GetThinUserInfoResponse(await UserCollection.GetCachedUserAsync(Step.RetriedByUserId.Value));
+				retriedByUserInfo = new GetThinUserInfoResponse(await _userCollection.GetCachedUserAsync(step.RetriedByUserId.Value));
 			}
 
-			return new GetStepResponse(Step, AbortedByUserInfo, RetriedByUserInfo);
+			return new GetStepResponse(step, abortedByUserInfo, retriedByUserInfo);
 		}
 
 		/// <summary>
 		/// Gets information about the graph for a specific job.
 		/// </summary>
-		/// <param name="JobId">Id of the job to find</param>
-		/// <param name="Filter">Filter for the fields to return</param>
+		/// <param name="jobId">Id of the job to find</param>
+		/// <param name="filter">Filter for the fields to return</param>
 		/// <returns>Information about the requested job</returns>
 		[HttpGet]
-		[Route("/api/v1/jobs/{JobId}/graph")]
+		[Route("/api/v1/jobs/{jobId}/graph")]
 		[ProducesResponseType(typeof(GetGraphResponse), 200)]
-		public async Task<ActionResult<object>> GetJobGraphAsync(JobId JobId, [FromQuery] PropertyFilter? Filter = null)
+		public async Task<ActionResult<object>> GetJobGraphAsync(JobId jobId, [FromQuery] PropertyFilter? filter = null)
 		{
-			IJob? Job = await JobService.GetJobAsync(JobId);
-			if (Job == null)
+			IJob? job = await _jobService.GetJobAsync(jobId);
+			if (job == null)
 			{
-				return NotFound(JobId);
+				return NotFound(jobId);
 			}
-			if (!await JobService.AuthorizeAsync(Job, AclAction.ViewJob, User, null))
+			if (!await _jobService.AuthorizeAsync(job, AclAction.ViewJob, User, null))
 			{
-				return Forbid(AclAction.ViewJob, Job.Id);
+				return Forbid(AclAction.ViewJob, job.Id);
 			}
 
-			IGraph Graph = await JobService.GetGraphAsync(Job);
-			return PropertyFilter.Apply(new GetGraphResponse(Graph), Filter);
+			IGraph graph = await _jobService.GetGraphAsync(job);
+			return PropertyFilter.Apply(new GetGraphResponse(graph), filter);
 		}
 
 		/// <summary>
 		/// Gets timing information about the graph for a specific job.
 		/// </summary>
-		/// <param name="JobId">Id of the job to find</param>
-		/// <param name="Filter">Filter for the fields to return</param>
+		/// <param name="jobId">Id of the job to find</param>
+		/// <param name="filter">Filter for the fields to return</param>
 		/// <returns>Information about the requested job</returns>
 		[HttpGet]
-		[Route("/api/v1/jobs/{JobId}/timing")]
+		[Route("/api/v1/jobs/{jobId}/timing")]
 		[ProducesResponseType(typeof(GetJobTimingResponse), 200)]
-		public async Task<ActionResult<object>> GetJobTimingAsync(JobId JobId, [FromQuery] PropertyFilter? Filter = null)
+		public async Task<ActionResult<object>> GetJobTimingAsync(JobId jobId, [FromQuery] PropertyFilter? filter = null)
 		{
-			IJob? Job = await JobService.GetJobAsync(JobId);
-			if (Job == null)
+			IJob? job = await _jobService.GetJobAsync(jobId);
+			if (job == null)
 			{
-				return NotFound(JobId);
+				return NotFound(jobId);
 			}
-			if (!await JobService.AuthorizeAsync(Job, AclAction.ViewJob, User, null))
+			if (!await _jobService.AuthorizeAsync(job, AclAction.ViewJob, User, null))
 			{
-				return Forbid(AclAction.ViewJob, JobId);
+				return Forbid(AclAction.ViewJob, jobId);
 			}
 
-			IJobTiming JobTiming = await JobService.GetJobTimingAsync(Job);
-			IGraph Graph = await JobService.GetGraphAsync(Job);
-			return PropertyFilter.Apply(await CreateJobTimingResponse(Job, Graph, JobTiming), Filter);
+			IJobTiming jobTiming = await _jobService.GetJobTimingAsync(job);
+			IGraph graph = await _jobService.GetGraphAsync(job);
+			return PropertyFilter.Apply(await CreateJobTimingResponse(job, graph, jobTiming), filter);
 		}
 
-		private async Task<GetJobTimingResponse> CreateJobTimingResponse(IJob Job, IGraph Graph, IJobTiming JobTiming, bool IncludeJobResponse = false)
+		private async Task<GetJobTimingResponse> CreateJobTimingResponse(IJob job, IGraph graph, IJobTiming jobTiming, bool includeJobResponse = false)
 		{
-			Dictionary<INode, TimingInfo> NodeToTimingInfo = Job.GetTimingInfo(Graph, JobTiming);
+			Dictionary<INode, TimingInfo> nodeToTimingInfo = job.GetTimingInfo(graph, jobTiming);
 
-			Dictionary<string, GetStepTimingInfoResponse> Steps = new Dictionary<string, GetStepTimingInfoResponse>();
-			foreach (IJobStepBatch Batch in Job.Batches)
+			Dictionary<string, GetStepTimingInfoResponse> steps = new Dictionary<string, GetStepTimingInfoResponse>();
+			foreach (IJobStepBatch batch in job.Batches)
 			{
-				foreach (IJobStep Step in Batch.Steps)
+				foreach (IJobStep step in batch.Steps)
 				{
-					INode Node = Graph.Groups[Batch.GroupIdx].Nodes[Step.NodeIdx];
-					Steps[Step.Id.ToString()] = new GetStepTimingInfoResponse(Node.Name, NodeToTimingInfo[Node]);
+					INode node = graph.Groups[batch.GroupIdx].Nodes[step.NodeIdx];
+					steps[step.Id.ToString()] = new GetStepTimingInfoResponse(node.Name, nodeToTimingInfo[node]);
 				}
 			}
 
-			List<GetLabelTimingInfoResponse> Labels = new List<GetLabelTimingInfoResponse>();
-			foreach (ILabel Label in Graph.Labels)
+			List<GetLabelTimingInfoResponse> labels = new List<GetLabelTimingInfoResponse>();
+			foreach (ILabel label in graph.Labels)
 			{
-				TimingInfo TimingInfo = TimingInfo.Max(Label.GetDependencies(Graph.Groups).Select(x => NodeToTimingInfo[x]));
-				Labels.Add(new GetLabelTimingInfoResponse(Label, TimingInfo));
+				TimingInfo timingInfo = TimingInfo.Max(label.GetDependencies(graph.Groups).Select(x => nodeToTimingInfo[x]));
+				labels.Add(new GetLabelTimingInfoResponse(label, timingInfo));
 			}
 
-			GetJobResponse? JobResponse = null;
-			if (IncludeJobResponse)
+			GetJobResponse? jobResponse = null;
+			if (includeJobResponse)
 			{
-				JobResponse = await CreateJobResponseAsync(Job, Graph, true, true, false, true);
+				jobResponse = await CreateJobResponseAsync(job, graph, true, true, false, true);
 			}
 
-			return new GetJobTimingResponse(Job, JobResponse, Steps, Labels);
+			return new GetJobTimingResponse(job, jobResponse, steps, labels);
 		}
 		
 		/// <summary>
 		/// Find timing information about the graph for multiple jobs
 		/// </summary>
-		/// <param name="StreamId">The stream to search in</param>
-		/// <param name="Templates">List of templates to find</param>
-		/// <param name="Filter">Filter for the fields to return</param>
-		/// <param name="Count">Number of results to return</param>
+		/// <param name="streamId">The stream to search in</param>
+		/// <param name="templates">List of templates to find</param>
+		/// <param name="filter">Filter for the fields to return</param>
+		/// <param name="count">Number of results to return</param>
 		/// <returns>Job timings for each job ID</returns>
 		[HttpGet]
 		[Route("/api/v1/jobs/timing")]
 		[ProducesResponseType(typeof(FindJobTimingsResponse), 200)]
 		public async Task<ActionResult<object>> FindJobTimingsAsync(
-			[FromQuery] string? StreamId = null,
-			[FromQuery(Name = "template")] string[]? Templates = null,
-			[FromQuery] PropertyFilter? Filter = null,
-			[FromQuery] int Count = 100)
+			[FromQuery] string? streamId = null,
+			[FromQuery(Name = "template")] string[]? templates = null,
+			[FromQuery] PropertyFilter? filter = null,
+			[FromQuery] int count = 100)
 		{
-			if (StreamId == null)
+			if (streamId == null)
 			{
 				return BadRequest("Missing/invalid query parameter streamId");
 			}
 
-			TemplateRefId[] TemplateRefIds = Templates switch
+			TemplateRefId[] templateRefIds = templates switch
 			{
-				{ Length: > 0 } => Templates.Select(x => new TemplateRefId(x)).ToArray(),
+				{ Length: > 0 } => templates.Select(x => new TemplateRefId(x)).ToArray(),
 				_ => Array.Empty<TemplateRefId>()
 			};
 
-			List<IJob> Jobs = await JobService.FindJobsByStreamWithTemplatesAsync(new StreamId(StreamId), TemplateRefIds, Count: Count, ConsistentRead: false);
+			List<IJob> jobs = await _jobService.FindJobsByStreamWithTemplatesAsync(new StreamId(streamId), templateRefIds, count: count, consistentRead: false);
 
-			Dictionary<string, GetJobTimingResponse> JobTimings = await Jobs.ToAsyncEnumerable()
-				.WhereAwait(async Job => await JobService.AuthorizeAsync(Job, AclAction.ViewJob, User, null))
-				.ToDictionaryAwaitAsync(x => ValueTask.FromResult(x.Id.ToString()), async Job =>
+			Dictionary<string, GetJobTimingResponse> jobTimings = await jobs.ToAsyncEnumerable()
+				.WhereAwait(async job => await _jobService.AuthorizeAsync(job, AclAction.ViewJob, User, null))
+				.ToDictionaryAwaitAsync(x => ValueTask.FromResult(x.Id.ToString()), async job =>
 				{
-					IJobTiming JobTiming = await JobService.GetJobTimingAsync(Job);
-					IGraph Graph = await JobService.GetGraphAsync(Job);
-					return await CreateJobTimingResponse(Job, Graph, JobTiming, true);
+					IJobTiming jobTiming = await _jobService.GetJobTimingAsync(job);
+					IGraph graph = await _jobService.GetGraphAsync(job);
+					return await CreateJobTimingResponse(job, graph, jobTiming, true);
 				});
 			
-			return PropertyFilter.Apply(new FindJobTimingsResponse(JobTimings), Filter);
+			return PropertyFilter.Apply(new FindJobTimingsResponse(jobTimings), filter);
 		}
 
 		/// <summary>
 		/// Gets information about the template for a specific job.
 		/// </summary>
-		/// <param name="JobId">Id of the job to find</param>
-		/// <param name="Filter">Filter for the fields to return</param>
+		/// <param name="jobId">Id of the job to find</param>
+		/// <param name="filter">Filter for the fields to return</param>
 		/// <returns>Information about the requested job</returns>
 		[HttpGet]
-		[Route("/api/v1/jobs/{JobId}/template")]
+		[Route("/api/v1/jobs/{jobId}/template")]
 		[ProducesResponseType(typeof(GetTemplateResponse), 200)]
-		public async Task<ActionResult<object>> GetJobTemplateAsync(JobId JobId, [FromQuery] PropertyFilter? Filter = null)
+		public async Task<ActionResult<object>> GetJobTemplateAsync(JobId jobId, [FromQuery] PropertyFilter? filter = null)
 		{
-			IJob? Job = await JobService.GetJobAsync(JobId);
-			if (Job == null || Job.TemplateHash == null)
+			IJob? job = await _jobService.GetJobAsync(jobId);
+			if (job == null || job.TemplateHash == null)
 			{
-				return NotFound(JobId);
+				return NotFound(jobId);
 			}
-			if (!await JobService.AuthorizeAsync(Job, AclAction.ViewJob, User, null))
+			if (!await _jobService.AuthorizeAsync(job, AclAction.ViewJob, User, null))
 			{
-				return Forbid(AclAction.ViewJob, JobId);
+				return Forbid(AclAction.ViewJob, jobId);
 			}
 
-			ITemplate? Template = await TemplateCollection.GetAsync(Job.TemplateHash);
-			if(Template == null)
+			ITemplate? template = await _templateCollection.GetAsync(job.TemplateHash);
+			if(template == null)
 			{
-				return NotFound(Job.StreamId, Job.TemplateId);
+				return NotFound(job.StreamId, job.TemplateId);
 			}
 
-			return new GetTemplateResponse(Template).ApplyFilter(Filter);
+			return new GetTemplateResponse(template).ApplyFilter(filter);
 		}
 
 		/// <summary>
 		/// Find jobs matching a criteria
 		/// </summary>
-		/// <param name="Ids">The job ids to return</param>
-		/// <param name="Name">Name of the job to find</param>
-		/// <param name="Templates">List of templates to find</param>
-		/// <param name="StreamId">The stream to search for</param>
-		/// <param name="MinChange">The minimum changelist number</param>
-		/// <param name="MaxChange">The maximum changelist number</param>
-		/// <param name="IncludePreflight">Whether to include preflight jobs</param>
-		/// <param name="PreflightOnly">Whether to only include preflight jobs</param>
-		/// <param name="PreflightChange">The preflighted changelist</param>
-		/// <param name="StartedByUserId">User id for which to include jobs</param>
-		/// <param name="PreflightStartedByUserId">User id for which to include preflight jobs</param>
-		/// <param name="MinCreateTime">Minimum creation time</param>
-		/// <param name="MaxCreateTime">Maximum creation time</param>
-		/// <param name="ModifiedBefore">If specified, only jobs updated before the give time will be returned</param>
-		/// <param name="ModifiedAfter">If specified, only jobs updated after the give time will be returned</param>
-		/// <param name="Target">Target to filter the returned jobs by</param>
-		/// <param name="State">Filter state of the returned jobs</param>
-		/// <param name="Outcome">Filter outcome of the returned jobs</param>
-		/// <param name="Filter">Filter for properties to return</param>
-		/// <param name="Index">Index of the first result to be returned</param>
-		/// <param name="Count">Number of results to return</param>
+		/// <param name="ids">The job ids to return</param>
+		/// <param name="name">Name of the job to find</param>
+		/// <param name="templates">List of templates to find</param>
+		/// <param name="streamId">The stream to search for</param>
+		/// <param name="minChange">The minimum changelist number</param>
+		/// <param name="maxChange">The maximum changelist number</param>
+		/// <param name="includePreflight">Whether to include preflight jobs</param>
+		/// <param name="preflightOnly">Whether to only include preflight jobs</param>
+		/// <param name="preflightChange">The preflighted changelist</param>
+		/// <param name="startedByUserId">User id for which to include jobs</param>
+		/// <param name="preflightStartedByUserId">User id for which to include preflight jobs</param>
+		/// <param name="minCreateTime">Minimum creation time</param>
+		/// <param name="maxCreateTime">Maximum creation time</param>
+		/// <param name="modifiedBefore">If specified, only jobs updated before the give time will be returned</param>
+		/// <param name="modifiedAfter">If specified, only jobs updated after the give time will be returned</param>
+		/// <param name="target">Target to filter the returned jobs by</param>
+		/// <param name="state">Filter state of the returned jobs</param>
+		/// <param name="outcome">Filter outcome of the returned jobs</param>
+		/// <param name="filter">Filter for properties to return</param>
+		/// <param name="index">Index of the first result to be returned</param>
+		/// <param name="count">Number of results to return</param>
 		/// <returns>List of jobs</returns>
 		[HttpGet]
 		[Route("/api/v1/jobs")]
 		[ProducesResponseType(typeof(List<GetJobResponse>), 200)]
 		public async Task<ActionResult<List<object>>> FindJobsAsync(
-			[FromQuery(Name = "Id")] string[]? Ids = null,
-			[FromQuery] string? StreamId = null,
-			[FromQuery] string? Name = null,
-			[FromQuery(Name = "template")] string[]? Templates = null,
-			[FromQuery] int? MinChange = null,
-			[FromQuery] int? MaxChange = null,
-			[FromQuery] bool IncludePreflight = true,
-			[FromQuery] bool? PreflightOnly = null,
-			[FromQuery] int? PreflightChange = null,
-			[FromQuery] string? PreflightStartedByUserId = null,
-			[FromQuery] string? StartedByUserId = null,
-			[FromQuery] DateTimeOffset? MinCreateTime = null,
-			[FromQuery] DateTimeOffset? MaxCreateTime = null,
-			[FromQuery] DateTimeOffset? ModifiedBefore = null,
-			[FromQuery] DateTimeOffset? ModifiedAfter = null,
-			[FromQuery] string? Target = null,
-			[FromQuery] JobStepState[]? State = null,
-			[FromQuery] JobStepOutcome[]? Outcome = null, 
-			[FromQuery] PropertyFilter? Filter = null,
-			[FromQuery] int Index = 0,
-			[FromQuery] int Count = 100)
+			[FromQuery(Name = "Id")] string[]? ids = null,
+			[FromQuery] string? streamId = null,
+			[FromQuery] string? name = null,
+			[FromQuery(Name = "template")] string[]? templates = null,
+			[FromQuery] int? minChange = null,
+			[FromQuery] int? maxChange = null,
+			[FromQuery] bool includePreflight = true,
+			[FromQuery] bool? preflightOnly = null,
+			[FromQuery] int? preflightChange = null,
+			[FromQuery] string? preflightStartedByUserId = null,
+			[FromQuery] string? startedByUserId = null,
+			[FromQuery] DateTimeOffset? minCreateTime = null,
+			[FromQuery] DateTimeOffset? maxCreateTime = null,
+			[FromQuery] DateTimeOffset? modifiedBefore = null,
+			[FromQuery] DateTimeOffset? modifiedAfter = null,
+			[FromQuery] string? target = null,
+			[FromQuery] JobStepState[]? state = null,
+			[FromQuery] JobStepOutcome[]? outcome = null, 
+			[FromQuery] PropertyFilter? filter = null,
+			[FromQuery] int index = 0,
+			[FromQuery] int count = 100)
 		{
-			JobId[]? JobIdValues = (Ids == null) ? (JobId[]?)null : Array.ConvertAll(Ids, x => new JobId(x));
-			StreamId? StreamIdValue = (StreamId == null)? (StreamId?)null : new StreamId(StreamId);
+			JobId[]? jobIdValues = (ids == null) ? (JobId[]?)null : Array.ConvertAll(ids, x => new JobId(x));
+			StreamId? streamIdValue = (streamId == null)? (StreamId?)null : new StreamId(streamId);
 			
-			TemplateRefId[]? TemplateRefIds = (Templates != null && Templates.Length > 0) ? Templates.Select(x => new TemplateRefId(x)).ToArray() : null;
+			TemplateRefId[]? templateRefIds = (templates != null && templates.Length > 0) ? templates.Select(x => new TemplateRefId(x)).ToArray() : null;
 
-			if (IncludePreflight == false)
+			if (includePreflight == false)
 			{
-				PreflightChange = 0;
+				preflightChange = 0;
 			}
 
-			UserId? PreflightStartedByUserIdValue = null;
+			UserId? preflightStartedByUserIdValue = null;
 
-			if (PreflightStartedByUserId != null)
+			if (preflightStartedByUserId != null)
 			{
-				PreflightStartedByUserIdValue = new UserId(PreflightStartedByUserId);
+				preflightStartedByUserIdValue = new UserId(preflightStartedByUserId);
 			}
 
-			UserId? StartedByUserIdValue = null;
+			UserId? startedByUserIdValue = null;
 
-			if (StartedByUserId != null)
+			if (startedByUserId != null)
 			{
-				StartedByUserIdValue = new UserId(StartedByUserId);
+				startedByUserIdValue = new UserId(startedByUserId);
 			}
 
-			List<IJob> Jobs;
+			List<IJob> jobs;
 			using (IScope _ = GlobalTracer.Instance.BuildSpan("FindJobs").StartActive())
 			{
-				Jobs = await JobService.FindJobsAsync(JobIdValues, StreamIdValue, Name, TemplateRefIds, MinChange,
-					MaxChange, PreflightChange, PreflightOnly, PreflightStartedByUserIdValue, StartedByUserIdValue, MinCreateTime?.UtcDateTime, MaxCreateTime?.UtcDateTime, Target, State, Outcome,
-					ModifiedBefore, ModifiedAfter, Index, Count, false);
+				jobs = await _jobService.FindJobsAsync(jobIdValues, streamIdValue, name, templateRefIds, minChange,
+					maxChange, preflightChange, preflightOnly, preflightStartedByUserIdValue, startedByUserIdValue, minCreateTime?.UtcDateTime, maxCreateTime?.UtcDateTime, target, state, outcome,
+					modifiedBefore, modifiedAfter, index, count, false);
 			}
 
-			StreamPermissionsCache PermissionsCache = new StreamPermissionsCache();
+			StreamPermissionsCache permissionsCache = new StreamPermissionsCache();
 
-			List<object> Responses = new List<object>();
-			foreach (IJob Job in Jobs)
+			List<object> responses = new List<object>();
+			foreach (IJob job in jobs)
 			{
-				using IScope JobScope = GlobalTracer.Instance.BuildSpan("JobIteration").StartActive();
-				JobScope.Span.SetTag("jobId", Job.Id.ToString());
+				using IScope jobScope = GlobalTracer.Instance.BuildSpan("JobIteration").StartActive();
+				jobScope.Span.SetTag("jobId", job.Id.ToString());
 
-				if (Job.GraphHash == null)
+				if (job.GraphHash == null)
 				{
-					Logger.LogWarning("Job {JobId} has no GraphHash", Job.Id);
+					_logger.LogWarning("Job {JobId} has no GraphHash", job.Id);
 					continue;
 				}
 
-				bool ViewJobAuthorized;
+				bool viewJobAuthorized;
 				using (IScope _ = GlobalTracer.Instance.BuildSpan("AuthorizeViewJob").StartActive())
 				{
-					ViewJobAuthorized = await JobService.AuthorizeAsync(Job, AclAction.ViewJob, User, PermissionsCache);
+					viewJobAuthorized = await _jobService.AuthorizeAsync(job, AclAction.ViewJob, User, permissionsCache);
 				}
 				
-				if (ViewJobAuthorized)
+				if (viewJobAuthorized)
 				{
-					IGraph Graph;
+					IGraph graph;
 					using (IScope _ = GlobalTracer.Instance.BuildSpan("GetGraph").StartActive())
 					{
-						Graph = await JobService.GetGraphAsync(Job);
+						graph = await _jobService.GetGraphAsync(job);
 					}
 
 					bool bIncludeAcl;
 					using (IScope _ = GlobalTracer.Instance.BuildSpan("AuthorizeViewPermissions").StartActive())
 					{
-						bIncludeAcl = await JobService.AuthorizeAsync(Job, AclAction.ViewPermissions, User, PermissionsCache);
+						bIncludeAcl = await _jobService.AuthorizeAsync(job, AclAction.ViewPermissions, User, permissionsCache);
 					}
 
 					bool bIncludeCosts;
 					using (IScope _ = GlobalTracer.Instance.BuildSpan("AuthorizeViewCosts").StartActive())
 					{
-						bIncludeCosts = await JobService.AuthorizeAsync(Job, AclAction.ViewCosts, User, PermissionsCache);
+						bIncludeCosts = await _jobService.AuthorizeAsync(job, AclAction.ViewCosts, User, permissionsCache);
 					}
 
 					using (IScope _ = GlobalTracer.Instance.BuildSpan("CreateResponse").StartActive())
 					{
-						Responses.Add(await CreateJobResponseAsync(Job, Graph, bIncludeAcl, bIncludeCosts, Filter));
+						responses.Add(await CreateJobResponseAsync(job, graph, bIncludeAcl, bIncludeCosts, filter));
 					}
 				}
 			}
-			return Responses;
+			return responses;
 		}
 
 		/// <summary>
 		/// Find jobs for a stream with given templates, sorted by creation date
 		/// </summary>
-		/// <param name="StreamId">The stream to search for</param>
-		/// <param name="Templates">List of templates to find</param>
-		/// <param name="PreflightStartedByUserId">User id for which to include preflight jobs</param>
-		/// <param name="MaxCreateTime">Maximum creation time</param>
-		/// <param name="ModifiedAfter">If specified, only jobs updated after the given time will be returned</param>
-		/// <param name="Filter">Filter for properties to return</param>
-		/// <param name="Index">Index of the first result to be returned</param>
-		/// <param name="Count">Number of results to return</param>
-		/// <param name="ConsistentRead">If a read to the primary database is required, for read consistency. Usually not required.</param>
+		/// <param name="streamId">The stream to search for</param>
+		/// <param name="templates">List of templates to find</param>
+		/// <param name="preflightStartedByUserId">User id for which to include preflight jobs</param>
+		/// <param name="maxCreateTime">Maximum creation time</param>
+		/// <param name="modifiedAfter">If specified, only jobs updated after the given time will be returned</param>
+		/// <param name="filter">Filter for properties to return</param>
+		/// <param name="index">Index of the first result to be returned</param>
+		/// <param name="count">Number of results to return</param>
+		/// <param name="consistentRead">If a read to the primary database is required, for read consistency. Usually not required.</param>
 		/// <returns>List of jobs</returns>
 		[HttpGet]
-		[Route("/api/v1/jobs/streams/{StreamId}")]
+		[Route("/api/v1/jobs/streams/{streamId}")]
 		[ProducesResponseType(typeof(List<GetJobResponse>), 200)]
 		public async Task<ActionResult<List<object>>> FindJobsByStreamWithTemplatesAsync(
-			string StreamId,
-			[FromQuery(Name = "template")] string[] Templates,
-			[FromQuery] string? PreflightStartedByUserId = null,
-			[FromQuery] DateTimeOffset? MaxCreateTime = null,
-			[FromQuery] DateTimeOffset? ModifiedAfter = null,
-			[FromQuery] PropertyFilter? Filter = null,
-			[FromQuery] int Index = 0,
-			[FromQuery] int Count = 100,
-			[FromQuery] bool ConsistentRead = false)
+			string streamId,
+			[FromQuery(Name = "template")] string[] templates,
+			[FromQuery] string? preflightStartedByUserId = null,
+			[FromQuery] DateTimeOffset? maxCreateTime = null,
+			[FromQuery] DateTimeOffset? modifiedAfter = null,
+			[FromQuery] PropertyFilter? filter = null,
+			[FromQuery] int index = 0,
+			[FromQuery] int count = 100,
+			[FromQuery] bool consistentRead = false)
 		{
-			StreamId StreamIdValue = new StreamId(StreamId);
-			TemplateRefId[] TemplateRefIds = Templates.Select(x => new TemplateRefId(x)).ToArray();
-			UserId? PreflightStartedByUserIdValue = PreflightStartedByUserId != null ? new UserId(PreflightStartedByUserId) : null;
-			Count = Math.Min(1000, Count);
+			StreamId streamIdValue = new StreamId(streamId);
+			TemplateRefId[] templateRefIds = templates.Select(x => new TemplateRefId(x)).ToArray();
+			UserId? preflightStartedByUserIdValue = preflightStartedByUserId != null ? new UserId(preflightStartedByUserId) : null;
+			count = Math.Min(1000, count);
 
-			List<IJob> Jobs = await JobService.FindJobsByStreamWithTemplatesAsync(StreamIdValue, TemplateRefIds, PreflightStartedByUserIdValue, MaxCreateTime, ModifiedAfter, Index, Count, ConsistentRead);
-			return await CreateAuthorizedJobResponses(Jobs, Filter);
+			List<IJob> jobs = await _jobService.FindJobsByStreamWithTemplatesAsync(streamIdValue, templateRefIds, preflightStartedByUserIdValue, maxCreateTime, modifiedAfter, index, count, consistentRead);
+			return await CreateAuthorizedJobResponses(jobs, filter);
 		}
 
-		private async Task<List<object>> CreateAuthorizedJobResponses(List<IJob> Jobs, PropertyFilter? Filter = null)
+		private async Task<List<object>> CreateAuthorizedJobResponses(List<IJob> jobs, PropertyFilter? filter = null)
 		{
-			StreamPermissionsCache PermissionsCache = new ();
-			List<object> Responses = new ();
-			foreach (IJob Job in Jobs)
+			StreamPermissionsCache permissionsCache = new ();
+			List<object> responses = new ();
+			foreach (IJob job in jobs)
 			{
-				using IScope JobScope = GlobalTracer.Instance.BuildSpan("JobIteration").StartActive();
-				JobScope.Span.SetTag("jobId", Job.Id.ToString());
+				using IScope jobScope = GlobalTracer.Instance.BuildSpan("JobIteration").StartActive();
+				jobScope.Span.SetTag("jobId", job.Id.ToString());
 				
-				bool ViewJobAuthorized;
+				bool viewJobAuthorized;
 				using (IScope _ = GlobalTracer.Instance.BuildSpan("AuthorizeViewJob").StartActive())
 				{
-					ViewJobAuthorized = await JobService.AuthorizeAsync(Job, AclAction.ViewJob, User, PermissionsCache);
+					viewJobAuthorized = await _jobService.AuthorizeAsync(job, AclAction.ViewJob, User, permissionsCache);
 				}
 				
-				if (ViewJobAuthorized)
+				if (viewJobAuthorized)
 				{
-					IGraph Graph;
+					IGraph graph;
 					using (IScope _ = GlobalTracer.Instance.BuildSpan("GetGraph").StartActive())
 					{
-						Graph = await JobService.GetGraphAsync(Job);
+						graph = await _jobService.GetGraphAsync(job);
 					}
 
 					bool bIncludeAcl;
 					using (IScope _ = GlobalTracer.Instance.BuildSpan("AuthorizeViewPermissions").StartActive())
 					{
-						bIncludeAcl = await JobService.AuthorizeAsync(Job, AclAction.ViewPermissions, User, PermissionsCache);
+						bIncludeAcl = await _jobService.AuthorizeAsync(job, AclAction.ViewPermissions, User, permissionsCache);
 					}
 
 					bool bIncludeCosts;
 					using (IScope _ = GlobalTracer.Instance.BuildSpan("AuthorizeViewCosts").StartActive())
 					{
-						bIncludeCosts = await JobService.AuthorizeAsync(Job, AclAction.ViewCosts, User, PermissionsCache);
+						bIncludeCosts = await _jobService.AuthorizeAsync(job, AclAction.ViewCosts, User, permissionsCache);
 					}
 
 					using (IScope _ = GlobalTracer.Instance.BuildSpan("CreateResponse").StartActive())
 					{
-						Responses.Add(await CreateJobResponseAsync(Job, Graph, bIncludeAcl, bIncludeCosts, Filter));
+						responses.Add(await CreateJobResponseAsync(job, graph, bIncludeAcl, bIncludeCosts, filter));
 					}
 				}
 			}
 
 			
-			return Responses;
+			return responses;
 		}
 
 		/// <summary>
 		/// Adds an array of nodes to be executed for a job
 		/// </summary>
-		/// <param name="JobId">Unique id for the job</param>
-		/// <param name="Requests">Properties of the new nodes</param>
+		/// <param name="jobId">Unique id for the job</param>
+		/// <param name="requests">Properties of the new nodes</param>
 		/// <returns>Id of the new job</returns>
 		[HttpPost]
-		[Route("/api/v1/jobs/{JobId}/groups")]
-		public async Task<ActionResult> CreateGroupsAsync(JobId JobId, [FromBody] List<NewGroup> Requests)
+		[Route("/api/v1/jobs/{jobId}/groups")]
+		public async Task<ActionResult> CreateGroupsAsync(JobId jobId, [FromBody] List<NewGroup> requests)
 		{
-			Dictionary<string, int> ExpectedDurationCache = new Dictionary<string, int>();
-
 			for (; ; )
 			{
-				IJob? Job = await JobService.GetJobAsync(JobId);
-				if (Job == null)
+				IJob? job = await _jobService.GetJobAsync(jobId);
+				if (job == null)
 				{
-					return NotFound(JobId);
+					return NotFound(jobId);
 				}
-				if (!await JobService.AuthorizeAsync(Job, AclAction.ExecuteJob, User, null))
+				if (!await _jobService.AuthorizeAsync(job, AclAction.ExecuteJob, User, null))
 				{
-					return Forbid(AclAction.ExecuteJob, JobId);
+					return Forbid(AclAction.ExecuteJob, jobId);
 				}
 
-				IGraph Graph = await JobService.GetGraphAsync(Job);
-				Graph = await Graphs.AppendAsync(Graph, Requests, null, null);
+				IGraph graph = await _jobService.GetGraphAsync(job);
+				graph = await _graphs.AppendAsync(graph, requests, null, null);
 
-				IJob? NewJob = await JobService.TryUpdateGraphAsync(Job, Graph);
-				if (NewJob != null)
+				IJob? newJob = await _jobService.TryUpdateGraphAsync(job, graph);
+				if (newJob != null)
 				{
 					return Ok();
 				}
@@ -921,188 +912,188 @@ namespace Horde.Build.Controllers
 		/// <summary>
 		/// Gets the nodes to be executed for a job
 		/// </summary>
-		/// <param name="JobId">Unique id for the job</param>
-		/// <param name="Filter">Filter for the properties to return</param>
+		/// <param name="jobId">Unique id for the job</param>
+		/// <param name="filter">Filter for the properties to return</param>
 		/// <returns>List of nodes to be executed</returns>
 		[HttpGet]
-		[Route("/api/v1/jobs/{JobId}/groups")]
+		[Route("/api/v1/jobs/{jobId}/groups")]
 		[ProducesResponseType(typeof(List<GetGroupResponse>), 200)]
-		public async Task<ActionResult<List<object>>> GetGroupsAsync(JobId JobId, [FromQuery] PropertyFilter? Filter = null)
+		public async Task<ActionResult<List<object>>> GetGroupsAsync(JobId jobId, [FromQuery] PropertyFilter? filter = null)
 		{
-			IJob? Job = await JobService.GetJobAsync(JobId);
-			if (Job == null)
+			IJob? job = await _jobService.GetJobAsync(jobId);
+			if (job == null)
 			{
-				return NotFound(JobId);
+				return NotFound(jobId);
 			}
-			if (!await JobService.AuthorizeAsync(Job, AclAction.ViewJob, User, null))
+			if (!await _jobService.AuthorizeAsync(job, AclAction.ViewJob, User, null))
 			{
-				return Forbid(AclAction.ViewJob, JobId);
+				return Forbid(AclAction.ViewJob, jobId);
 			}
 
-			IGraph Graph = await JobService.GetGraphAsync(Job);
-			return Graph.Groups.ConvertAll(x => new GetGroupResponse(x, Graph.Groups).ApplyFilter(Filter));
+			IGraph graph = await _jobService.GetGraphAsync(job);
+			return graph.Groups.ConvertAll(x => new GetGroupResponse(x, graph.Groups).ApplyFilter(filter));
 		}
 
 		/// <summary>
 		/// Gets the nodes in a group to be executed for a job
 		/// </summary>
-		/// <param name="JobId">Unique id for the job</param>
-		/// <param name="GroupIdx">The group index</param>
-		/// <param name="Filter">Filter for the properties to return</param>
+		/// <param name="jobId">Unique id for the job</param>
+		/// <param name="groupIdx">The group index</param>
+		/// <param name="filter">Filter for the properties to return</param>
 		/// <returns>List of nodes to be executed</returns>
 		[HttpGet]
-		[Route("/api/v1/jobs/{JobId}/groups/{GroupIdx}")]
+		[Route("/api/v1/jobs/{jobId}/groups/{groupIdx}")]
 		[ProducesResponseType(typeof(GetGroupResponse), 200)]
-		public async Task<ActionResult<object>> GetGroupAsync(JobId JobId, int GroupIdx, [FromQuery] PropertyFilter? Filter = null)
+		public async Task<ActionResult<object>> GetGroupAsync(JobId jobId, int groupIdx, [FromQuery] PropertyFilter? filter = null)
 		{
-			IJob? Job = await JobService.GetJobAsync(JobId);
-			if (Job == null)
+			IJob? job = await _jobService.GetJobAsync(jobId);
+			if (job == null)
 			{
-				return NotFound(JobId);
+				return NotFound(jobId);
 			}
-			if (!await JobService.AuthorizeAsync(Job, AclAction.ViewJob, User, null))
+			if (!await _jobService.AuthorizeAsync(job, AclAction.ViewJob, User, null))
 			{
-				return Forbid(AclAction.ViewJob, JobId);
+				return Forbid(AclAction.ViewJob, jobId);
 			}
 
-			IGraph Graph = await JobService.GetGraphAsync(Job);
-			if (GroupIdx < 0 || GroupIdx >= Graph.Groups.Count)
+			IGraph graph = await _jobService.GetGraphAsync(job);
+			if (groupIdx < 0 || groupIdx >= graph.Groups.Count)
 			{
-				return NotFound(JobId, GroupIdx);
+				return NotFound(jobId, groupIdx);
 			}
 
-			return new GetGroupResponse(Graph.Groups[GroupIdx], Graph.Groups).ApplyFilter(Filter);
+			return new GetGroupResponse(graph.Groups[groupIdx], graph.Groups).ApplyFilter(filter);
 		}
 
 		/// <summary>
 		/// Gets the nodes for a particular group
 		/// </summary>
-		/// <param name="JobId">Unique id for the job</param>
-		/// <param name="GroupIdx">Index of the group containing the node to update</param>
-		/// <param name="Filter">Filter for the properties to return</param>
+		/// <param name="jobId">Unique id for the job</param>
+		/// <param name="groupIdx">Index of the group containing the node to update</param>
+		/// <param name="filter">Filter for the properties to return</param>
 		[HttpGet]
-		[Route("/api/v1/jobs/{JobId}/groups/{GroupIdx}/nodes")]
+		[Route("/api/v1/jobs/{jobId}/groups/{groupIdx}/nodes")]
 		[ProducesResponseType(typeof(List<GetNodeResponse>), 200)]
-		public async Task<ActionResult<List<object>>> GetNodesAsync(JobId JobId, int GroupIdx, [FromQuery] PropertyFilter? Filter = null)
+		public async Task<ActionResult<List<object>>> GetNodesAsync(JobId jobId, int groupIdx, [FromQuery] PropertyFilter? filter = null)
 		{
-			IJob? Job = await JobService.GetJobAsync(JobId);
-			if (Job == null)
+			IJob? job = await _jobService.GetJobAsync(jobId);
+			if (job == null)
 			{
-				return NotFound(JobId);
+				return NotFound(jobId);
 			}
-			if (!await JobService.AuthorizeAsync(Job, AclAction.ViewJob, User, null))
+			if (!await _jobService.AuthorizeAsync(job, AclAction.ViewJob, User, null))
 			{
-				return Forbid(AclAction.ViewJob, JobId);
+				return Forbid(AclAction.ViewJob, jobId);
 			}
 
-			IGraph Graph = await JobService.GetGraphAsync(Job);
-			if (GroupIdx < 0 || GroupIdx >= Graph.Groups.Count)
+			IGraph graph = await _jobService.GetGraphAsync(job);
+			if (groupIdx < 0 || groupIdx >= graph.Groups.Count)
 			{
-				return NotFound(JobId, GroupIdx);
+				return NotFound(jobId, groupIdx);
 			}
 
-			return Graph.Groups[GroupIdx].Nodes.ConvertAll(x => new GetNodeResponse(x, Graph.Groups).ApplyFilter(Filter));
+			return graph.Groups[groupIdx].Nodes.ConvertAll(x => new GetNodeResponse(x, graph.Groups).ApplyFilter(filter));
 		}
 
 		/// <summary>
 		/// Gets a particular node definition
 		/// </summary>
-		/// <param name="JobId">Unique id for the job</param>
-		/// <param name="GroupIdx">Index of the group containing the node to update</param>
-		/// <param name="NodeIdx">Index of the node to update</param>
-		/// <param name="Filter">Filter for the properties to return</param>
+		/// <param name="jobId">Unique id for the job</param>
+		/// <param name="groupIdx">Index of the group containing the node to update</param>
+		/// <param name="nodeIdx">Index of the node to update</param>
+		/// <param name="filter">Filter for the properties to return</param>
 		[HttpGet]
-		[Route("/api/v1/jobs/{JobId}/groups/{GroupIdx}/nodes/{NodeIdx}")]
+		[Route("/api/v1/jobs/{jobId}/groups/{groupIdx}/nodes/{nodeIdx}")]
 		[ProducesResponseType(typeof(GetNodeResponse), 200)]
-		public async Task<ActionResult<object>> GetNodeAsync(JobId JobId, int GroupIdx, int NodeIdx, [FromQuery] PropertyFilter? Filter = null)
+		public async Task<ActionResult<object>> GetNodeAsync(JobId jobId, int groupIdx, int nodeIdx, [FromQuery] PropertyFilter? filter = null)
 		{
-			IJob? Job = await JobService.GetJobAsync(JobId);
-			if (Job == null)
+			IJob? job = await _jobService.GetJobAsync(jobId);
+			if (job == null)
 			{
-				return NotFound(JobId);
+				return NotFound(jobId);
 			}
-			if (!await JobService.AuthorizeAsync(Job, AclAction.ViewJob, User, null))
+			if (!await _jobService.AuthorizeAsync(job, AclAction.ViewJob, User, null))
 			{
-				return Forbid(AclAction.ViewJob, JobId);
+				return Forbid(AclAction.ViewJob, jobId);
 			}
 
-			IGraph Graph = await JobService.GetGraphAsync(Job);
-			if (GroupIdx < 0 || GroupIdx >= Graph.Groups.Count)
+			IGraph graph = await _jobService.GetGraphAsync(job);
+			if (groupIdx < 0 || groupIdx >= graph.Groups.Count)
 			{
-				return NotFound(JobId, GroupIdx);
+				return NotFound(jobId, groupIdx);
 			}
-			if(NodeIdx < 0 || NodeIdx >= Graph.Groups[GroupIdx].Nodes.Count)
+			if(nodeIdx < 0 || nodeIdx >= graph.Groups[groupIdx].Nodes.Count)
 			{
-				return NotFound(JobId, GroupIdx, NodeIdx);
+				return NotFound(jobId, groupIdx, nodeIdx);
 			}
 
-			return new GetNodeResponse(Graph.Groups[GroupIdx].Nodes[NodeIdx], Graph.Groups).ApplyFilter(Filter);
+			return new GetNodeResponse(graph.Groups[groupIdx].Nodes[nodeIdx], graph.Groups).ApplyFilter(filter);
 		}
 
 		/// <summary>
 		/// Gets the steps currently scheduled to be executed for a job
 		/// </summary>
-		/// <param name="JobId">Unique id for the job</param>
-		/// <param name="Filter">Filter for the properties to return</param>
+		/// <param name="jobId">Unique id for the job</param>
+		/// <param name="filter">Filter for the properties to return</param>
 		/// <returns>List of nodes to be executed</returns>
 		[HttpGet]
-		[Route("/api/v1/jobs/{JobId}/batches")]
+		[Route("/api/v1/jobs/{jobId}/batches")]
 		[ProducesResponseType(typeof(List<GetBatchResponse>), 200)]
-		public async Task<ActionResult<List<object>>> GetBatchesAsync(JobId JobId, [FromQuery] PropertyFilter? Filter = null)
+		public async Task<ActionResult<List<object>>> GetBatchesAsync(JobId jobId, [FromQuery] PropertyFilter? filter = null)
 		{
-			IJob? Job = await JobService.GetJobAsync(JobId);
-			if (Job == null)
+			IJob? job = await _jobService.GetJobAsync(jobId);
+			if (job == null)
 			{
-				return NotFound(JobId);
+				return NotFound(jobId);
 			}
 
-			StreamPermissionsCache Cache = new StreamPermissionsCache();
-			if (!await JobService.AuthorizeAsync(Job, AclAction.ViewJob, User, Cache))
+			StreamPermissionsCache cache = new StreamPermissionsCache();
+			if (!await _jobService.AuthorizeAsync(job, AclAction.ViewJob, User, cache))
 			{
-				return Forbid(AclAction.ViewJob, JobId);
+				return Forbid(AclAction.ViewJob, jobId);
 			}
 
-			bool bIncludeCosts = await JobService.AuthorizeAsync(Job, AclAction.ViewCosts, User, Cache);
+			bool bIncludeCosts = await _jobService.AuthorizeAsync(job, AclAction.ViewCosts, User, cache);
 
-			List<object> Responses = new List<object>();
-			foreach (IJobStepBatch Batch in Job.Batches)
+			List<object> responses = new List<object>();
+			foreach (IJobStepBatch batch in job.Batches)
 			{
-				GetBatchResponse Response = await CreateBatchResponseAsync(Batch, bIncludeCosts);
-				Responses.Add(Response.ApplyFilter(Filter));
+				GetBatchResponse response = await CreateBatchResponseAsync(batch, bIncludeCosts);
+				responses.Add(response.ApplyFilter(filter));
 			}
-			return Responses;
+			return responses;
 		}
 
 		/// <summary>
 		/// Updates the state of a jobstep
 		/// </summary>
-		/// <param name="JobId">Unique id for the job</param>
-		/// <param name="BatchId">Unique id for the step</param>
-		/// <param name="Request">Updates to apply to the node</param>
+		/// <param name="jobId">Unique id for the job</param>
+		/// <param name="batchId">Unique id for the step</param>
+		/// <param name="request">Updates to apply to the node</param>
 		[HttpPut]
-		[Route("/api/v1/jobs/{JobId}/batches/{BatchId}")]
-		public async Task<ActionResult> UpdateBatchAsync(JobId JobId, SubResourceId BatchId, [FromBody] UpdateBatchRequest Request)
+		[Route("/api/v1/jobs/{jobId}/batches/{batchId}")]
+		public async Task<ActionResult> UpdateBatchAsync(JobId jobId, SubResourceId batchId, [FromBody] UpdateBatchRequest request)
 		{
-			IJob? Job = await JobService.GetJobAsync(JobId);
-			if (Job == null)
+			IJob? job = await _jobService.GetJobAsync(jobId);
+			if (job == null)
 			{
-				return NotFound(JobId);
+				return NotFound(jobId);
 			}
 
-			IJobStepBatch? Batch = Job.Batches.FirstOrDefault(x => x.Id == BatchId);
-			if (Batch == null)
+			IJobStepBatch? batch = job.Batches.FirstOrDefault(x => x.Id == batchId);
+			if (batch == null)
 			{
-				return NotFound(JobId, BatchId);
+				return NotFound(jobId, batchId);
 			}
-			if (Batch.SessionId == null || !User.HasSessionClaim(Batch.SessionId.Value))
+			if (batch.SessionId == null || !User.HasSessionClaim(batch.SessionId.Value))
 			{
-				return Forbid("Missing session claim for job {JobId} batch {BatchId}", JobId, BatchId);
+				return Forbid("Missing session claim for job {JobId} batch {BatchId}", jobId, batchId);
 			}
 
-			IJob? NewJob = await JobService.UpdateBatchAsync(Job, BatchId, Request.LogId?.ToObjectId<ILogFile>(), Request.State);
-			if (NewJob == null)
+			IJob? newJob = await _jobService.UpdateBatchAsync(job, batchId, request.LogId?.ToObjectId<ILogFile>(), request.State);
+			if (newJob == null)
 			{
-				return NotFound(JobId);
+				return NotFound(jobId);
 			}
 			return Ok();
 		}
@@ -1110,179 +1101,177 @@ namespace Horde.Build.Controllers
 		/// <summary>
 		/// Gets a particular step currently scheduled to be executed for a job
 		/// </summary>
-		/// <param name="JobId">Unique id for the job</param>
-		/// <param name="BatchId">Unique id for the step</param>
-		/// <param name="Filter">Filter for the properties to return</param>
+		/// <param name="jobId">Unique id for the job</param>
+		/// <param name="batchId">Unique id for the step</param>
+		/// <param name="filter">Filter for the properties to return</param>
 		/// <returns>List of nodes to be executed</returns>
 		[HttpGet]
-		[Route("/api/v1/jobs/{JobId}/batches/{BatchId}")]
+		[Route("/api/v1/jobs/{jobId}/batches/{batchId}")]
 		[ProducesResponseType(typeof(GetBatchResponse), 200)]
-		public async Task<ActionResult<object>> GetBatchAsync(JobId JobId, SubResourceId BatchId, [FromQuery] PropertyFilter? Filter = null)
+		public async Task<ActionResult<object>> GetBatchAsync(JobId jobId, SubResourceId batchId, [FromQuery] PropertyFilter? filter = null)
 		{
-			IJob? Job = await JobService.GetJobAsync(JobId);
-			if (Job == null)
+			IJob? job = await _jobService.GetJobAsync(jobId);
+			if (job == null)
 			{
-				return NotFound(JobId);
+				return NotFound(jobId);
 			}
 
-			StreamPermissionsCache Cache = new StreamPermissionsCache();
-			if (!await JobService.AuthorizeAsync(Job, AclAction.ViewJob, User, Cache))
+			StreamPermissionsCache cache = new StreamPermissionsCache();
+			if (!await _jobService.AuthorizeAsync(job, AclAction.ViewJob, User, cache))
 			{
-				return Forbid(AclAction.ViewJob, JobId);
+				return Forbid(AclAction.ViewJob, jobId);
 			}
 
-			bool bIncludeCosts = await JobService.AuthorizeAsync(Job, AclAction.ViewCosts, User, Cache);
+			bool bIncludeCosts = await _jobService.AuthorizeAsync(job, AclAction.ViewCosts, User, cache);
 
-			IGraph Graph = await JobService.GetGraphAsync(Job);
-			foreach (IJobStepBatch Batch in Job.Batches)
+			foreach (IJobStepBatch batch in job.Batches)
 			{
-				if (Batch.Id == BatchId)
+				if (batch.Id == batchId)
 				{
-					GetBatchResponse Response = await CreateBatchResponseAsync(Batch, bIncludeCosts);
-					return Response.ApplyFilter(Filter);
+					GetBatchResponse response = await CreateBatchResponseAsync(batch, bIncludeCosts);
+					return response.ApplyFilter(filter);
 				}
 			}
 
-			return NotFound(JobId, BatchId);
+			return NotFound(jobId, batchId);
 		}
 
 		/// <summary>
 		/// Gets the steps currently scheduled to be executed for a job
 		/// </summary>
-		/// <param name="JobId">Unique id for the job</param>
-		/// <param name="BatchId">Unique id for the batch</param>
-		/// <param name="Filter">Filter for the properties to return</param>
+		/// <param name="jobId">Unique id for the job</param>
+		/// <param name="batchId">Unique id for the batch</param>
+		/// <param name="filter">Filter for the properties to return</param>
 		/// <returns>List of nodes to be executed</returns>
 		[HttpGet]
-		[Route("/api/v1/jobs/{JobId}/batches/{BatchId}/steps")]
+		[Route("/api/v1/jobs/{jobId}/batches/{batchId}/steps")]
 		[ProducesResponseType(typeof(List<GetStepResponse>), 200)]
-		public async Task<ActionResult<List<object>>> GetStepsAsync(JobId JobId, SubResourceId BatchId, [FromQuery] PropertyFilter? Filter = null)
+		public async Task<ActionResult<List<object>>> GetStepsAsync(JobId jobId, SubResourceId batchId, [FromQuery] PropertyFilter? filter = null)
 		{
-			IJob? Job = await JobService.GetJobAsync(JobId);
-			if (Job == null)
+			IJob? job = await _jobService.GetJobAsync(jobId);
+			if (job == null)
 			{
-				return NotFound(JobId);
+				return NotFound(jobId);
 			}
-			if (!await JobService.AuthorizeAsync(Job, AclAction.ViewJob, User, null))
+			if (!await _jobService.AuthorizeAsync(job, AclAction.ViewJob, User, null))
 			{
-				return Forbid(AclAction.ViewJob, JobId);
+				return Forbid(AclAction.ViewJob, jobId);
 			}
 
-			IGraph Graph = await JobService.GetGraphAsync(Job);
-			foreach (IJobStepBatch Batch in Job.Batches)
+			foreach (IJobStepBatch batch in job.Batches)
 			{
-				if (Batch.Id == BatchId)
+				if (batch.Id == batchId)
 				{
-					List<object> Responses = new List<object>();
-					foreach (IJobStep Step in Batch.Steps)
+					List<object> responses = new List<object>();
+					foreach (IJobStep step in batch.Steps)
 					{
-						GetStepResponse Response = await CreateStepResponseAsync(Step);
-						Responses.Add(Response.ApplyFilter(Filter));
+						GetStepResponse response = await CreateStepResponseAsync(step);
+						responses.Add(response.ApplyFilter(filter));
 					}
-					return Responses;
+					return responses;
 				}
 			}
 
-			return NotFound(JobId, BatchId);
+			return NotFound(jobId, batchId);
 		}
 
 		/// <summary>
 		/// Updates the state of a jobstep
 		/// </summary>
-		/// <param name="JobId">Unique id for the job</param>
-		/// <param name="BatchId">Unique id for the batch</param>
-		/// <param name="StepId">Unique id for the step</param>
-		/// <param name="Request">Updates to apply to the node</param>
+		/// <param name="jobId">Unique id for the job</param>
+		/// <param name="batchId">Unique id for the batch</param>
+		/// <param name="stepId">Unique id for the step</param>
+		/// <param name="request">Updates to apply to the node</param>
 		[HttpPut]
-		[Route("/api/v1/jobs/{JobId}/batches/{BatchId}/steps/{StepId}")]
-		public async Task<ActionResult<UpdateStepResponse>> UpdateStepAsync(JobId JobId, SubResourceId BatchId, SubResourceId StepId, [FromBody] UpdateStepRequest Request)
+		[Route("/api/v1/jobs/{jobId}/batches/{batchId}/steps/{stepId}")]
+		public async Task<ActionResult<UpdateStepResponse>> UpdateStepAsync(JobId jobId, SubResourceId batchId, SubResourceId stepId, [FromBody] UpdateStepRequest request)
 		{
-			IJob? Job = await JobService.GetJobAsync(JobId);
-			if (Job == null)
+			IJob? job = await _jobService.GetJobAsync(jobId);
+			if (job == null)
 			{
-				return NotFound(JobId);
+				return NotFound(jobId);
 			}
 
 			// Check permissions for updating this step. Only the agent executing the step can modify the state of it.
-			if (Request.State != JobStepState.Unspecified || Request.Outcome != JobStepOutcome.Unspecified)
+			if (request.State != JobStepState.Unspecified || request.Outcome != JobStepOutcome.Unspecified)
 			{
-				IJobStepBatch? Batch = Job.Batches.FirstOrDefault(x => x.Id == BatchId);
-				if (Batch == null)
+				IJobStepBatch? batch = job.Batches.FirstOrDefault(x => x.Id == batchId);
+				if (batch == null)
 				{
-					return NotFound(JobId, BatchId);
+					return NotFound(jobId, batchId);
 				}
-				if (!Batch.SessionId.HasValue || !User.HasSessionClaim(Batch.SessionId.Value))
+				if (!batch.SessionId.HasValue || !User.HasSessionClaim(batch.SessionId.Value))
 				{
 					return Forbid();
 				}
 			}
-			if (Request.Retry != null || Request.Priority != null)
+			if (request.Retry != null || request.Priority != null)
 			{
-				if (!await JobService.AuthorizeAsync(Job, AclAction.RetryJobStep, User, null))
+				if (!await _jobService.AuthorizeAsync(job, AclAction.RetryJobStep, User, null))
 				{
-					return Forbid(AclAction.RetryJobStep, JobId);
+					return Forbid(AclAction.RetryJobStep, jobId);
 				}
 			}
-			if (Request.Properties != null)
+			if (request.Properties != null)
 			{
-				if (!await JobService.AuthorizeAsync(Job, AclAction.UpdateJob, User, null))
+				if (!await _jobService.AuthorizeAsync(job, AclAction.UpdateJob, User, null))
 				{
-					return Forbid(AclAction.UpdateJob, JobId);
+					return Forbid(AclAction.UpdateJob, jobId);
 				}
 			}
 
-			UserId? RetryByUser = (Request.Retry.HasValue && Request.Retry.Value) ? User.GetUserId() : null;
-			UserId? AbortByUser = (Request.AbortRequested.HasValue && Request.AbortRequested.Value) ? User.GetUserId() : null;
+			UserId? retryByUser = (request.Retry.HasValue && request.Retry.Value) ? User.GetUserId() : null;
+			UserId? abortByUser = (request.AbortRequested.HasValue && request.AbortRequested.Value) ? User.GetUserId() : null;
 
 			try
 			{
-				IJob? NewJob = await JobService.UpdateStepAsync(Job, BatchId, StepId, Request.State, Request.Outcome, Request.AbortRequested, AbortByUser, Request.LogId?.ToObjectId<ILogFile>(), null, RetryByUser, Request.Priority, null, Request.Properties);
-				if (NewJob == null)
+				IJob? newJob = await _jobService.UpdateStepAsync(job, batchId, stepId, request.State, request.Outcome, request.AbortRequested, abortByUser, request.LogId?.ToObjectId<ILogFile>(), null, retryByUser, request.Priority, null, request.Properties);
+				if (newJob == null)
 				{
-					return NotFound(JobId);
+					return NotFound(jobId);
 				}
 
-				UpdateStepResponse Response = new UpdateStepResponse();
-				if (Request.Retry ?? false)
+				UpdateStepResponse response = new UpdateStepResponse();
+				if (request.Retry ?? false)
 				{
-					JobStepRefId? RetriedStepId = FindRetriedStep(Job, BatchId, StepId);
-					if (RetriedStepId != null)
+					JobStepRefId? retriedStepId = FindRetriedStep(job, batchId, stepId);
+					if (retriedStepId != null)
 					{
-						Response.BatchId = RetriedStepId.Value.BatchId.ToString();
-						Response.StepId = RetriedStepId.Value.StepId.ToString();
+						response.BatchId = retriedStepId.Value.BatchId.ToString();
+						response.StepId = retriedStepId.Value.StepId.ToString();
 					}
 				}
-				return Response;
+				return response;
 			}
-			catch (RetryNotAllowedException Ex)
+			catch (RetryNotAllowedException ex)
 			{
-				return BadRequest(Ex.Message);
+				return BadRequest(ex.Message);
 			}
 		}
 
 		/// <summary>
 		/// Find the first retried step after the given step
 		/// </summary>
-		/// <param name="Job">The job being run</param>
-		/// <param name="BatchId">Batch id of the last step instance</param>
-		/// <param name="StepId">Step id of the last instance</param>
+		/// <param name="job">The job being run</param>
+		/// <param name="batchId">Batch id of the last step instance</param>
+		/// <param name="stepId">Step id of the last instance</param>
 		/// <returns>The retried step information</returns>
-		static JobStepRefId? FindRetriedStep(IJob Job, SubResourceId BatchId, SubResourceId StepId)
+		static JobStepRefId? FindRetriedStep(IJob job, SubResourceId batchId, SubResourceId stepId)
 		{
-			NodeRef? LastNodeRef = null;
-			foreach (IJobStepBatch Batch in Job.Batches)
+			NodeRef? lastNodeRef = null;
+			foreach (IJobStepBatch batch in job.Batches)
 			{
-				if ((LastNodeRef == null && Batch.Id == BatchId) || (LastNodeRef != null && Batch.GroupIdx == LastNodeRef.GroupIdx))
+				if ((lastNodeRef == null && batch.Id == batchId) || (lastNodeRef != null && batch.GroupIdx == lastNodeRef.GroupIdx))
 				{
-					foreach (IJobStep Step in Batch.Steps)
+					foreach (IJobStep step in batch.Steps)
 					{
-						if (LastNodeRef == null && Step.Id == StepId)
+						if (lastNodeRef == null && step.Id == stepId)
 						{
-							LastNodeRef = new NodeRef(Batch.GroupIdx, Step.NodeIdx);
+							lastNodeRef = new NodeRef(batch.GroupIdx, step.NodeIdx);
 						}
-						else if (LastNodeRef != null && Step.NodeIdx == LastNodeRef.NodeIdx)
+						else if (lastNodeRef != null && step.NodeIdx == lastNodeRef.NodeIdx)
 						{
-							return new JobStepRefId(Job.Id, Batch.Id, Step.Id);
+							return new JobStepRefId(job.Id, batch.Id, step.Id);
 						}
 					}
 				}
@@ -1293,209 +1282,209 @@ namespace Horde.Build.Controllers
 		/// <summary>
 		/// Gets a particular step currently scheduled to be executed for a job
 		/// </summary>
-		/// <param name="JobId">Unique id for the job</param>
-		/// <param name="BatchId">Unique id for the batch</param>
-		/// <param name="StepId">Unique id for the step</param>
-		/// <param name="Filter">Filter for the properties to return</param>
+		/// <param name="jobId">Unique id for the job</param>
+		/// <param name="batchId">Unique id for the batch</param>
+		/// <param name="stepId">Unique id for the step</param>
+		/// <param name="filter">Filter for the properties to return</param>
 		/// <returns>List of nodes to be executed</returns>
 		[HttpGet]
-		[Route("/api/v1/jobs/{JobId}/batches/{BatchId}/steps/{StepId}")]
+		[Route("/api/v1/jobs/{jobId}/batches/{batchId}/steps/{stepId}")]
 		[ProducesResponseType(typeof(GetStepResponse), 200)]
-		public async Task<ActionResult<object>> GetStepAsync(JobId JobId, SubResourceId BatchId, SubResourceId StepId, [FromQuery] PropertyFilter? Filter = null)
+		public async Task<ActionResult<object>> GetStepAsync(JobId jobId, SubResourceId batchId, SubResourceId stepId, [FromQuery] PropertyFilter? filter = null)
 		{
-			IJob? Job = await JobService.GetJobAsync(JobId);
-			if (Job == null)
+			IJob? job = await _jobService.GetJobAsync(jobId);
+			if (job == null)
 			{
-				return NotFound(JobId);
+				return NotFound(jobId);
 			}
-			if (!await JobService.AuthorizeAsync(Job, AclAction.ViewJob, User, null))
+			if (!await _jobService.AuthorizeAsync(job, AclAction.ViewJob, User, null))
 			{
-				return Forbid(AclAction.ViewJob, JobId);
+				return Forbid(AclAction.ViewJob, jobId);
 			}
 
-			foreach (IJobStepBatch Batch in Job.Batches)
+			foreach (IJobStepBatch batch in job.Batches)
 			{
-				if (Batch.Id == BatchId)
+				if (batch.Id == batchId)
 				{
-					foreach (IJobStep Step in Batch.Steps)
+					foreach (IJobStep step in batch.Steps)
 					{
-						if (Step.Id == StepId)
+						if (step.Id == stepId)
 						{
-							GetStepResponse Response = await CreateStepResponseAsync(Step);
-							return Response.ApplyFilter(Filter);
+							GetStepResponse response = await CreateStepResponseAsync(step);
+							return response.ApplyFilter(filter);
 						}
 					}
-					return NotFound(JobId, BatchId, StepId);
+					return NotFound(jobId, batchId, stepId);
 				}
 			}
 
-			return NotFound(JobId, BatchId);
+			return NotFound(jobId, batchId);
 		}
 
 		/// <summary>
 		/// Updates notifications for a specific job.
 		/// </summary>
-		/// <param name="JobId">Unique id for the job</param>
-		/// <param name="BatchId">Unique id for the batch</param>
-		/// <param name="StepId">Unique id for the step</param>
-		/// <param name="Request">The notification request</param>
+		/// <param name="jobId">Unique id for the job</param>
+		/// <param name="batchId">Unique id for the batch</param>
+		/// <param name="stepId">Unique id for the step</param>
+		/// <param name="request">The notification request</param>
 		/// <returns>Information about the requested job</returns>
 		[HttpPut]
-		[Route("/api/v1/jobs/{JobId}/batches/{BatchId}/steps/{StepId}/notifications")]
-		public async Task<ActionResult> UpdateStepNotificationsAsync(JobId JobId, SubResourceId BatchId, SubResourceId StepId, [FromBody] UpdateNotificationsRequest Request)
+		[Route("/api/v1/jobs/{jobId}/batches/{batchId}/steps/{stepId}/notifications")]
+		public async Task<ActionResult> UpdateStepNotificationsAsync(JobId jobId, SubResourceId batchId, SubResourceId stepId, [FromBody] UpdateNotificationsRequest request)
 		{
-			IJob? Job = await JobService.GetJobAsync(JobId);
-			if (Job == null)
+			IJob? job = await _jobService.GetJobAsync(jobId);
+			if (job == null)
 			{
-				return NotFound(JobId);
+				return NotFound(jobId);
 			}
-			if (!await JobService.AuthorizeAsync(Job, AclAction.CreateSubscription, User, null))
+			if (!await _jobService.AuthorizeAsync(job, AclAction.CreateSubscription, User, null))
 			{
-				return Forbid(AclAction.CreateSubscription, JobId);
+				return Forbid(AclAction.CreateSubscription, jobId);
 			}
 
-			if (!Job.TryGetBatch(BatchId, out IJobStepBatch? Batch))
+			if (!job.TryGetBatch(batchId, out IJobStepBatch? batch))
 			{
-				return NotFound(JobId, BatchId);
+				return NotFound(jobId, batchId);
 			}
 
-			if (!Batch.TryGetStep(StepId, out IJobStep? Step))
+			if (!batch.TryGetStep(stepId, out IJobStep? step))
 			{
-				return NotFound(JobId, BatchId, StepId);
+				return NotFound(jobId, batchId, stepId);
 			}
 
-			ObjectId? TriggerId = Step.NotificationTriggerId;
-			if (TriggerId == null)
+			ObjectId? triggerId = step.NotificationTriggerId;
+			if (triggerId == null)
 			{
-				TriggerId = ObjectId.GenerateNewId();
-				if (await JobService.UpdateStepAsync(Job, BatchId, StepId, JobStepState.Unspecified, JobStepOutcome.Unspecified, null, null, null, TriggerId, null, null, null) == null)
+				triggerId = ObjectId.GenerateNewId();
+				if (await _jobService.UpdateStepAsync(job, batchId, stepId, JobStepState.Unspecified, JobStepOutcome.Unspecified, null, null, null, triggerId, null, null, null) == null)
 				{
-					return NotFound(JobId, BatchId, StepId);
+					return NotFound(jobId, batchId, stepId);
 				}
 			}
 
-			await NotificationService.UpdateSubscriptionsAsync(TriggerId.Value, User, Request.Email, Request.Slack);
+			await _notificationService.UpdateSubscriptionsAsync(triggerId.Value, User, request.Email, request.Slack);
 			return Ok();
 		}
 
 		/// <summary>
 		/// Gets information about a specific job.
 		/// </summary>
-		/// <param name="JobId">Id of the job to find</param>
-		/// <param name="BatchId">Unique id for the batch</param>
-		/// <param name="StepId">Unique id for the step</param>
+		/// <param name="jobId">Id of the job to find</param>
+		/// <param name="batchId">Unique id for the batch</param>
+		/// <param name="stepId">Unique id for the step</param>
 		/// <returns>Information about the requested job</returns>
 		[HttpGet]
-		[Route("/api/v1/jobs/{JobId}/batches/{BatchId}/steps/{StepId}/notifications")]
-		public async Task<ActionResult<GetNotificationResponse>> GetStepNotificationsAsync(JobId JobId, SubResourceId BatchId, SubResourceId StepId)
+		[Route("/api/v1/jobs/{jobId}/batches/{batchId}/steps/{stepId}/notifications")]
+		public async Task<ActionResult<GetNotificationResponse>> GetStepNotificationsAsync(JobId jobId, SubResourceId batchId, SubResourceId stepId)
 		{
-			IJob? Job = await JobService.GetJobAsync(JobId);
-			if (Job == null)
+			IJob? job = await _jobService.GetJobAsync(jobId);
+			if (job == null)
 			{
-				return NotFound(JobId);
+				return NotFound(jobId);
 			}
 
-			IJobStep? Step;
-			if (!Job.TryGetBatch(BatchId, out IJobStepBatch? Batch))
+			IJobStep? step;
+			if (!job.TryGetBatch(batchId, out IJobStepBatch? batch))
 			{
-				return NotFound(JobId, BatchId);
+				return NotFound(jobId, batchId);
 			}
-			if (!Batch.TryGetStep(StepId, out Step))
+			if (!batch.TryGetStep(stepId, out step))
 			{
-				return NotFound(JobId, BatchId, StepId);
+				return NotFound(jobId, batchId, stepId);
 			}
-			if (!await JobService.AuthorizeAsync(Job, AclAction.CreateSubscription, User, null))
+			if (!await _jobService.AuthorizeAsync(job, AclAction.CreateSubscription, User, null))
 			{
-				return Forbid(AclAction.CreateSubscription, JobId);
+				return Forbid(AclAction.CreateSubscription, jobId);
 			}
 
-			INotificationSubscription? Subscription;
-			if (Step.NotificationTriggerId == null)
+			INotificationSubscription? subscription;
+			if (step.NotificationTriggerId == null)
 			{
-				Subscription = null;
+				subscription = null;
 			}
 			else
 			{
-				Subscription = await NotificationService.GetSubscriptionsAsync(Step.NotificationTriggerId.Value, User);
+				subscription = await _notificationService.GetSubscriptionsAsync(step.NotificationTriggerId.Value, User);
 			}
-			return new GetNotificationResponse(Subscription);
+			return new GetNotificationResponse(subscription);
 		}
 
 		/// <summary>
 		/// Gets a particular step currently scheduled to be executed for a job
 		/// </summary>
-		/// <param name="JobId">Unique id for the job</param>
-		/// <param name="BatchId">Unique id for the batch</param>
-		/// <param name="StepId">Unique id for the step</param>
-		/// <param name="Name"></param>
+		/// <param name="jobId">Unique id for the job</param>
+		/// <param name="batchId">Unique id for the batch</param>
+		/// <param name="stepId">Unique id for the step</param>
+		/// <param name="name"></param>
 		/// <returns>List of nodes to be executed</returns>
 		[HttpGet]
-		[Route("/api/v1/jobs/{JobId}/batches/{BatchId}/steps/{StepId}/artifacts/{*Name}")]
-		public async Task<ActionResult> GetArtifactAsync(JobId JobId, SubResourceId BatchId, SubResourceId StepId, string Name)
+		[Route("/api/v1/jobs/{jobId}/batches/{batchId}/steps/{stepId}/artifacts/{*name}")]
+		public async Task<ActionResult> GetArtifactAsync(JobId jobId, SubResourceId batchId, SubResourceId stepId, string name)
 		{
-			IJob? Job = await JobService.GetJobAsync(JobId);
-			if (Job == null)
+			IJob? job = await _jobService.GetJobAsync(jobId);
+			if (job == null)
 			{
-				return NotFound(JobId);
+				return NotFound(jobId);
 			}
-			if (!await JobService.AuthorizeAsync(Job, AclAction.ViewJob, User, null))
+			if (!await _jobService.AuthorizeAsync(job, AclAction.ViewJob, User, null))
 			{
-				return Forbid(AclAction.ViewJob, JobId);
+				return Forbid(AclAction.ViewJob, jobId);
 			}
-			if (!Job.TryGetBatch(BatchId, out IJobStepBatch? Batch))
+			if (!job.TryGetBatch(batchId, out IJobStepBatch? batch))
 			{
-				return NotFound(JobId, BatchId);
+				return NotFound(jobId, batchId);
 			}
-			if (!Batch.TryGetStep(StepId, out _))
+			if (!batch.TryGetStep(stepId, out _))
 			{
-				return NotFound(JobId, BatchId, StepId);
+				return NotFound(jobId, batchId, stepId);
 			}
 
-			List<IArtifact> Artifacts = await ArtifactCollection.GetArtifactsAsync(JobId, StepId, Name);
-			if (Artifacts.Count == 0)
+			List<IArtifact> artifacts = await _artifactCollection.GetArtifactsAsync(jobId, stepId, name);
+			if (artifacts.Count == 0)
 			{
 				return NotFound();
 			}
 
-			IArtifact Artifact = Artifacts[0];
-			return new FileStreamResult(await ArtifactCollection.OpenArtifactReadStreamAsync(Artifact), Artifact.MimeType);
+			IArtifact artifact = artifacts[0];
+			return new FileStreamResult(await _artifactCollection.OpenArtifactReadStreamAsync(artifact), artifact.MimeType);
 		}
 
 		/// <summary>
 		/// Gets a particular step currently scheduled to be executed for a job
 		/// </summary>
-		/// <param name="JobId">Unique id for the job</param>
-		/// <param name="BatchId">Unique id for the batch</param>
-		/// <param name="StepId">Unique id for the step</param>
+		/// <param name="jobId">Unique id for the job</param>
+		/// <param name="batchId">Unique id for the batch</param>
+		/// <param name="stepId">Unique id for the step</param>
 		/// <returns>List of nodes to be executed</returns>
 		[HttpGet]
-		[Route("/api/v1/jobs/{JobId}/batches/{BatchId}/steps/{StepId}/trace")]
-		public async Task<ActionResult> GetStepTraceAsync(JobId JobId, SubResourceId BatchId, SubResourceId StepId)
+		[Route("/api/v1/jobs/{jobId}/batches/{batchId}/steps/{stepId}/trace")]
+		public async Task<ActionResult> GetStepTraceAsync(JobId jobId, SubResourceId batchId, SubResourceId stepId)
 		{
-			IJob? Job = await JobService.GetJobAsync(JobId);
-			if (Job == null)
+			IJob? job = await _jobService.GetJobAsync(jobId);
+			if (job == null)
 			{
-				return NotFound(JobId);
+				return NotFound(jobId);
 			}
-			if (!await JobService.AuthorizeAsync(Job, AclAction.ViewJob, User, null))
+			if (!await _jobService.AuthorizeAsync(job, AclAction.ViewJob, User, null))
 			{
-				return Forbid(AclAction.ViewJob, JobId);
+				return Forbid(AclAction.ViewJob, jobId);
 			}
-			if (!Job.TryGetBatch(BatchId, out IJobStepBatch? Batch))
+			if (!job.TryGetBatch(batchId, out IJobStepBatch? batch))
 			{
-				return NotFound(JobId, BatchId);
+				return NotFound(jobId, batchId);
 			}
-			if (!Batch.TryGetStep(StepId, out _))
+			if (!batch.TryGetStep(stepId, out _))
 			{
-				return NotFound(JobId, BatchId, StepId);
+				return NotFound(jobId, batchId, stepId);
 			}
 
-			List<IArtifact> Artifacts = await ArtifactCollection.GetArtifactsAsync(JobId, StepId, null);
-			foreach (IArtifact Artifact in Artifacts)
+			List<IArtifact> artifacts = await _artifactCollection.GetArtifactsAsync(jobId, stepId, null);
+			foreach (IArtifact artifact in artifacts)
 			{
-				if (Artifact.Name.Equals("trace.json", StringComparison.OrdinalIgnoreCase))
+				if (artifact.Name.Equals("trace.json", StringComparison.OrdinalIgnoreCase))
 				{
-					return new FileStreamResult(await ArtifactCollection.OpenArtifactReadStreamAsync(Artifact), "text/json");
+					return new FileStreamResult(await _artifactCollection.OpenArtifactReadStreamAsync(artifact), "text/json");
 				}
 			}
 			return NotFound();
@@ -1504,73 +1493,73 @@ namespace Horde.Build.Controllers
 		/// <summary>
 		/// Updates notifications for a specific label.
 		/// </summary>
-		/// <param name="JobId">Unique id for the job</param>
-		/// <param name="LabelIndex">Index for the label</param>
-		/// <param name="Request">The notification request</param>
+		/// <param name="jobId">Unique id for the job</param>
+		/// <param name="labelIndex">Index for the label</param>
+		/// <param name="request">The notification request</param>
 		[HttpPut]
-		[Route("/api/v1/jobs/{JobId}/labels/{LabelIndex}/notifications")]
-		public async Task<ActionResult> UpdateLabelNotificationsAsync(JobId JobId, int LabelIndex, [FromBody] UpdateNotificationsRequest Request)
+		[Route("/api/v1/jobs/{jobId}/labels/{labelIndex}/notifications")]
+		public async Task<ActionResult> UpdateLabelNotificationsAsync(JobId jobId, int labelIndex, [FromBody] UpdateNotificationsRequest request)
 		{
-			ObjectId TriggerId;
+			ObjectId triggerId;
 			for (; ; )
 			{
-				IJob? Job = await JobService.GetJobAsync(JobId);
-				if (Job == null)
+				IJob? job = await _jobService.GetJobAsync(jobId);
+				if (job == null)
 				{
-					return NotFound(JobId);
+					return NotFound(jobId);
 				}
-				if (!await JobService.AuthorizeAsync(Job, AclAction.CreateSubscription, User, null))
+				if (!await _jobService.AuthorizeAsync(job, AclAction.CreateSubscription, User, null))
 				{
-					return Forbid(AclAction.CreateSubscription, JobId);
+					return Forbid(AclAction.CreateSubscription, jobId);
 				}
 
-				ObjectId NewTriggerId;
-				if (Job.LabelIdxToTriggerId.TryGetValue(LabelIndex, out NewTriggerId))
+				ObjectId newTriggerId;
+				if (job.LabelIdxToTriggerId.TryGetValue(labelIndex, out newTriggerId))
 				{
-					TriggerId = NewTriggerId;
+					triggerId = newTriggerId;
 					break;
 				}
 
-				NewTriggerId = ObjectId.GenerateNewId();
+				newTriggerId = ObjectId.GenerateNewId();
 
-				IJob? NewJob = await JobService.UpdateJobAsync(Job, LabelIdxToTriggerId: new KeyValuePair<int, ObjectId>(LabelIndex, NewTriggerId));
-				if (NewJob != null)
+				IJob? newJob = await _jobService.UpdateJobAsync(job, labelIdxToTriggerId: new KeyValuePair<int, ObjectId>(labelIndex, newTriggerId));
+				if (newJob != null)
 				{
-					TriggerId = NewTriggerId;
+					triggerId = newTriggerId;
 					break;
 				}
 			}
 
-			await NotificationService.UpdateSubscriptionsAsync(TriggerId, User, Request.Email, Request.Slack);
+			await _notificationService.UpdateSubscriptionsAsync(triggerId, User, request.Email, request.Slack);
 			return Ok();
 		}
 
 		/// <summary>
 		/// Gets notification info about a specific label in a job.
 		/// </summary>
-		/// <param name="JobId">Id of the job to find</param>
-		/// <param name="LabelIndex">Index for the label</param>
+		/// <param name="jobId">Id of the job to find</param>
+		/// <param name="labelIndex">Index for the label</param>
 		/// <returns>Notification info for the requested label in the job</returns>
 		[HttpGet]
-		[Route("/api/v1/jobs/{JobId}/labels/{LabelIndex}/notifications")]
-		public async Task<ActionResult<GetNotificationResponse>> GetLabelNotificationsAsync(JobId JobId, int LabelIndex)
+		[Route("/api/v1/jobs/{jobId}/labels/{labelIndex}/notifications")]
+		public async Task<ActionResult<GetNotificationResponse>> GetLabelNotificationsAsync(JobId jobId, int labelIndex)
 		{
-			IJob? Job = await JobService.GetJobAsync(JobId);
-			if (Job == null)
+			IJob? job = await _jobService.GetJobAsync(jobId);
+			if (job == null)
 			{
-				return NotFound(JobId);
+				return NotFound(jobId);
 			}
 
-			INotificationSubscription? Subscription;
-			if (!Job.LabelIdxToTriggerId.ContainsKey(LabelIndex))
+			INotificationSubscription? subscription;
+			if (!job.LabelIdxToTriggerId.ContainsKey(labelIndex))
 			{
-				Subscription = null;
+				subscription = null;
 			}
 			else
 			{
-				Subscription = await NotificationService.GetSubscriptionsAsync(Job.LabelIdxToTriggerId[LabelIndex], User);
+				subscription = await _notificationService.GetSubscriptionsAsync(job.LabelIdxToTriggerId[labelIndex], User);
 			}
-			return new GetNotificationResponse(Subscription);
+			return new GetNotificationResponse(subscription);
 		}
 	}
 }
