Horde: Fix coding conventions within Horde.Build.

diff --git a/Engine/Source/Programs/Horde/Horde.Build/Controllers/LeasesController.cs b/Engine/Source/Programs/Horde/Horde.Build/Controllers/LeasesController.cs
--- a/Engine/Source/Programs/Horde/Horde.Build/Controllers/LeasesController.cs
+++ b/Engine/Source/Programs/Horde/Horde.Build/Controllers/LeasesController.cs
@@ -1,36 +1,23 @@
-﻿// Copyright Epic Games, Inc. All Rights Reserved.
+// Copyright Epic Games, Inc. All Rights Reserved.
 
+using System;
+using System.Collections.Generic;
+using System.Linq;
+using System.Threading.Tasks;
 using Horde.Build.Acls;
 using Horde.Build.Api;
 using Horde.Build.Models;
 using Horde.Build.Services;
 using Horde.Build.Utilities;
-using Microsoft.AspNetCore.Authentication;
 using Microsoft.AspNetCore.Authorization;
 using Microsoft.AspNetCore.Mvc;
-using Microsoft.AspNetCore.Mvc.ModelBinding;
-using Microsoft.Extensions.Options;
-using Microsoft.IdentityModel.Tokens;
-using MongoDB.Bson;
-using System;
-using System.Collections.Generic;
-using System.ComponentModel.DataAnnotations;
-using System.IdentityModel.Tokens.Jwt;
-using System.Linq;
-using System.Reflection;
-using System.Security.Claims;
-using System.Text.Json;
-using System.Threading.Tasks;
-using Microsoft.Extensions.Logging;
 using OpenTracing;
 using OpenTracing.Util;
 
 namespace Horde.Build.Controllers
 {
 	using LeaseId = ObjectId<ILease>;
-	using PoolId = StringId<IPool>;
 	using SessionId = ObjectId<ISession>;
-	using AgentSoftwareChannelName = StringId<AgentSoftwareChannels>;
 
 	/// <summary>
 	/// Controller for the /api/v1/leases endpoint
@@ -43,199 +30,191 @@ namespace Horde.Build.Controllers
 		/// <summary>
 		/// Singleton instance of the ACL service
 		/// </summary>
-		AclService AclService;
+		readonly AclService _aclService;
 
 		/// <summary>
 		/// Singleton instance of the agent service
 		/// </summary>
-		AgentService AgentService;
-
-		/// <summary>
-		/// Logger instance
-		/// </summary>
-		private ILogger<LeasesController> Logger;
+		readonly AgentService _agentService;
 
 		/// <summary>
 		/// Constructor
 		/// </summary>
-		/// <param name="AclService">The ACL service singleton</param>
-		/// <param name="AgentService">The agent service</param>
-		/// <param name="Logger">The controller logger</param>
-		public LeasesController(AclService AclService, AgentService AgentService, ILogger<LeasesController> Logger)
+		/// <param name="aclService">The ACL service singleton</param>
+		/// <param name="agentService">The agent service</param>
+		public LeasesController(AclService aclService, AgentService agentService)
 		{
-			this.AclService = AclService;
-			this.AgentService = AgentService;
-			this.Logger = Logger;
+			_aclService = aclService;
+			_agentService = agentService;
 		}
 
 		/// <summary>
 		/// Find all the leases for a particular agent
 		/// </summary>
-		/// <param name="AgentId">Unique id of the agent to find</param>
-		/// <param name="SessionId">The session to query</param>
-		/// <param name="StartTime">Start of the time window to consider</param>
-		/// <param name="FinishTime">End of the time window to consider</param>
-		/// <param name="MinFinishTime">Start of the time window to consider when querying by finish time (if set, other criteria are ignored)</param>
-		/// <param name="MaxFinishTime">End of the time window to consider when querying by finish time (if set, other criteria are ignored)</param>
-		/// <param name="Index">Index of the first result to return</param>
-		/// <param name="Count">Number of results to return</param>
-		/// <param name="Filter">Filter to apply to the results</param>
+		/// <param name="agentId">Unique id of the agent to find</param>
+		/// <param name="sessionId">The session to query</param>
+		/// <param name="startTime">Start of the time window to consider</param>
+		/// <param name="finishTime">End of the time window to consider</param>
+		/// <param name="minFinishTime">Start of the time window to consider when querying by finish time (if set, other criteria are ignored)</param>
+		/// <param name="maxFinishTime">End of the time window to consider when querying by finish time (if set, other criteria are ignored)</param>
+		/// <param name="index">Index of the first result to return</param>
+		/// <param name="count">Number of results to return</param>
+		/// <param name="filter">Filter to apply to the results</param>
 		/// <returns>Sessions </returns>
 		[HttpGet]
 		[Route("/api/v1/leases")]
 		[ProducesResponseType(200, Type = typeof(List<GetAgentLeaseResponse>))]
 		public async Task<ActionResult<List<object>>> FindLeasesAsync(
-			[FromQuery] AgentId? AgentId,
-			[FromQuery] SessionId? SessionId,
-			[FromQuery] DateTimeOffset? StartTime,
-			[FromQuery] DateTimeOffset? FinishTime,
-			[FromQuery] DateTimeOffset? MinFinishTime,
-			[FromQuery] DateTimeOffset? MaxFinishTime,
-			[FromQuery] int Index = 0,
-			[FromQuery] int Count = 1000,
-			[FromQuery] PropertyFilter? Filter = null)
+			[FromQuery] AgentId? agentId,
+			[FromQuery] SessionId? sessionId,
+			[FromQuery] DateTimeOffset? startTime,
+			[FromQuery] DateTimeOffset? finishTime,
+			[FromQuery] DateTimeOffset? minFinishTime,
+			[FromQuery] DateTimeOffset? maxFinishTime,
+			[FromQuery] int index = 0,
+			[FromQuery] int count = 1000,
+			[FromQuery] PropertyFilter? filter = null)
 		{
-			GlobalPermissionsCache PermissionsCache = new GlobalPermissionsCache();
-			if (AgentId == null)
+			GlobalPermissionsCache permissionsCache = new GlobalPermissionsCache();
+			if (agentId == null)
 			{
-				if (!await AclService.AuthorizeAsync(AclAction.ViewLeases, User, PermissionsCache))
+				if (!await _aclService.AuthorizeAsync(AclAction.ViewLeases, User, permissionsCache))
 				{
 					return Forbid(AclAction.ViewLeases);
 				}
 			}
 			else
 			{
-				IAgent? Agent = await AgentService.GetAgentAsync(AgentId.Value);
-				if (Agent == null)
+				IAgent? agent = await _agentService.GetAgentAsync(agentId.Value);
+				if (agent == null)
 				{
-					return NotFound(AgentId.Value);
+					return NotFound(agentId.Value);
 				}
-				if (!await AgentService.AuthorizeAsync(Agent, AclAction.ViewLeases, User, PermissionsCache))
+				if (!await _agentService.AuthorizeAsync(agent, AclAction.ViewLeases, User, permissionsCache))
 				{
-					return Forbid(AclAction.ViewLeases, AgentId.Value);
+					return Forbid(AclAction.ViewLeases, agentId.Value);
 				}
 			}
 
-			bool IncludeCosts = await AclService.AuthorizeAsync(AclAction.ViewCosts, User, PermissionsCache);
+			bool includeCosts = await _aclService.AuthorizeAsync(AclAction.ViewCosts, User, permissionsCache);
 
-			List<ILease> Leases;
-			if (MinFinishTime == null && MaxFinishTime == null)
+			List<ILease> leases;
+			if (minFinishTime == null && maxFinishTime == null)
 			{
-				Leases = await AgentService.FindLeasesAsync(AgentId, SessionId, StartTime?.UtcDateTime, FinishTime?.UtcDateTime, Index, Count);
+				leases = await _agentService.FindLeasesAsync(agentId, sessionId, startTime?.UtcDateTime, finishTime?.UtcDateTime, index, count);
 			}
 			else
 			{
 				// Optimized path for queries made by finish time
-				Leases = await AgentService.FindLeasesByFinishTimeAsync(MinFinishTime?.UtcDateTime, MaxFinishTime?.UtcDateTime, Index, Count);
+				leases = await _agentService.FindLeasesByFinishTimeAsync(minFinishTime?.UtcDateTime, maxFinishTime?.UtcDateTime, index, count);
 			}
 
-			List<object> Responses = new List<object>();
+			List<object> responses = new List<object>();
 
 			using (IScope _ = GlobalTracer.Instance.BuildSpan($"GenerateResponses").StartActive())
 			{
-				Dictionary<AgentId, double?> CachedAgentRates = new Dictionary<AgentId, double?>();
-				foreach (ILease Lease in Leases)
+				Dictionary<AgentId, double?> cachedAgentRates = new Dictionary<AgentId, double?>();
+				foreach (ILease lease in leases)
 				{
-					double? AgentRate = null;
-					if (IncludeCosts && !CachedAgentRates.TryGetValue(Lease.AgentId, out AgentRate))
+					double? agentRate = null;
+					if (includeCosts && !cachedAgentRates.TryGetValue(lease.AgentId, out agentRate))
 					{
-						AgentRate = await AgentService.GetRateAsync(Lease.AgentId);
-						CachedAgentRates.Add(Lease.AgentId, AgentRate);
+						agentRate = await _agentService.GetRateAsync(lease.AgentId);
+						cachedAgentRates.Add(lease.AgentId, agentRate);
 					}
 
-					Dictionary<string, string>? Details = AgentService.GetPayloadDetails(Lease.Payload);
-					Responses.Add(PropertyFilter.Apply(new GetAgentLeaseResponse(Lease, Details, AgentRate), Filter));
+					Dictionary<string, string>? details = _agentService.GetPayloadDetails(lease.Payload);
+					responses.Add(PropertyFilter.Apply(new GetAgentLeaseResponse(lease, details, agentRate), filter));
 				}
 			}
 
-			return Responses;
+			return responses;
 		}
 
 		/// <summary>
 		/// Get info about a particular lease
 		/// </summary>
-		/// <param name="LeaseId">Unique id of the particular lease</param>
+		/// <param name="leaseId">Unique id of the particular lease</param>
 		/// <returns>Lease matching the given id</returns>
 		[HttpGet]
-		[Route("/api/v1/leases/{LeaseId}")]
-		public async Task<ActionResult<GetAgentLeaseResponse>> GetLeaseAsync(LeaseId LeaseId)
+		[Route("/api/v1/leases/{leaseId}")]
+		public async Task<ActionResult<GetAgentLeaseResponse>> GetLeaseAsync(LeaseId leaseId)
 		{
-			ILease? Lease = await AgentService.GetLeaseAsync(LeaseId);
-			if (Lease == null)
+			ILease? lease = await _agentService.GetLeaseAsync(leaseId);
+			if (lease == null)
 			{
-				return NotFound(LeaseId);
+				return NotFound(leaseId);
 			}
 
-			IAgent? Agent = await AgentService.GetAgentAsync(Lease.AgentId);
-			if (Agent == null)
+			IAgent? agent = await _agentService.GetAgentAsync(lease.AgentId);
+			if (agent == null)
 			{
-				return NotFound(Lease.AgentId);
+				return NotFound(lease.AgentId);
 			}
 
-			GlobalPermissionsCache PermissionsCache = new GlobalPermissionsCache();
-			if (!await AgentService.AuthorizeAsync(Agent, AclAction.ViewLeases, User, PermissionsCache))
+			GlobalPermissionsCache permissionsCache = new GlobalPermissionsCache();
+			if (!await _agentService.AuthorizeAsync(agent, AclAction.ViewLeases, User, permissionsCache))
 			{
 				return Forbid(AclAction.ViewLeases);
 			}
 
-			double? AgentRate = null;
-			if (await AclService.AuthorizeAsync(AclAction.ViewCosts, User, PermissionsCache))
+			double? agentRate = null;
+			if (await _aclService.AuthorizeAsync(AclAction.ViewCosts, User, permissionsCache))
 			{
-				AgentRate = await AgentService.GetRateAsync(Agent.Id);
+				agentRate = await _agentService.GetRateAsync(agent.Id);
 			}
 
-			Dictionary<string, string>? Details = AgentService.GetPayloadDetails(Lease.Payload);
-			return new GetAgentLeaseResponse(Lease, Details, AgentRate);
+			Dictionary<string, string>? details = _agentService.GetPayloadDetails(lease.Payload);
+			return new GetAgentLeaseResponse(lease, details, agentRate);
 		}
 
 		/// <summary>
 		/// Update a particular lease
 		/// </summary>
-		/// <param name="LeaseId">Unique id of the particular lease</param>
-		/// <param name="Request"></param>
+		/// <param name="leaseId">Unique id of the particular lease</param>
+		/// <param name="request"></param>
 		/// <returns>Lease matching the given id</returns>
 		[HttpPut]
-		[Route("/api/v1/leases/{LeaseId}")]
-		public async Task<ActionResult> UpdateLeaseAsync(LeaseId LeaseId, [FromBody] UpdateLeaseRequest Request)
+		[Route("/api/v1/leases/{leaseId}")]
+		public async Task<ActionResult> UpdateLeaseAsync(LeaseId leaseId, [FromBody] UpdateLeaseRequest request)
 		{
 			// only update supported right now is abort
-			if (!Request.Aborted.HasValue || !Request.Aborted.Value)
+			if (!request.Aborted.HasValue || !request.Aborted.Value)
 			{
 				return Ok();
 			}
 
-			ILease? Lease = await AgentService.GetLeaseAsync(LeaseId);
-			if (Lease == null)
+			ILease? lease = await _agentService.GetLeaseAsync(leaseId);
+			if (lease == null)
 			{
-				return NotFound(LeaseId);
+				return NotFound(leaseId);
 			}
 
-			IAgent? Agent = await AgentService.GetAgentAsync(Lease.AgentId);
-			if (Agent == null)
+			IAgent? agent = await _agentService.GetAgentAsync(lease.AgentId);
+			if (agent == null)
 			{
-				return NotFound(Lease.AgentId);
+				return NotFound(lease.AgentId);
 			}
 
-			if (!await AgentService.AuthorizeAsync(Agent, AclAction.AdminWrite, User, null))
+			if (!await _agentService.AuthorizeAsync(agent, AclAction.AdminWrite, User, null))
 			{
-				return Forbid(AclAction.AdminWrite, Lease.AgentId);
+				return Forbid(AclAction.AdminWrite, lease.AgentId);
 			}
 
-			AgentLease? AgentLease = Agent.Leases.FirstOrDefault(x => x.Id == LeaseId);
+			AgentLease? agentLease = agent.Leases.FirstOrDefault(x => x.Id == leaseId);
 
-			if (AgentLease == null)
+			if (agentLease == null)
 			{
-				return NotFound(Agent.Id, LeaseId);
+				return NotFound(agent.Id, leaseId);
 			}
 
-			if (!AgentLease.IsConformLease())
+			if (!agentLease.IsConformLease())
 			{
-				return BadRequest("Lease abort only supported on conform leases for now, {LeaseId}", LeaseId);
+				return BadRequest("Lease abort only supported on conform leases for now, {LeaseId}", leaseId);
 			}
 
-			await AgentService.CancelLeaseAsync(Agent, LeaseId);
+			await _agentService.CancelLeaseAsync(agent, leaseId);
 			return Ok();
 		}
-
 	}
 }
diff --git a/Engine/Source/Programs/Horde/Horde.Build/Controllers/LifetimeController.cs b/Engine/Source/Programs/Horde/Horde.Build/Controllers/LifetimeController.cs
--- a/Engine/Source/Programs/Horde/Horde.Build/Controllers/LifetimeController.cs
+++ b/Engine/Source/Programs/Horde/Horde.Build/Controllers/LifetimeController.cs
@@ -1,8 +1,8 @@
 // Copyright Epic Games, Inc. All Rights Reserved.
 
+using System.Threading.Tasks;
 using Horde.Build.Services;
 using Microsoft.AspNetCore.Mvc;
-using System.Threading.Tasks;
 
 namespace Horde.Build.Controllers
 {
@@ -16,15 +16,15 @@ namespace Horde.Build.Controllers
 		/// <summary>
 		/// Singleton instance of the lifetime service
 		/// </summary>
-		LifetimeService LifetimeService;
+		readonly LifetimeService _lifetimeService;
 
 		/// <summary>
 		/// Constructor
 		/// </summary>
-		/// <param name="LifetimeService">The lifetime service singleton</param>
-		public LifetimeController(LifetimeService LifetimeService)
+		/// <param name="lifetimeService">The lifetime service singleton</param>
+		public LifetimeController(LifetimeService lifetimeService)
 		{
-			this.LifetimeService = LifetimeService;
+			_lifetimeService = lifetimeService;
 		}
 
 		/// <summary>
@@ -40,16 +40,16 @@ namespace Horde.Build.Controllers
 		[Route("/health/ready")]
 		public Task<ActionResult> ServerReadiness()
 		{
-			int StatusCode = 200;
-			string Content = "ok";
+			int statusCode = 200;
+			string content = "ok";
 			
-			if (LifetimeService.IsPreStopping)
+			if (_lifetimeService.IsPreStopping)
 			{
-				StatusCode = 503; // Service Unavailable
-				Content = "stopping";
+				statusCode = 503; // Service Unavailable
+				content = "stopping";
 			}
 
-			return Task.FromResult<ActionResult>(new ContentResult { ContentType = "text/plain", StatusCode = StatusCode, Content = Content });
+			return Task.FromResult<ActionResult>(new ContentResult { ContentType = "text/plain", StatusCode = statusCode, Content = content });
 		}
 		
 		/// <summary>
@@ -59,7 +59,7 @@ namespace Horde.Build.Controllers
 		/// <returns>Ok if app is not stopping and all databases can be reached</returns>
 		[HttpGet]
 		[Route("/health/live")]
-		public ActionResult K8sLivenessProbe()
+		public ActionResult K8SLivenessProbe()
 		{
 			return new ContentResult { ContentType = "text/plain", StatusCode = 200, Content = "ok" };
 		}
diff --git a/Engine/Source/Programs/Horde/Horde.Build/Controllers/LogsController.cs b/Engine/Source/Programs/Horde/Horde.Build/Controllers/LogsController.cs
--- a/Engine/Source/Programs/Horde/Horde.Build/Controllers/LogsController.cs
+++ b/Engine/Source/Programs/Horde/Horde.Build/Controllers/LogsController.cs
@@ -1,15 +1,10 @@
-﻿// Copyright Epic Games, Inc. All Rights Reserved.
+// Copyright Epic Games, Inc. All Rights Reserved.
 
 using System;
-using System.Buffers;
 using System.Collections.Generic;
-using System.ComponentModel.DataAnnotations;
 using System.IO;
-using System.Linq;
-using System.Net.Http;
 using System.Security.Claims;
 using System.Text;
-using System.Text.Json;
 using System.Threading.Tasks;
 using Horde.Build.Acls;
 using Horde.Build.Api;
@@ -19,10 +14,7 @@ using Horde.Build.Models;
 using Horde.Build.Services;
 using Horde.Build.Utilities;
 using Microsoft.AspNetCore.Authorization;
-using Microsoft.AspNetCore.Http;
 using Microsoft.AspNetCore.Mvc;
-using Microsoft.Extensions.Logging;
-using Microsoft.Extensions.Logging.Abstractions;
 using MongoDB.Bson;
 
 namespace Horde.Build.Controllers
@@ -54,127 +46,127 @@ namespace Horde.Build.Controllers
 	[Route("[controller]")]
 	public class LogsController : ControllerBase
 	{
-		private readonly ILogFileService LogFileService;
-		private readonly IIssueCollection IssueCollection;
-		private readonly AclService AclService;
-		private readonly JobService JobService;
+		private readonly ILogFileService _logFileService;
+		private readonly IIssueCollection _issueCollection;
+		private readonly AclService _aclService;
+		private readonly JobService _jobService;
 
 		/// <summary>
 		/// Constructor
 		/// </summary>
-		public LogsController(ILogFileService LogFileService, IIssueCollection IssueCollection, AclService AclService, JobService JobService)
+		public LogsController(ILogFileService logFileService, IIssueCollection issueCollection, AclService aclService, JobService jobService)
 		{
-			this.LogFileService = LogFileService;
-			this.IssueCollection = IssueCollection;
-			this.AclService = AclService;
-			this.JobService = JobService;
+			_logFileService = logFileService;
+			_issueCollection = issueCollection;
+			_aclService = aclService;
+			_jobService = jobService;
 		}
 
 		/// <summary>
 		/// Retrieve metadata about a specific log file
 		/// </summary>
-		/// <param name="LogFileId">Id of the log file to get information about</param>
-		/// <param name="Filter">Filter for the properties to return</param>
+		/// <param name="logFileId">Id of the log file to get information about</param>
+		/// <param name="filter">Filter for the properties to return</param>
 		/// <returns>Information about the requested project</returns>
 		[HttpGet]
-		[Route("/api/v1/logs/{LogFileId}")]
+		[Route("/api/v1/logs/{logFileId}")]
 		[ProducesResponseType(typeof(GetLogFileResponse), 200)]
-		public async Task<ActionResult<object>> GetLog(LogId LogFileId, [FromQuery] PropertyFilter? Filter = null)
+		public async Task<ActionResult<object>> GetLog(LogId logFileId, [FromQuery] PropertyFilter? filter = null)
 		{
-			ILogFile? LogFile = await LogFileService.GetLogFileAsync(LogFileId);
-			if (LogFile == null)
+			ILogFile? logFile = await _logFileService.GetLogFileAsync(logFileId);
+			if (logFile == null)
 			{
 				return NotFound();
 			}
-			if (!await AuthorizeAsync(LogFile, AclAction.ViewLog, User, null))
+			if (!await AuthorizeAsync(logFile, AclAction.ViewLog, User, null))
 			{
 				return Forbid();
 			}
 
-			LogMetadata Metadata = await LogFileService.GetMetadataAsync(LogFile);
-			return new GetLogFileResponse(LogFile, Metadata).ApplyFilter(Filter);       
+			LogMetadata metadata = await _logFileService.GetMetadataAsync(logFile);
+			return new GetLogFileResponse(logFile, metadata).ApplyFilter(filter);       
 		}
 
 		/// <summary>
 		/// Retrieve raw data for a log file
 		/// </summary>
-		/// <param name="LogFileId">Id of the log file to get information about</param>
-		/// <param name="Format">Format for the returned data</param>
-		/// <param name="Offset">The log offset in bytes</param>
-		/// <param name="Length">Number of bytes to return</param>
-		/// <param name="FileName">Name of the default filename to download</param>
-		/// <param name="Download">Whether to download the file rather than display in the browser</param>
+		/// <param name="logFileId">Id of the log file to get information about</param>
+		/// <param name="format">Format for the returned data</param>
+		/// <param name="offset">The log offset in bytes</param>
+		/// <param name="length">Number of bytes to return</param>
+		/// <param name="fileName">Name of the default filename to download</param>
+		/// <param name="download">Whether to download the file rather than display in the browser</param>
 		/// <returns>Raw log data for the requested range</returns>
 		[HttpGet]
-		[Route("/api/v1/logs/{LogFileId}/data")]
-		public async Task<ActionResult> GetLogData(LogId LogFileId, [FromQuery] LogOutputFormat Format = LogOutputFormat.Raw, [FromQuery] long Offset = 0, [FromQuery] long Length = long.MaxValue, [FromQuery] string? FileName = null, [FromQuery] bool Download = false)
+		[Route("/api/v1/logs/{logFileId}/data")]
+		public async Task<ActionResult> GetLogData(LogId logFileId, [FromQuery] LogOutputFormat format = LogOutputFormat.Raw, [FromQuery] long offset = 0, [FromQuery] long length = Int64.MaxValue, [FromQuery] string? fileName = null, [FromQuery] bool download = false)
 		{
-			ILogFile? LogFile = await LogFileService.GetLogFileAsync(LogFileId);
-			if (LogFile == null)
+			ILogFile? logFile = await _logFileService.GetLogFileAsync(logFileId);
+			if (logFile == null)
 			{
 				return NotFound();
 			}
-			if (!await AuthorizeAsync(LogFile, AclAction.ViewLog, User, null))
+			if (!await AuthorizeAsync(logFile, AclAction.ViewLog, User, null))
 			{
 				return Forbid();
 			}
 
-			Func<Stream, ActionContext, Task> CopyTask;
-			if (Format == LogOutputFormat.Text && LogFile.Type == LogType.Json)
+			Func<Stream, ActionContext, Task> copyTask;
+			if (format == LogOutputFormat.Text && logFile.Type == LogType.Json)
 			{
-				CopyTask = (OutputStream, Context) => LogFileService.CopyPlainTextStreamAsync(LogFile, Offset, Length, OutputStream);
+				copyTask = (outputStream, context) => _logFileService.CopyPlainTextStreamAsync(logFile, offset, length, outputStream);
 			}
 			else
 			{
-				CopyTask = (OutputStream, Context) => LogFileService.CopyRawStreamAsync(LogFile, Offset, Length, OutputStream);
+				copyTask = (outputStream, context) => _logFileService.CopyRawStreamAsync(logFile, offset, length, outputStream);
 			}
 
-			return new CustomFileCallbackResult(FileName ?? $"log-{LogFileId}.txt", "text/plain", !Download, CopyTask);
+			return new CustomFileCallbackResult(fileName ?? $"log-{logFileId}.txt", "text/plain", !download, copyTask);
 		}
 
 		/// <summary>
 		/// Retrieve line data for a logfile
 		/// </summary>
-		/// <param name="LogFileId">Id of the log file to get information about</param>
-		/// <param name="Index">Index of the first line to retrieve</param>
-		/// <param name="Count">Number of lines to retrieve</param>
+		/// <param name="logFileId">Id of the log file to get information about</param>
+		/// <param name="index">Index of the first line to retrieve</param>
+		/// <param name="count">Number of lines to retrieve</param>
 		/// <returns>Information about the requested project</returns>
 		[HttpGet]
-		[Route("/api/v1/logs/{LogFileId}/lines")]
-		public async Task<ActionResult> GetLogLines(LogId LogFileId, [FromQuery] int Index = 0, [FromQuery] int Count = int.MaxValue)
+		[Route("/api/v1/logs/{logFileId}/lines")]
+		public async Task<ActionResult> GetLogLines(LogId logFileId, [FromQuery] int index = 0, [FromQuery] int count = Int32.MaxValue)
 		{
-			ILogFile? LogFile = await LogFileService.GetLogFileAsync(LogFileId);
-			if (LogFile == null)
+			ILogFile? logFile = await _logFileService.GetLogFileAsync(logFileId);
+			if (logFile == null)
 			{
 				return NotFound();
 			}
-			if (!await AuthorizeAsync(LogFile, AclAction.ViewLog, User, null))
+			if (!await AuthorizeAsync(logFile, AclAction.ViewLog, User, null))
 			{
 				return Forbid();
 			}
 
-			LogMetadata Metadata = await LogFileService.GetMetadataAsync(LogFile);
+			LogMetadata metadata = await _logFileService.GetMetadataAsync(logFile);
 
-			(int MinIndex, long MinOffset) = await LogFileService.GetLineOffsetAsync(LogFile, Index);
-			(int MaxIndex, long MaxOffset) = await LogFileService.GetLineOffsetAsync(LogFile, Index + Math.Min(Count, int.MaxValue - Index));
-			Index = MinIndex;
-			Count = MaxIndex - MinIndex;
+			(int minIndex, long minOffset) = await _logFileService.GetLineOffsetAsync(logFile, index);
+			(int maxIndex, long maxOffset) = await _logFileService.GetLineOffsetAsync(logFile, index + Math.Min(count, Int32.MaxValue - index));
+			index = minIndex;
+			count = maxIndex - minIndex;
 
-			byte[] Result;
-			using (System.IO.Stream Stream = await LogFileService.OpenRawStreamAsync(LogFile, MinOffset, MaxOffset - MinOffset))
+			byte[] result;
+			using (System.IO.Stream stream = await _logFileService.OpenRawStreamAsync(logFile, minOffset, maxOffset - minOffset))
 			{
-				Result = new byte[Stream.Length];
-				await Stream.ReadFixedSizeDataAsync(Result, 0, Result.Length);
+				result = new byte[stream.Length];
+				await stream.ReadFixedSizeDataAsync(result, 0, result.Length);
 			}
 
-			using (MemoryStream Stream = new MemoryStream(Result.Length + (Count * 20)))
+			using (MemoryStream stream = new MemoryStream(result.Length + (count * 20)))
 			{
-				Stream.WriteByte((byte)'{');
+				stream.WriteByte((byte)'{');
 
-				Stream.Write(Encoding.UTF8.GetBytes($"\"index\":{Index},"));
-				Stream.Write(Encoding.UTF8.GetBytes($"\"count\":{Count},"));
-				Stream.Write(Encoding.UTF8.GetBytes($"\"maxLineIndex\":{Metadata.MaxLineIndex},"));
-				Stream.Write(Encoding.UTF8.GetBytes($"\"format\":{ (LogFile.Type == LogType.Json ? "\"JSON\"" : "\"TEXT\"")},"));
+				stream.Write(Encoding.UTF8.GetBytes($"\"index\":{index},"));
+				stream.Write(Encoding.UTF8.GetBytes($"\"count\":{count},"));
+				stream.Write(Encoding.UTF8.GetBytes($"\"maxLineIndex\":{metadata.MaxLineIndex},"));
+				stream.Write(Encoding.UTF8.GetBytes($"\"format\":{ (logFile.Type == LogType.Json ? "\"JSON\"" : "\"TEXT\"")},"));
 
 //				Stream.Write(Encoding.UTF8.GetBytes($"\"minIndex\":{MinIndex},"));
 //				Stream.Write(Encoding.UTF8.GetBytes($"\"minOffset\":{MinOffset},"));
@@ -182,81 +174,81 @@ namespace Horde.Build.Controllers
 //				Stream.Write(Encoding.UTF8.GetBytes($"\"maxOffset\":{MaxOffset},"));
 //				Stream.Write(Encoding.UTF8.GetBytes($"\"length\":{Result.Length},"));
 
-				Stream.Write(Encoding.UTF8.GetBytes($"\"lines\":["));
-				Stream.WriteByte((byte)'\n');
+				stream.Write(Encoding.UTF8.GetBytes($"\"lines\":["));
+				stream.WriteByte((byte)'\n');
 
-				int Offset = 0;
-				for (int Line = Index; Line < Index + Count; Line++)
+				int offset = 0;
+				for (int line = index; line < index + count; line++)
 				{
-					Stream.WriteByte((byte)' ');
-					Stream.WriteByte((byte)' ');
+					stream.WriteByte((byte)' ');
+					stream.WriteByte((byte)' ');
 
-					if (LogFile.Type == LogType.Json)
+					if (logFile.Type == LogType.Json)
 					{
 						// Find the end of the line and output it as an opaque blob
-						int StartOffset = Offset;
-						for (; ; Offset++)
+						int startOffset = offset;
+						for (; ; offset++)
 						{
-							if (Offset == Result.Length)
+							if (offset == result.Length)
 							{
-								Stream.WriteByte((byte)'{');
-								Stream.WriteByte((byte)'}');
+								stream.WriteByte((byte)'{');
+								stream.WriteByte((byte)'}');
 								break;
 							}
-							else if (Result[Offset] == (byte)'\n')
+							else if (result[offset] == (byte)'\n')
 							{
-								await Stream.WriteAsync(Result.AsMemory(StartOffset, Offset - StartOffset));
-								Offset++;
+								await stream.WriteAsync(result.AsMemory(startOffset, offset - startOffset));
+								offset++;
 								break;
 							}
 						}
 					}
 					else
 					{
-						Stream.WriteByte((byte)'\"');
-						for (; Offset < Result.Length; Offset++)
+						stream.WriteByte((byte)'\"');
+						for (; offset < result.Length; offset++)
 						{
-							if (Result[Offset] == '\\' || Result[Offset] == '\"')
+							if (result[offset] == '\\' || result[offset] == '\"')
 							{
-								Stream.WriteByte((byte)'\\');
-								Stream.WriteByte(Result[Offset]);
+								stream.WriteByte((byte)'\\');
+								stream.WriteByte(result[offset]);
 							}
-							else if (Result[Offset] == (byte)'\n')
+							else if (result[offset] == (byte)'\n')
 							{
-								Offset++;
+								offset++;
 								break;
 							}
-							else if (Result[Offset] >= 32 && Result[Offset] <= 126)
+							else if (result[offset] >= 32 && result[offset] <= 126)
 							{
-								Stream.WriteByte(Result[Offset]);
+								stream.WriteByte(result[offset]);
 							}
 							else
 							{
-								Stream.Write(Encoding.UTF8.GetBytes($"\\x{Result[Offset]:x2}"));
+								stream.Write(Encoding.UTF8.GetBytes($"\\x{result[offset]:x2}"));
 							}
 						}
-						Stream.WriteByte((byte)'\"');
+						stream.WriteByte((byte)'\"');
 					}
 
-					if (Line + 1 < Index + Count)
+					if (line + 1 < index + count)
 					{
-						Stream.WriteByte((byte)',');
+						stream.WriteByte((byte)',');
 					}
 
-					Stream.WriteByte((byte)'\n');
+					stream.WriteByte((byte)'\n');
 				}
 
-				if (LogFile.Type == LogType.Json)
+				if (logFile.Type == LogType.Json)
 				{
-					Stream.Write(Encoding.UTF8.GetBytes($"]"));
+					stream.Write(Encoding.UTF8.GetBytes($"]"));
 				}
 
-				Stream.WriteByte((byte)'}');
+				stream.WriteByte((byte)'}');
 
 				Response.ContentType = "application/json";
-				Response.Headers.ContentLength = Stream.Length;
-				Stream.Position = 0;
-				await Stream.CopyToAsync(Response.Body);
+				Response.Headers.ContentLength = stream.Length;
+				stream.Position = 0;
+				await stream.CopyToAsync(Response.Body);
 			}
 			return new EmptyResult();
 		}
@@ -264,100 +256,100 @@ namespace Horde.Build.Controllers
 		/// <summary>
 		/// Search log data
 		/// </summary>
-		/// <param name="LogFileId">Id of the log file to get information about</param>
-		/// <param name="Text">Text to search for</param>
-		/// <param name="FirstLine">First line to search from</param>
-		/// <param name="Count">Number of results to return</param>
+		/// <param name="logFileId">Id of the log file to get information about</param>
+		/// <param name="text">Text to search for</param>
+		/// <param name="firstLine">First line to search from</param>
+		/// <param name="count">Number of results to return</param>
 		/// <returns>Raw log data for the requested range</returns>
 		[HttpGet]
-		[Route("/api/v1/logs/{LogFileId}/search")]
-		public async Task<ActionResult<SearchLogFileResponse>> SearchLogFileAsync(LogId LogFileId, [FromQuery] string Text, [FromQuery] int FirstLine = 0, [FromQuery] int Count = 5)
+		[Route("/api/v1/logs/{logFileId}/search")]
+		public async Task<ActionResult<SearchLogFileResponse>> SearchLogFileAsync(LogId logFileId, [FromQuery] string text, [FromQuery] int firstLine = 0, [FromQuery] int count = 5)
 		{
-			ILogFile? LogFile = await LogFileService.GetLogFileAsync(LogFileId);
-			if (LogFile == null)
+			ILogFile? logFile = await _logFileService.GetLogFileAsync(logFileId);
+			if (logFile == null)
 			{
 				return NotFound();
 			}
-			if (!await AuthorizeAsync(LogFile, AclAction.ViewLog, User, null))
+			if (!await AuthorizeAsync(logFile, AclAction.ViewLog, User, null))
 			{
 				return Forbid();
 			}
 
-			SearchLogFileResponse Response = new SearchLogFileResponse();
-			Response.Stats = new LogSearchStats();
-			Response.Lines = await LogFileService.SearchLogDataAsync(LogFile, Text, FirstLine, Count, Response.Stats);
-			return Response;
+			SearchLogFileResponse response = new SearchLogFileResponse();
+			response.Stats = new LogSearchStats();
+			response.Lines = await _logFileService.SearchLogDataAsync(logFile, text, firstLine, count, response.Stats);
+			return response;
 		}
 
 		/// <summary>
 		/// Retrieve events for a logfile
 		/// </summary>
-		/// <param name="LogFileId">Id of the log file to get information about</param>
-		/// <param name="Index">Index of the first line to retrieve</param>
-		/// <param name="Count">Number of lines to retrieve</param>
+		/// <param name="logFileId">Id of the log file to get information about</param>
+		/// <param name="index">Index of the first line to retrieve</param>
+		/// <param name="count">Number of lines to retrieve</param>
 		/// <returns>Information about the requested project</returns>
 		[HttpGet]
-		[Route("/api/v1/logs/{LogFileId}/events")]
+		[Route("/api/v1/logs/{logFileId}/events")]
 		[ProducesResponseType(typeof(List<GetLogEventResponse>), 200)]
-		public async Task<ActionResult<List<GetLogEventResponse>>> GetEventsAsync(LogId LogFileId, [FromQuery] int? Index = null, [FromQuery] int? Count = null)
+		public async Task<ActionResult<List<GetLogEventResponse>>> GetEventsAsync(LogId logFileId, [FromQuery] int? index = null, [FromQuery] int? count = null)
 		{
-			ILogFile? LogFile = await LogFileService.GetLogFileAsync(LogFileId);
-			if (LogFile == null)
+			ILogFile? logFile = await _logFileService.GetLogFileAsync(logFileId);
+			if (logFile == null)
 			{
 				return NotFound();
 			}
-			if (!await AuthorizeAsync(LogFile, AclAction.ViewLog, User, null))
+			if (!await AuthorizeAsync(logFile, AclAction.ViewLog, User, null))
 			{
 				return Forbid();
 			}
 
-			List<ILogEvent> LogEvents = await LogFileService.FindLogEventsAsync(LogFile, Index, Count);
+			List<ILogEvent> logEvents = await _logFileService.FindLogEventsAsync(logFile, index, count);
 
-			Dictionary<ObjectId, int?> SpanIdToIssueId = new Dictionary<ObjectId, int?>();
+			Dictionary<ObjectId, int?> spanIdToIssueId = new Dictionary<ObjectId, int?>();
 
-			List<GetLogEventResponse> Responses = new List<GetLogEventResponse>();
-			foreach (ILogEvent LogEvent in LogEvents)
+			List<GetLogEventResponse> responses = new List<GetLogEventResponse>();
+			foreach (ILogEvent logEvent in logEvents)
 			{
-				ILogEventData LogEventData = await LogFileService.GetEventDataAsync(LogFile, LogEvent.LineIndex, LogEvent.LineCount);
+				ILogEventData logEventData = await _logFileService.GetEventDataAsync(logFile, logEvent.LineIndex, logEvent.LineCount);
 
-				int? IssueId = null;
-				if (LogEvent.SpanId != null && !SpanIdToIssueId.TryGetValue(LogEvent.SpanId.Value, out IssueId))
+				int? issueId = null;
+				if (logEvent.SpanId != null && !spanIdToIssueId.TryGetValue(logEvent.SpanId.Value, out issueId))
 				{
-					IIssueSpan? Span = await IssueCollection.GetSpanAsync(LogEvent.SpanId.Value);
-					IssueId = Span?.IssueId;
-					SpanIdToIssueId[LogEvent.SpanId.Value] = IssueId;
+					IIssueSpan? span = await _issueCollection.GetSpanAsync(logEvent.SpanId.Value);
+					issueId = span?.IssueId;
+					spanIdToIssueId[logEvent.SpanId.Value] = issueId;
 				}
 
-				Responses.Add(new GetLogEventResponse(LogEvent, LogEventData, IssueId));
+				responses.Add(new GetLogEventResponse(logEvent, logEventData, issueId));
 			}
-			return Responses;
+			return responses;
 		}
 
 		/// <summary>
 		/// Appends data to a log file
 		/// </summary>
-		/// <param name="LogFileId">The logfile id</param>
-		/// <param name="Offset">Offset within the log file</param>
-		/// <param name="LineIndex">The line index</param>
+		/// <param name="logFileId">The logfile id</param>
+		/// <param name="offset">Offset within the log file</param>
+		/// <param name="lineIndex">The line index</param>
 		/// <returns>Http result code</returns>
 		[HttpPost]
-		[Route("/api/v1/logs/{LogFileId}")]
-		public async Task<ActionResult> WriteData(LogId LogFileId, [FromQuery] long Offset, [FromQuery] int LineIndex)
+		[Route("/api/v1/logs/{logFileId}")]
+		public async Task<ActionResult> WriteData(LogId logFileId, [FromQuery] long offset, [FromQuery] int lineIndex)
 		{
-			ILogFile? LogFile = await LogFileService.GetLogFileAsync(LogFileId);
-			if (LogFile == null)
+			ILogFile? logFile = await _logFileService.GetLogFileAsync(logFileId);
+			if (logFile == null)
 			{
 				return NotFound();
 			}
-			if (!await AuthorizeAsync(LogFile, AclAction.WriteLogData, User, null))
+			if (!await AuthorizeAsync(logFile, AclAction.WriteLogData, User, null))
 			{
 				return Forbid();
 			}
 
-			using (MemoryStream BodyStream = new MemoryStream())
+			using (MemoryStream bodyStream = new MemoryStream())
 			{
-				await Request.Body.CopyToAsync(BodyStream);
-				await LogFileService.WriteLogDataAsync(LogFile, Offset, LineIndex, BodyStream.ToArray(), false);
+				await Request.Body.CopyToAsync(bodyStream);
+				await _logFileService.WriteLogDataAsync(logFile, offset, lineIndex, bodyStream.ToArray(), false);
 			}
 			return Ok();
 		}
@@ -365,18 +357,18 @@ namespace Horde.Build.Controllers
 		/// <summary>
 		/// Determines if the user is authorized to perform an action on a particular template
 		/// </summary>
-		/// <param name="LogFile">The template to check</param>
-		/// <param name="Action">The action being performed</param>
-		/// <param name="User">The principal to authorize</param>
-		/// <param name="PermissionsCache">Permissions cache</param>
+		/// <param name="logFile">The template to check</param>
+		/// <param name="action">The action being performed</param>
+		/// <param name="user">The principal to authorize</param>
+		/// <param name="permissionsCache">Permissions cache</param>
 		/// <returns>True if the action is authorized</returns>
-		async Task<bool> AuthorizeAsync(ILogFile LogFile, AclAction Action, ClaimsPrincipal User, JobPermissionsCache? PermissionsCache)
+		async Task<bool> AuthorizeAsync(ILogFile logFile, AclAction action, ClaimsPrincipal user, JobPermissionsCache? permissionsCache)
 		{
-			if (LogFile.JobId != JobId.Empty && await JobService.AuthorizeAsync(LogFile.JobId, Action, User, PermissionsCache))
+			if (logFile.JobId != JobId.Empty && await _jobService.AuthorizeAsync(logFile.JobId, action, user, permissionsCache))
 			{
 				return true;
 			}
-			if (LogFile.SessionId != null && await AclService.AuthorizeAsync(AclAction.ViewSession, User, PermissionsCache))
+			if (logFile.SessionId != null && await _aclService.AuthorizeAsync(AclAction.ViewSession, user, permissionsCache))
 			{
 				return true;
 			}
diff --git a/Engine/Source/Programs/Horde/Horde.Build/Controllers/NoticesController.cs b/Engine/Source/Programs/Horde/Horde.Build/Controllers/NoticesController.cs
--- a/Engine/Source/Programs/Horde/Horde.Build/Controllers/NoticesController.cs
+++ b/Engine/Source/Programs/Horde/Horde.Build/Controllers/NoticesController.cs
@@ -1,5 +1,9 @@
-﻿// Copyright Epic Games, Inc. All Rights Reserved.
+// Copyright Epic Games, Inc. All Rights Reserved.
 
+using System;
+using System.Collections.Generic;
+using System.Globalization;
+using System.Threading.Tasks;
 using Horde.Build.Acls;
 using Horde.Build.Api;
 using Horde.Build.Collections;
@@ -10,10 +14,6 @@ using Microsoft.AspNetCore.Authorization;
 using Microsoft.AspNetCore.Mvc;
 using Microsoft.Extensions.Options;
 using MongoDB.Bson;
-using System;
-using System.Collections.Generic;
-using System.Globalization;
-using System.Threading.Tasks;
 using TimeZoneConverter;
 
 namespace Horde.Build.Controllers
@@ -29,94 +29,93 @@ namespace Horde.Build.Controllers
 		/// <summary>
 		/// 
 		/// </summary>
-		NoticeService NoticeService;
+		readonly NoticeService _noticeService;
 
 		/// <summary>
 		/// The acl service singleton
 		/// </summary>
-		AclService AclService;
+		readonly AclService _aclService;
 
 		/// <summary>
 		/// Reference to the globals singleton
 		/// </summary>
-		ISingletonDocument<Globals> Globals;
+		readonly ISingletonDocument<Globals> _globals;
 
 		/// <summary>
 		/// Reference to the user collection
 		/// </summary>
-		private readonly IUserCollection UserCollection;
+		private readonly IUserCollection _userCollection;
 
 		/// <summary>
 		/// Server settings
 		/// </summary>
-		IOptionsMonitor<ServerSettings> Settings;
+		readonly IOptionsMonitor<ServerSettings> _settings;
 
 		/// <summary>
 		/// Constructor
 		/// </summary>
-		/// <param name="Settings">The server settings</param>
-		/// <param name="NoticeService">The notice service singleton</param>
-		/// <param name="AclService">The acl service singleton</param>
-		/// <param name="UserCollection">The user collection singleton</param>
-		/// <param name="Globals">The global singleton</param>
-		public NoticesController(IOptionsMonitor<ServerSettings> Settings, NoticeService NoticeService, AclService AclService, IUserCollection UserCollection, ISingletonDocument<Globals> Globals)
+		/// <param name="settings">The server settings</param>
+		/// <param name="noticeService">The notice service singleton</param>
+		/// <param name="aclService">The acl service singleton</param>
+		/// <param name="userCollection">The user collection singleton</param>
+		/// <param name="globals">The global singleton</param>
+		public NoticesController(IOptionsMonitor<ServerSettings> settings, NoticeService noticeService, AclService aclService, IUserCollection userCollection, ISingletonDocument<Globals> globals)
 		{
-			this.Settings = Settings;
-			this.NoticeService = NoticeService;
-			this.AclService = AclService;
-			this.UserCollection = UserCollection;
-			this.Globals = Globals;
+			_settings = settings;
+			_noticeService = noticeService;
+			_aclService = aclService;
+			_userCollection = userCollection;
+			_globals = globals;
 		}
 
 		/// <summary>
 		/// Add a status message
 		/// </summary>
-		/// <param name="Request"></param>
+		/// <param name="request"></param>
 		/// <returns></returns>
 		[HttpPost("/api/v1/notices")]
-		public async Task<ActionResult> AddNoticeAsync(CreateNoticeRequest Request)
+		public async Task<ActionResult> AddNoticeAsync(CreateNoticeRequest request)
 		{
-			if (!await AclService.AuthorizeAsync(AclAction.AdminWrite, User))
+			if (!await _aclService.AuthorizeAsync(AclAction.AdminWrite, User))
 			{
 				return Forbid();
 			}
 
-			INotice? Notice = await NoticeService.AddNoticeAsync(Request.Message, User.GetUserId(), Request.StartTime, Request.FinishTime);
+			INotice? notice = await _noticeService.AddNoticeAsync(request.Message, User.GetUserId(), request.StartTime, request.FinishTime);
 
-			return Notice == null ? NotFound() : Ok();
+			return notice == null ? NotFound() : Ok();
 
 		}
 
 		/// <summary>
 		/// Update a status message
 		/// </summary>
-		/// <param name="Request"></param>
+		/// <param name="request"></param>
 		/// <returns></returns>
 		[HttpPut("/api/v1/notices")]
-		public async Task<ActionResult> UpdateNoticeAsync(UpdateNoticeRequest Request)
+		public async Task<ActionResult> UpdateNoticeAsync(UpdateNoticeRequest request)
 		{
-			if (!await AclService.AuthorizeAsync(AclAction.AdminWrite, User))
+			if (!await _aclService.AuthorizeAsync(AclAction.AdminWrite, User))
 			{
 				return Forbid();
 			}
 
-			return await NoticeService.UpdateNoticeAsync(new ObjectId(Request.Id), Request.Message, Request.StartTime, Request.FinishTime) ? Ok() : NotFound();
+			return await _noticeService.UpdateNoticeAsync(new ObjectId(request.Id), request.Message, request.StartTime, request.FinishTime) ? Ok() : NotFound();
 		}
 
-
 		/// <summary>
 		/// Remove a manually added status message
 		/// </summary>
 		/// <returns></returns>
-		[HttpDelete("/api/v1/notices/{Id}")]
-		public async Task<ActionResult> DeleteNoticeAsync(string Id)
+		[HttpDelete("/api/v1/notices/{id}")]
+		public async Task<ActionResult> DeleteNoticeAsync(string id)
 		{
-			if (!await AclService.AuthorizeAsync(AclAction.AdminWrite, User))
+			if (!await _aclService.AuthorizeAsync(AclAction.AdminWrite, User))
 			{
 				return Forbid();
 			}
 
-			await NoticeService.RemoveNoticeAsync(new ObjectId(Id));
+			await _noticeService.RemoveNoticeAsync(new ObjectId(id));
 			
 			return Ok();
 		}
@@ -128,64 +127,64 @@ namespace Horde.Build.Controllers
 		[HttpGet("/api/v1/notices")]
 		public async Task<List<GetNoticeResponse>> GetNoticesAsync()
 		{			
-			string? ScheduleTimeZone = Settings.CurrentValue.ScheduleTimeZone;
-			TimeZoneInfo TimeZone = (ScheduleTimeZone == null) ? TimeZoneInfo.Local : TZConvert.GetTimeZoneInfo(ScheduleTimeZone);
+			string? scheduleTimeZone = _settings.CurrentValue.ScheduleTimeZone;
+			TimeZoneInfo timeZone = (scheduleTimeZone == null) ? TimeZoneInfo.Local : TZConvert.GetTimeZoneInfo(scheduleTimeZone);
 			
-			Globals GlobalsValue = await Globals.GetAsync();
-			List<GetNoticeResponse> Messages = new List<GetNoticeResponse>();
+			Globals globalsValue = await _globals.GetAsync();
+			List<GetNoticeResponse> messages = new List<GetNoticeResponse>();
 
 			// Add system downtime notices
-			for (int Idx = 0; Idx < GlobalsValue.ScheduledDowntime.Count; Idx++)
+			for (int idx = 0; idx < globalsValue.ScheduledDowntime.Count; idx++)
 			{
-				ScheduledDowntime Downtime = GlobalsValue.ScheduledDowntime[Idx];
+				ScheduledDowntime downtime = globalsValue.ScheduledDowntime[idx];
 
 				// @todo: handle this, though likely will have a custom message
-				if (Downtime.Frequency == ScheduledDowntimeFrequency.Once)
+				if (downtime.Frequency == ScheduledDowntimeFrequency.Once)
 				{
 					continue;
 				}
 
-				DateTimeOffset Now = TimeZoneInfo.ConvertTime(DateTimeOffset.Now, Settings.CurrentValue.TimeZoneInfo);
+				DateTimeOffset now = TimeZoneInfo.ConvertTime(DateTimeOffset.Now, _settings.CurrentValue.TimeZoneInfo);
 				
 				// @todo: this time and formatting should be on client
-				if (Downtime.IsActive(Now))
+				if (downtime.IsActive(now))
 				{
-					(DateTimeOffset StartTime, DateTimeOffset FinishTime) = Downtime.GetNext(Now);
-					string FinishTimeString = String.Format(CultureInfo.CurrentCulture, "{0:t}", TimeZoneInfo.ConvertTime(FinishTime, TimeZone));
-					Messages.Add(new GetNoticeResponse() { StartTime = StartTime.UtcDateTime, FinishTime = FinishTime.UtcDateTime, Message = $"Horde is currently in scheduled downtime. Jobs will resume execution at {FinishTimeString} {TimeZone.Id}." });
+					(DateTimeOffset startTime, DateTimeOffset finishTime) = downtime.GetNext(now);
+					string finishTimeString = String.Format(CultureInfo.CurrentCulture, "{0:t}", TimeZoneInfo.ConvertTime(finishTime, timeZone));
+					messages.Add(new GetNoticeResponse() { StartTime = startTime.UtcDateTime, FinishTime = finishTime.UtcDateTime, Message = $"Horde is currently in scheduled downtime. Jobs will resume execution at {finishTimeString} {timeZone.Id}." });
 				}
 				else 
 				{
 					// add a week to get the actual next scheduled time
-					if (Downtime.Frequency == ScheduledDowntimeFrequency.Weekly)
+					if (downtime.Frequency == ScheduledDowntimeFrequency.Weekly)
 					{
-						Now += TimeSpan.FromDays(7);
+						now += TimeSpan.FromDays(7);
 					}
 
-					(DateTimeOffset StartTime, DateTimeOffset FinishTime) = Downtime.GetNext(Now);
-					string StartTimeString = String.Format(CultureInfo.CurrentCulture, "{0:g}", StartTime);
-					string FinishTimeString = String.Format(CultureInfo.CurrentCulture, "{0:g}", FinishTime);
-					Messages.Add(new GetNoticeResponse() { StartTime = StartTime.UtcDateTime, FinishTime = FinishTime.UtcDateTime, Message = $"Downtime is scheduled from {StartTimeString} to {FinishTimeString} {TimeZone.Id}. No jobs will run during this time." });
+					(DateTimeOffset startTime, DateTimeOffset finishTime) = downtime.GetNext(now);
+					string startTimeString = String.Format(CultureInfo.CurrentCulture, "{0:g}", startTime);
+					string finishTimeString = String.Format(CultureInfo.CurrentCulture, "{0:g}", finishTime);
+					messages.Add(new GetNoticeResponse() { StartTime = startTime.UtcDateTime, FinishTime = finishTime.UtcDateTime, Message = $"Downtime is scheduled from {startTimeString} to {finishTimeString} {timeZone.Id}. No jobs will run during this time." });
 				}
 			}
 
 			// Add user notices
-			List<INotice> Notices = await NoticeService.GetNoticesAsync();
+			List<INotice> notices = await _noticeService.GetNoticesAsync();
 
-			for (int i = 0; i < Notices.Count; i++)
+			for (int i = 0; i < notices.Count; i++)
 			{
-				INotice Notice = Notices[i];
-				GetThinUserInfoResponse? UserInfo = null;
+				INotice notice = notices[i];
+				GetThinUserInfoResponse? userInfo = null;
 
-				if (Notice.UserId != null)
+				if (notice.UserId != null)
 				{
-					UserInfo = new GetThinUserInfoResponse(await UserCollection.GetCachedUserAsync(Notice.UserId));
+					userInfo = new GetThinUserInfoResponse(await _userCollection.GetCachedUserAsync(notice.UserId));
 				}								
 
-				Messages.Add(new GetNoticeResponse(Notice, UserInfo));
+				messages.Add(new GetNoticeResponse(notice, userInfo));
 			}
 
-			return Messages;
+			return messages;
 		}
 	}
 }
diff --git a/Engine/Source/Programs/Horde/Horde.Build/Controllers/PerforceController.cs b/Engine/Source/Programs/Horde/Horde.Build/Controllers/PerforceController.cs
--- a/Engine/Source/Programs/Horde/Horde.Build/Controllers/PerforceController.cs
+++ b/Engine/Source/Programs/Horde/Horde.Build/Controllers/PerforceController.cs
@@ -1,14 +1,14 @@
 ﻿// Copyright Epic Games, Inc. All Rights Reserved.
 
-using Microsoft.AspNetCore.Mvc;
+using System.Collections.Generic;
 using System.Net;
-using Microsoft.Extensions.Logging;
 using System.Threading.Tasks;
-using System.Collections.Generic;
 using Horde.Build.Acls;
+using Horde.Build.Models;
 using Horde.Build.Services;
 using Microsoft.AspNetCore.Authorization;
-using Horde.Build.Models;
+using Microsoft.AspNetCore.Mvc;
+using Microsoft.Extensions.Logging;
 
 namespace Horde.Build.Controllers
 {
@@ -23,26 +23,26 @@ namespace Horde.Build.Controllers
 		/// <summary>
 		/// The database service instance
 		/// </summary>
-		private DatabaseService DatabaseService;
+		private readonly DatabaseService _databaseService;
 
 		/// <summary>
 		/// The ACL service instance
 		/// </summary>
-		private AclService AclService;
+		private readonly AclService _aclService;
 
 		/// <summary>
 		/// Load balancer instance
 		/// </summary>
-		private PerforceLoadBalancer PerforceLoadBalancer;
+		private readonly PerforceLoadBalancer _perforceLoadBalancer;
 
 		/// <summary>
 		/// Constructor
 		/// </summary>
-		public PerforceController(DatabaseService DatabaseService, AclService AclService, PerforceLoadBalancer PerforceLoadBalancer)
+		public PerforceController(DatabaseService databaseService, AclService aclService, PerforceLoadBalancer perforceLoadBalancer)
 		{
-			this.DatabaseService = DatabaseService;
-			this.AclService = AclService;
-			this.PerforceLoadBalancer = PerforceLoadBalancer;
+			_databaseService = databaseService;
+			_aclService = aclService;
+			_perforceLoadBalancer = perforceLoadBalancer;
 		}
 
 		/// <summary>
@@ -53,14 +53,14 @@ namespace Horde.Build.Controllers
 		[Route("/api/v1/perforce/status")]
 		public async Task<ActionResult<List<object>>> GetStatusAsync()
 		{
-			List<IPerforceServer> Servers = await PerforceLoadBalancer.GetServersAsync();
+			List<IPerforceServer> servers = await _perforceLoadBalancer.GetServersAsync();
 
-			List<object> Responses = new List<object>();
-			foreach (IPerforceServer Server in Servers)
+			List<object> responses = new List<object>();
+			foreach (IPerforceServer server in servers)
 			{
-				Responses.Add(new { Server.ServerAndPort, Server.BaseServerAndPort, Server.Cluster, Server.NumLeases, Server.Status, Server.Detail, Server.LastUpdateTime });
+				responses.Add(new { server.ServerAndPort, server.BaseServerAndPort, server.Cluster, server.NumLeases, server.Status, server.Detail, server.LastUpdateTime });
 			}
-			return Responses;
+			return responses;
 		}
 
 		/// <summary>
@@ -71,13 +71,13 @@ namespace Horde.Build.Controllers
 		[Route("/api/v1/perforce/settings")]
 		public async Task<ActionResult<List<PerforceCluster>>> GetPerforceSettingsAsync()
 		{
-			if (!await AclService.AuthorizeAsync(AclAction.AdminRead, User))
+			if (!await _aclService.AuthorizeAsync(AclAction.AdminRead, User))
 			{
 				return Forbid();
 			}
 
-			Globals Globals = await DatabaseService.GetGlobalsAsync();
-			return Globals.PerforceClusters;
+			Globals globals = await _databaseService.GetGlobalsAsync();
+			return globals.PerforceClusters;
 		}
 	}
 
@@ -91,14 +91,14 @@ namespace Horde.Build.Controllers
 		/// <summary>
 		/// Logger instance
 		/// </summary>
-		private readonly ILogger<PerforceController> Logger;
+		private readonly ILogger<PerforceController> _logger;
 
 		/// <summary>
 		/// Constructor
 		/// </summary>
-		public PublicPerforceController(ILogger<PerforceController> Logger)
+		public PublicPerforceController(ILogger<PerforceController> logger)
 		{
-			this.Logger = Logger;
+			_logger = logger;
 		}
 
 		/// <summary>
@@ -106,13 +106,13 @@ namespace Horde.Build.Controllers
 		/// </summary>
 		/// <returns>200 OK on success</returns>
 		[HttpPost]
-		[Route("/api/v1/perforce/trigger/{Type}")]
-		public ActionResult TriggerCallback(string Type, [FromQuery] long? Changelist = null, [FromQuery(Name = "user")] string? PerforceUser = null)
+		[Route("/api/v1/perforce/trigger/{type}")]
+		public ActionResult TriggerCallback(string type, [FromQuery] long? changelist = null, [FromQuery(Name = "user")] string? perforceUser = null)
 		{
 			// Currently just a placeholder until correct triggers are in place.
-			Logger.LogDebug("Received Perforce trigger callback. Type={Type} Changelist={Changelist} User={User}", Type, Changelist, PerforceUser);
-			string Content = "{\"message\": \"Trigger received\"}";
-			return new ContentResult { ContentType = "text/plain", StatusCode = (int)HttpStatusCode.OK, Content = Content };
+			_logger.LogDebug("Received Perforce trigger callback. Type={Type} Changelist={Changelist} User={User}", type, changelist, perforceUser);
+			string content = "{\"message\": \"Trigger received\"}";
+			return new ContentResult { ContentType = "text/plain", StatusCode = (int)HttpStatusCode.OK, Content = content };
 		}
 	}
 }
diff --git a/Engine/Source/Programs/Horde/Horde.Build/Controllers/PoolsController.cs b/Engine/Source/Programs/Horde/Horde.Build/Controllers/PoolsController.cs
--- a/Engine/Source/Programs/Horde/Horde.Build/Controllers/PoolsController.cs
+++ b/Engine/Source/Programs/Horde/Horde.Build/Controllers/PoolsController.cs
@@ -1,5 +1,8 @@
 ﻿// Copyright Epic Games, Inc. All Rights Reserved.
 
+using System;
+using System.Collections.Generic;
+using System.Threading.Tasks;
 using Horde.Build.Acls;
 using Horde.Build.Api;
 using Horde.Build.Models;
@@ -7,12 +10,6 @@ using Horde.Build.Services;
 using Horde.Build.Utilities;
 using Microsoft.AspNetCore.Authorization;
 using Microsoft.AspNetCore.Mvc;
-using MongoDB.Bson;
-using System;
-using System.Collections.Generic;
-using System.ComponentModel.DataAnnotations;
-using System.Linq;
-using System.Threading.Tasks;
 
 namespace Horde.Build.Controllers
 {
@@ -29,151 +26,151 @@ namespace Horde.Build.Controllers
 		/// <summary>
 		/// Singleton instance of the ACL service
 		/// </summary>
-		private readonly AclService AclService;
+		private readonly AclService _aclService;
 
 		/// <summary>
 		/// Singleton instance of the pool service
 		/// </summary>
-		private readonly PoolService PoolService;
+		private readonly PoolService _poolService;
 
 		/// <summary>
 		/// Constructor
 		/// </summary>
-		/// <param name="AclService">The ACL service</param>
-		/// <param name="PoolService">The pool service</param>
-		public PoolsController(AclService AclService, PoolService PoolService)
+		/// <param name="aclService">The ACL service</param>
+		/// <param name="poolService">The pool service</param>
+		public PoolsController(AclService aclService, PoolService poolService)
 		{
-			this.AclService = AclService;
-			this.PoolService = PoolService;
+			_aclService = aclService;
+			_poolService = poolService;
 		}
 
 		/// <summary>
 		/// Creates a new pool
 		/// </summary>
-		/// <param name="Create">Parameters for the new pool.</param>
+		/// <param name="create">Parameters for the new pool.</param>
 		/// <returns>Http result code</returns>
 		[HttpPost]
 		[Route("/api/v1/pools")]
-		public async Task<ActionResult<CreatePoolResponse>> CreatePoolAsync([FromBody] CreatePoolRequest Create)
+		public async Task<ActionResult<CreatePoolResponse>> CreatePoolAsync([FromBody] CreatePoolRequest create)
 		{
-			if(!await AclService.AuthorizeAsync(AclAction.CreatePool, User))
+			if(!await _aclService.AuthorizeAsync(AclAction.CreatePool, User))
 			{
 				return Forbid(AclAction.CreatePool);
 			}
 
-			Horde.Build.Fleet.Autoscale.LeaseUtilizationSettings? LuSettings = Create.LeaseUtilizationSettings?.Convert();
-			Horde.Build.Fleet.Autoscale.JobQueueSettings? JqSettings = Create.JobQueueSettings?.Convert();
+			Horde.Build.Fleet.Autoscale.LeaseUtilizationSettings? luSettings = create.LeaseUtilizationSettings?.Convert();
+			Horde.Build.Fleet.Autoscale.JobQueueSettings? jqSettings = create.JobQueueSettings?.Convert();
 			
-			TimeSpan? ScaleOutCooldown = Create.ScaleOutCooldown == null ? null : TimeSpan.FromSeconds(Create.ScaleOutCooldown.Value);
-			TimeSpan? ScaleInCooldown = Create.ScaleInCooldown == null ? null : TimeSpan.FromSeconds(Create.ScaleInCooldown.Value);
+			TimeSpan? scaleOutCooldown = create.ScaleOutCooldown == null ? null : TimeSpan.FromSeconds(create.ScaleOutCooldown.Value);
+			TimeSpan? scaleInCooldown = create.ScaleInCooldown == null ? null : TimeSpan.FromSeconds(create.ScaleInCooldown.Value);
 
-			IPool NewPool = await PoolService.CreatePoolAsync(
-				Create.Name, Create.Condition, Create.EnableAutoscaling, Create.MinAgents, Create.NumReserveAgents,
-				ScaleOutCooldown, ScaleInCooldown, Create.SizeStrategy, LuSettings, JqSettings, Create.Properties);
-			return new CreatePoolResponse(NewPool.Id.ToString());
+			IPool newPool = await _poolService.CreatePoolAsync(
+				create.Name, create.Condition, create.EnableAutoscaling, create.MinAgents, create.NumReserveAgents,
+				scaleOutCooldown, scaleInCooldown, create.SizeStrategy, luSettings, jqSettings, create.Properties);
+			return new CreatePoolResponse(newPool.Id.ToString());
 		}
 
 		/// <summary>
 		/// Query all the pools
 		/// </summary>
-		/// <param name="Filter">Filter for the properties to return</param>
+		/// <param name="filter">Filter for the properties to return</param>
 		/// <returns>Information about all the pools</returns>
 		[HttpGet]
 		[Route("/api/v1/pools")]
 		[ProducesResponseType(typeof(List<GetPoolResponse>), 200)]
-		public async Task<ActionResult<List<object>>> GetPoolsAsync([FromQuery] PropertyFilter? Filter = null)
+		public async Task<ActionResult<List<object>>> GetPoolsAsync([FromQuery] PropertyFilter? filter = null)
 		{
-			if (!await AclService.AuthorizeAsync(AclAction.ListPools, User))
+			if (!await _aclService.AuthorizeAsync(AclAction.ListPools, User))
 			{
 				return Forbid(AclAction.ListPools);
 			}
 
-			List<IPool> Pools = await PoolService.GetPoolsAsync();
+			List<IPool> pools = await _poolService.GetPoolsAsync();
 
-			List<object> Responses = new List<object>();
-			foreach (IPool Pool in Pools)
+			List<object> responses = new List<object>();
+			foreach (IPool pool in pools)
 			{
-				Responses.Add(new GetPoolResponse(Pool).ApplyFilter(Filter));
+				responses.Add(new GetPoolResponse(pool).ApplyFilter(filter));
 			}
-			return Responses;
+			return responses;
 		}
 
 		/// <summary>
 		/// Retrieve information about a specific pool
 		/// </summary>
-		/// <param name="PoolId">Id of the pool to get information about</param>
-		/// <param name="Filter">Filter to apply to the returned properties</param>
+		/// <param name="poolId">Id of the pool to get information about</param>
+		/// <param name="filter">Filter to apply to the returned properties</param>
 		/// <returns>Information about the requested pool</returns>
 		[HttpGet]
-		[Route("/api/v1/pools/{PoolId}")]
+		[Route("/api/v1/pools/{poolId}")]
 		[ProducesResponseType(typeof(GetPoolResponse), 200)]
-		public async Task<ActionResult<object>> GetPoolAsync(string PoolId, [FromQuery] PropertyFilter? Filter = null)
+		public async Task<ActionResult<object>> GetPoolAsync(string poolId, [FromQuery] PropertyFilter? filter = null)
 		{
-			if (!await AclService.AuthorizeAsync(AclAction.ViewPool, User))
+			if (!await _aclService.AuthorizeAsync(AclAction.ViewPool, User))
 			{
 				return Forbid(AclAction.ViewPool);
 			}
 
-			PoolId PoolIdValue = new PoolId(PoolId);
+			PoolId poolIdValue = new PoolId(poolId);
 
-			IPool? Pool = await PoolService.GetPoolAsync(PoolIdValue);
-			if (Pool == null)
+			IPool? pool = await _poolService.GetPoolAsync(poolIdValue);
+			if (pool == null)
 			{
-				return NotFound(PoolIdValue);
+				return NotFound(poolIdValue);
 			}
 
-			return new GetPoolResponse(Pool).ApplyFilter(Filter);
+			return new GetPoolResponse(pool).ApplyFilter(filter);
 		}
 
 		/// <summary>
 		/// Update a pool's properties.
 		/// </summary>
-		/// <param name="PoolId">Id of the pool to update</param>
-		/// <param name="Update">Items on the pool to update</param>
+		/// <param name="poolId">Id of the pool to update</param>
+		/// <param name="update">Items on the pool to update</param>
 		/// <returns>Http result code</returns>
 		[HttpPut]
-		[Route("/api/v1/pools/{PoolId}")]
-		public async Task<ActionResult> UpdatePoolAsync(string PoolId, [FromBody] UpdatePoolRequest Update)
+		[Route("/api/v1/pools/{poolId}")]
+		public async Task<ActionResult> UpdatePoolAsync(string poolId, [FromBody] UpdatePoolRequest update)
 		{
-			if (!await AclService.AuthorizeAsync(AclAction.UpdatePool, User))
+			if (!await _aclService.AuthorizeAsync(AclAction.UpdatePool, User))
 			{
 				return Forbid(AclAction.UpdatePool);
 			}
 
-			PoolId PoolIdValue = new PoolId(PoolId);
+			PoolId poolIdValue = new PoolId(poolId);
 			
-			IPool? Pool = await PoolService.GetPoolAsync(PoolIdValue);
-			if(Pool == null)
+			IPool? pool = await _poolService.GetPoolAsync(poolIdValue);
+			if(pool == null)
 			{
-				return NotFound(PoolIdValue);
+				return NotFound(poolIdValue);
 			}
 			
-			TimeSpan? ScaleOutCooldown = Update.ScaleOutCooldown == null ? null : TimeSpan.FromSeconds(Update.ScaleOutCooldown.Value);
-			TimeSpan? ScaleInCooldown = Update.ScaleInCooldown == null ? null : TimeSpan.FromSeconds(Update.ScaleInCooldown.Value);
+			TimeSpan? scaleOutCooldown = update.ScaleOutCooldown == null ? null : TimeSpan.FromSeconds(update.ScaleOutCooldown.Value);
+			TimeSpan? scaleInCooldown = update.ScaleInCooldown == null ? null : TimeSpan.FromSeconds(update.ScaleInCooldown.Value);
 
-			await PoolService.UpdatePoolAsync(Pool, Update.Name, Update.Condition, Update.EnableAutoscaling,
-				Update.MinAgents, Update.NumReserveAgents, Update.Properties, ScaleOutCooldown, ScaleInCooldown, Update.SizeStrategy);
+			await _poolService.UpdatePoolAsync(pool, update.Name, update.Condition, update.EnableAutoscaling,
+				update.MinAgents, update.NumReserveAgents, update.Properties, scaleOutCooldown, scaleInCooldown, update.SizeStrategy);
 			return new OkResult();
 		}
 
 		/// <summary>
 		/// Delete a pool
 		/// </summary>
-		/// <param name="PoolId">Id of the pool to delete</param>
+		/// <param name="poolId">Id of the pool to delete</param>
 		/// <returns>Http result code</returns>
 		[HttpDelete]
-		[Route("/api/v1/pools/{PoolId}")]
-		public async Task<ActionResult> DeletePoolAsync(string PoolId)
+		[Route("/api/v1/pools/{poolId}")]
+		public async Task<ActionResult> DeletePoolAsync(string poolId)
 		{
-			if (!await AclService.AuthorizeAsync(AclAction.DeletePool, User))
+			if (!await _aclService.AuthorizeAsync(AclAction.DeletePool, User))
 			{
 				return Forbid(AclAction.DeletePool);
 			}
 
-			PoolId PoolIdValue = new PoolId(PoolId);
-			if(!await PoolService.DeletePoolAsync(PoolIdValue))
+			PoolId poolIdValue = new PoolId(poolId);
+			if(!await _poolService.DeletePoolAsync(poolIdValue))
 			{
-				return NotFound(PoolIdValue);
+				return NotFound(poolIdValue);
 			}
 			return Ok();
 		}
@@ -181,28 +178,28 @@ namespace Horde.Build.Controllers
 		/// <summary>
 		/// Batch update pool properties
 		/// </summary>
-		/// <param name="BatchUpdates">List of pools to update</param>
+		/// <param name="batchUpdates">List of pools to update</param>
 		/// <returns>Http result code</returns>
 		[HttpPut]
 		[Route("/api/v1/pools")]
-		public async Task<ActionResult> UpdatePoolAsync([FromBody] List<BatchUpdatePoolRequest> BatchUpdates)
+		public async Task<ActionResult> UpdatePoolAsync([FromBody] List<BatchUpdatePoolRequest> batchUpdates)
 		{
-			if (!await AclService.AuthorizeAsync(AclAction.UpdatePool, User))
+			if (!await _aclService.AuthorizeAsync(AclAction.UpdatePool, User))
 			{
 				return Forbid(AclAction.UpdatePool);
 			}
 
-			foreach (BatchUpdatePoolRequest Update in BatchUpdates)
+			foreach (BatchUpdatePoolRequest update in batchUpdates)
 			{
-				PoolId PoolIdValue = new PoolId(Update.Id);
+				PoolId poolIdValue = new PoolId(update.Id);
 
-				IPool? Pool = await PoolService.GetPoolAsync(PoolIdValue);
-				if (Pool == null)
+				IPool? pool = await _poolService.GetPoolAsync(poolIdValue);
+				if (pool == null)
 				{
-					return NotFound(PoolIdValue);
+					return NotFound(poolIdValue);
 				}
 
-				await PoolService.UpdatePoolAsync(Pool, Update.Name, NewProperties: Update.Properties);
+				await _poolService.UpdatePoolAsync(pool, update.Name, newProperties: update.Properties);
 			}
 			return Ok();
 		}
diff --git a/Engine/Source/Programs/Horde/Horde.Build/Controllers/ProjectsController.cs b/Engine/Source/Programs/Horde/Horde.Build/Controllers/ProjectsController.cs
--- a/Engine/Source/Programs/Horde/Horde.Build/Controllers/ProjectsController.cs
+++ b/Engine/Source/Programs/Horde/Horde.Build/Controllers/ProjectsController.cs
@@ -1,5 +1,7 @@
 ﻿// Copyright Epic Games, Inc. All Rights Reserved.
 
+using System.Collections.Generic;
+using System.Threading.Tasks;
 using Horde.Build.Acls;
 using Horde.Build.Api;
 using Horde.Build.Models;
@@ -7,12 +9,6 @@ using Horde.Build.Services;
 using Horde.Build.Utilities;
 using Microsoft.AspNetCore.Authorization;
 using Microsoft.AspNetCore.Mvc;
-using MongoDB.Bson;
-using System;
-using System.Collections.Generic;
-using System.ComponentModel.DataAnnotations;
-using System.Linq;
-using System.Threading.Tasks;
 
 namespace Horde.Build.Controllers
 {
@@ -29,128 +25,128 @@ namespace Horde.Build.Controllers
 		/// <summary>
 		/// Singleton instance of the project service
 		/// </summary>
-		private readonly ProjectService ProjectService;
+		private readonly ProjectService _projectService;
 
 		/// <summary>
 		/// Singleton instance of the stream service
 		/// </summary>
-		private readonly StreamService StreamService;
+		private readonly StreamService _streamService;
 
 		/// <summary>
 		/// Constructor
 		/// </summary>
-		/// <param name="ProjectService">The project service</param>
-		/// <param name="StreamService">The stream service</param>
-		public ProjectsController(ProjectService ProjectService, StreamService StreamService)
+		/// <param name="projectService">The project service</param>
+		/// <param name="streamService">The stream service</param>
+		public ProjectsController(ProjectService projectService, StreamService streamService)
 		{
-			this.ProjectService = ProjectService;
-			this.StreamService = StreamService;
+			_projectService = projectService;
+			_streamService = streamService;
 		}
 
 		/// <summary>
 		/// Query all the projects
 		/// </summary>
-		/// <param name="IncludeStreams">Whether to include streams in the response</param>
-		/// <param name="IncludeCategories">Whether to include categories in the response</param>
-		/// <param name="Filter">Filter for the properties to return</param>
+		/// <param name="includeStreams">Whether to include streams in the response</param>
+		/// <param name="includeCategories">Whether to include categories in the response</param>
+		/// <param name="filter">Filter for the properties to return</param>
 		/// <returns>Information about all the projects</returns>
 		[HttpGet]
 		[Route("/api/v1/projects")]
 		[ProducesResponseType(typeof(List<GetProjectResponse>), 200)]
-		public async Task<ActionResult<List<object>>> GetProjectsAsync([FromQuery(Name = "Streams")] bool IncludeStreams = false, [FromQuery(Name = "Categories")] bool IncludeCategories = false, [FromQuery] PropertyFilter? Filter = null)
+		public async Task<ActionResult<List<object>>> GetProjectsAsync([FromQuery(Name = "Streams")] bool includeStreams = false, [FromQuery(Name = "Categories")] bool includeCategories = false, [FromQuery] PropertyFilter? filter = null)
 		{
-			List<IProject> Projects = await ProjectService.GetProjectsAsync();
-			ProjectPermissionsCache PermissionsCache = new ProjectPermissionsCache();
+			List<IProject> projects = await _projectService.GetProjectsAsync();
+			ProjectPermissionsCache permissionsCache = new ProjectPermissionsCache();
 
-			List<IStream>? Streams = null;
-			if (IncludeStreams || IncludeCategories)
+			List<IStream>? streams = null;
+			if (includeStreams || includeCategories)
 			{
-				Streams = await StreamService.GetStreamsAsync();
+				streams = await _streamService.GetStreamsAsync();
 			}
 
-			List<object> Responses = new List<object>();
-			foreach (IProject Project in Projects)
+			List<object> responses = new List<object>();
+			foreach (IProject project in projects)
 			{
-				if (await ProjectService.AuthorizeAsync(Project, AclAction.ViewProject, User, PermissionsCache))
+				if (await _projectService.AuthorizeAsync(project, AclAction.ViewProject, User, permissionsCache))
 				{
-					bool bIncludeAcl = await ProjectService.AuthorizeAsync(Project, AclAction.ViewPermissions, User, PermissionsCache);
-					Responses.Add(Project.ToResponse(IncludeStreams, IncludeCategories, Streams, bIncludeAcl).ApplyFilter(Filter));
+					bool bIncludeAcl = await _projectService.AuthorizeAsync(project, AclAction.ViewPermissions, User, permissionsCache);
+					responses.Add(project.ToResponse(includeStreams, includeCategories, streams, bIncludeAcl).ApplyFilter(filter));
 				}
 			}
-			return Responses;
+			return responses;
 		}
 
 		/// <summary>
 		/// Retrieve information about a specific project
 		/// </summary>
-		/// <param name="ProjectId">Id of the project to get information about</param>
-		/// <param name="Filter">Filter for the properties to return</param>
+		/// <param name="projectId">Id of the project to get information about</param>
+		/// <param name="filter">Filter for the properties to return</param>
 		/// <returns>Information about the requested project</returns>
 		[HttpGet]
-		[Route("/api/v1/projects/{ProjectId}")]
+		[Route("/api/v1/projects/{projectId}")]
 		[ProducesResponseType(typeof(List<GetProjectResponse>), 200)]
-		public async Task<ActionResult<object>> GetProjectAsync(ProjectId ProjectId, [FromQuery] PropertyFilter? Filter = null)
+		public async Task<ActionResult<object>> GetProjectAsync(ProjectId projectId, [FromQuery] PropertyFilter? filter = null)
 		{
-			IProject? Project = await ProjectService.GetProjectAsync(ProjectId);
-			if (Project == null)
+			IProject? project = await _projectService.GetProjectAsync(projectId);
+			if (project == null)
 			{
-				return NotFound(ProjectId);
+				return NotFound(projectId);
 			}
 
-			ProjectPermissionsCache Cache = new ProjectPermissionsCache();
-			if (!await ProjectService.AuthorizeAsync(Project, AclAction.ViewProject, User, Cache))
+			ProjectPermissionsCache cache = new ProjectPermissionsCache();
+			if (!await _projectService.AuthorizeAsync(project, AclAction.ViewProject, User, cache))
 			{
-				return Forbid(AclAction.ViewProject, ProjectId);
+				return Forbid(AclAction.ViewProject, projectId);
 			}
 
-			bool bIncludeStreams = PropertyFilter.Includes(Filter, nameof(GetProjectResponse.Streams));
-			bool bIncludeCategories = PropertyFilter.Includes(Filter, nameof(GetProjectResponse.Categories));
+			bool bIncludeStreams = PropertyFilter.Includes(filter, nameof(GetProjectResponse.Streams));
+			bool bIncludeCategories = PropertyFilter.Includes(filter, nameof(GetProjectResponse.Categories));
 
-			List<IStream>? VisibleStreams = null;
+			List<IStream>? visibleStreams = null;
 			if (bIncludeStreams || bIncludeCategories)
 			{
-				VisibleStreams = new List<IStream>();
+				visibleStreams = new List<IStream>();
 
-				List<IStream> Streams = await StreamService.GetStreamsAsync(Project.Id);
-				foreach (IStream Stream in Streams)
+				List<IStream> streams = await _streamService.GetStreamsAsync(project.Id);
+				foreach (IStream stream in streams)
 				{
-					if (await StreamService.AuthorizeAsync(Stream, AclAction.ViewStream, User, Cache))
+					if (await _streamService.AuthorizeAsync(stream, AclAction.ViewStream, User, cache))
 					{
-						VisibleStreams.Add(Stream);
+						visibleStreams.Add(stream);
 					}
 				}
 			}
 
-			bool bIncludeAcl = await ProjectService.AuthorizeAsync(Project, AclAction.ViewPermissions, User, Cache);
-			return Project.ToResponse(bIncludeStreams, bIncludeCategories, VisibleStreams, bIncludeAcl).ApplyFilter(Filter);
+			bool bIncludeAcl = await _projectService.AuthorizeAsync(project, AclAction.ViewPermissions, User, cache);
+			return project.ToResponse(bIncludeStreams, bIncludeCategories, visibleStreams, bIncludeAcl).ApplyFilter(filter);
 		}
 
 		/// <summary>
 		/// Retrieve information about a specific project
 		/// </summary>
-		/// <param name="ProjectId">Id of the project to get information about</param>
+		/// <param name="projectId">Id of the project to get information about</param>
 		/// <returns>Information about the requested project</returns>
 		[HttpGet]
-		[Route("/api/v1/projects/{ProjectId}/logo")]
-		public async Task<ActionResult<object>> GetProjectLogoAsync(ProjectId ProjectId)
+		[Route("/api/v1/projects/{projectId}/logo")]
+		public async Task<ActionResult<object>> GetProjectLogoAsync(ProjectId projectId)
 		{
-			IProject? Project = await ProjectService.GetProjectAsync(ProjectId);
-			if (Project == null)
+			IProject? project = await _projectService.GetProjectAsync(projectId);
+			if (project == null)
 			{
-				return NotFound(ProjectId);
+				return NotFound(projectId);
 			}
-			if (!await ProjectService.AuthorizeAsync(Project, AclAction.ViewProject, User, null))
+			if (!await _projectService.AuthorizeAsync(project, AclAction.ViewProject, User, null))
 			{
-				return Forbid(AclAction.ViewProject, ProjectId);
+				return Forbid(AclAction.ViewProject, projectId);
 			}
 
-			IProjectLogo? ProjectLogo = await ProjectService.Collection.GetLogoAsync(ProjectId);
-			if (ProjectLogo == null)
+			IProjectLogo? projectLogo = await _projectService.Collection.GetLogoAsync(projectId);
+			if (projectLogo == null)
 			{
 				return NotFound();
 			}
 
-			return new FileContentResult(ProjectLogo.Data, ProjectLogo.MimeType);
+			return new FileContentResult(projectLogo.Data, projectLogo.MimeType);
 		}
 	}
 }
diff --git a/Engine/Source/Programs/Horde/Horde.Build/Controllers/ReportsController.cs b/Engine/Source/Programs/Horde/Horde.Build/Controllers/ReportsController.cs
--- a/Engine/Source/Programs/Horde/Horde.Build/Controllers/ReportsController.cs
+++ b/Engine/Source/Programs/Horde/Horde.Build/Controllers/ReportsController.cs
@@ -2,15 +2,11 @@
 
 using System;
 using System.Collections.Generic;
-using System.Linq;
 using System.Threading.Tasks;
 using Horde.Build.Api;
 using Horde.Build.Collections;
 using Horde.Build.Models;
-using Horde.Build.Services;
-using Microsoft.AspNetCore.Authorization;
 using Microsoft.AspNetCore.Mvc;
-using EpicGames.Core;
 
 namespace Horde.Build.Controllers
 {
@@ -24,34 +20,34 @@ namespace Horde.Build.Controllers
 		/// <summary>
 		/// the Telemetry collection singleton
 		/// </summary>
-		ITelemetryCollection TelemetryCollection;
+		readonly ITelemetryCollection _telemetryCollection;
 
 		/// <summary>
 		/// Constructor
 		/// </summary>
-		/// <param name="TelemetryCollection">The telemetry collection</param>
-		public ReportsController(ITelemetryCollection TelemetryCollection)
+		/// <param name="telemetryCollection">The telemetry collection</param>
+		public ReportsController(ITelemetryCollection telemetryCollection)
 		{
-			this.TelemetryCollection = TelemetryCollection;
+			_telemetryCollection = telemetryCollection;
 		}
 
 		/// <summary>
 		/// Gets a collection of utilization data including an endpoint and a range
 		/// </summary>
 		[HttpGet]
-		[Route("/api/v1/reports/utilization/{EndDate}")]
-		public async Task<ActionResult<List<UtilizationTelemetryResponse>>> GetStreamUtilizationData(DateTime EndDate, [FromQuery(Name = "Range")] int Range, [FromQuery(Name = "TzOffset")] int? TzOffset)
+		[Route("/api/v1/reports/utilization/{endDate}")]
+		public async Task<ActionResult<List<UtilizationTelemetryResponse>>> GetStreamUtilizationData(DateTime endDate, [FromQuery(Name = "Range")] int range, [FromQuery(Name = "TzOffset")] int? tzOffset)
 		{
 			// Logic here is a bit messy. The client is always passing in a date at midnight
 			// If user passes in 12/1/2020 into the date with range of 1, the range should be 12/1/2020:00:00:00 to 12/1/2020:23:59:59
 			// Range 2 should be 11/30/2020:00:00:00 to 12/1/2020:23:59:59
-			int Offset = TzOffset ?? 0;
-			DateTimeOffset EndDateOffset = new DateTimeOffset(EndDate, TimeSpan.FromHours(Offset)).Add(new TimeSpan(23, 59, 59));
-			DateTimeOffset StartDateOffset = EndDate.Subtract(new TimeSpan(Range - 1, 0, 0, 0));
+			int offset = tzOffset ?? 0;
+			DateTimeOffset endDateOffset = new DateTimeOffset(endDate, TimeSpan.FromHours(offset)).Add(new TimeSpan(23, 59, 59));
+			DateTimeOffset startDateOffset = endDate.Subtract(new TimeSpan(range - 1, 0, 0, 0));
 
-			List<IUtilizationTelemetry> Telemetry = await TelemetryCollection.GetUtilizationTelemetryAsync(StartDateOffset.UtcDateTime, EndDateOffset.UtcDateTime);
+			List<IUtilizationTelemetry> telemetry = await _telemetryCollection.GetUtilizationTelemetryAsync(startDateOffset.UtcDateTime, endDateOffset.UtcDateTime);
 
-			return Telemetry.ConvertAll(Telemetry => new UtilizationTelemetryResponse(Telemetry));
+			return telemetry.ConvertAll(telemetry => new UtilizationTelemetryResponse(telemetry));
 		}
 	}
 }
diff --git a/Engine/Source/Programs/Horde/Horde.Build/Controllers/SchemaController.cs b/Engine/Source/Programs/Horde/Horde.Build/Controllers/SchemaController.cs
--- a/Engine/Source/Programs/Horde/Horde.Build/Controllers/SchemaController.cs
+++ b/Engine/Source/Programs/Horde/Horde.Build/Controllers/SchemaController.cs
@@ -1,19 +1,16 @@
 // Copyright Epic Games, Inc. All Rights Reserved.
 
-using EpicGames.Core;
-using Horde.Build.Models;
-using Horde.Build.Utilities;
-using Microsoft.AspNetCore.Mvc;
-using Microsoft.Extensions.Primitives;
 using System;
 using System.Collections.Generic;
 using System.IO;
 using System.Linq;
 using System.Net;
 using System.Reflection;
-using System.Text.Json;
 using System.Text.Json.Serialization;
-using System.Threading.Tasks;
+using EpicGames.Core;
+using Horde.Build.Utilities;
+using Microsoft.AspNetCore.Mvc;
+using Microsoft.Extensions.Primitives;
 
 namespace Horde.Build.Controllers
 {
@@ -47,15 +44,15 @@ namespace Horde.Build.Controllers
 
 		static Type[] FindSchemaTypes()
 		{
-			List<Type> SchemaTypes = new List<Type>();
-			foreach (Type Type in Assembly.GetExecutingAssembly().GetTypes())
+			List<Type> schemaTypes = new List<Type>();
+			foreach (Type type in Assembly.GetExecutingAssembly().GetTypes())
 			{
-				if (Type.GetCustomAttribute<JsonSchemaAttribute>() != null)
+				if (type.GetCustomAttribute<JsonSchemaAttribute>() != null)
 				{
-					SchemaTypes.Add(Type);
+					schemaTypes.Add(type);
 				}
 			}
-			return SchemaTypes.ToArray();
+			return schemaTypes.ToArray();
 		}
 
 		/// <summary>
@@ -66,55 +63,55 @@ namespace Horde.Build.Controllers
 		[Route("/api/v1/schema/catalog.json")]
 		public ActionResult GetCatalog()
 		{
-			string? Host = null;
+			string? host = null;
 
-			StringValues Hosts;
-			if (Request.Headers.TryGetValue("Host", out Hosts))
+			StringValues hosts;
+			if (Request.Headers.TryGetValue("Host", out hosts))
 			{
-				Host = Hosts.FirstOrDefault();
+				host = hosts.FirstOrDefault();
 			}
 
-			if (Host == null)
+			if (host == null)
 			{
-				Host = Dns.GetHostName();
+				host = Dns.GetHostName();
 			}
 
-			CatalogRoot Root = new CatalogRoot();
-			foreach (Type SchemaType in ConfigSchemas)
+			CatalogRoot root = new CatalogRoot();
+			foreach (Type schemaType in ConfigSchemas)
 			{
-				JsonSchemaAttribute? SchemaAttribute = SchemaType.GetCustomAttribute<JsonSchemaAttribute>();
-				if (SchemaAttribute != null)
+				JsonSchemaAttribute? schemaAttribute = schemaType.GetCustomAttribute<JsonSchemaAttribute>();
+				if (schemaAttribute != null)
 				{
-					JsonSchemaCatalogAttribute? CatalogAttribute = SchemaType.GetCustomAttribute<JsonSchemaCatalogAttribute>();
-					if (CatalogAttribute != null)
+					JsonSchemaCatalogAttribute? catalogAttribute = schemaType.GetCustomAttribute<JsonSchemaCatalogAttribute>();
+					if (catalogAttribute != null)
 					{
-						Uri Url = new Uri($"https://{Host}/api/v1/schema/types/{SchemaType.Name}.json");
-						Root.Schemas.Add(new CatalogItem { Name = CatalogAttribute.Name, Description = CatalogAttribute.Description, FileMatch = CatalogAttribute.FileMatch, Url = Url });
+						Uri url = new Uri($"https://{host}/api/v1/schema/types/{schemaType.Name}.json");
+						root.Schemas.Add(new CatalogItem { Name = catalogAttribute.Name, Description = catalogAttribute.Description, FileMatch = catalogAttribute.FileMatch, Url = url });
 					}
 				}
 			}
-			return Ok(Root);
+			return Ok(root);
 		}
 
 		/// <summary>
 		/// Gets a specific schema
 		/// </summary>
-		/// <param name="TypeName">The type name</param>
+		/// <param name="typeName">The type name</param>
 		/// <returns></returns>
 		[HttpGet]
-		[Route("/api/v1/schema/types/{TypeName}.json")]
-		public ActionResult GetSchema(string TypeName)
+		[Route("/api/v1/schema/types/{typeName}.json")]
+		public ActionResult GetSchema(string typeName)
 		{
-			foreach (Type SchemaType in ConfigSchemas)
+			foreach (Type schemaType in ConfigSchemas)
 			{
-				if (SchemaType.Name.Equals(TypeName, StringComparison.OrdinalIgnoreCase))
+				if (schemaType.Name.Equals(typeName, StringComparison.OrdinalIgnoreCase))
 				{
-					JsonSchema Schema = Schemas.CreateSchema(SchemaType);
+					JsonSchema schema = Schemas.CreateSchema(schemaType);
 
-					using MemoryStream Stream = new MemoryStream();
-					Schema.Write(Stream);
+					using MemoryStream stream = new MemoryStream();
+					schema.Write(stream);
 
-					return new FileContentResult(Stream.ToArray(), "application/json");
+					return new FileContentResult(stream.ToArray(), "application/json");
 				}
 			}
 			return NotFound();
diff --git a/Engine/Source/Programs/Horde/Horde.Build/Controllers/ServerController.cs b/Engine/Source/Programs/Horde/Horde.Build/Controllers/ServerController.cs
--- a/Engine/Source/Programs/Horde/Horde.Build/Controllers/ServerController.cs
+++ b/Engine/Source/Programs/Horde/Horde.Build/Controllers/ServerController.cs
@@ -1,27 +1,11 @@
 // Copyright Epic Games, Inc. All Rights Reserved.
 
-using EpicGames.Core;
+using System.Diagnostics;
+using System.Reflection;
 using Horde.Build.Api;
-using Horde.Build.Collections;
-using HordeCommon;
-using Horde.Build.Models;
-using Horde.Build.Services;
-using Horde.Build.Utilities;
 using Microsoft.AspNetCore.Authorization;
-using Microsoft.AspNetCore.Authorization.Infrastructure;
 using Microsoft.AspNetCore.Mvc;
 using Microsoft.Extensions.Options;
-using MongoDB.Bson;
-using System;
-using System.Collections.Generic;
-using System.IdentityModel.Tokens.Jwt;
-using System.IO;
-using System.Linq;
-using System.Runtime.InteropServices;
-using System.Security.Claims;
-using System.Threading.Tasks;
-using System.Reflection;
-using System.Diagnostics;
 
 namespace Horde.Build.Controllers
 {
@@ -37,14 +21,14 @@ namespace Horde.Build.Controllers
 		/// <summary>
 		/// Settings for the server
 		/// </summary>
-		IOptionsMonitor<ServerSettings> Settings;
+		readonly IOptionsMonitor<ServerSettings> _settings;
 
 		/// <summary>
 		/// Constructor
 		/// </summary>
-		public ServerController(IOptionsMonitor<ServerSettings> Settings)
+		public ServerController(IOptionsMonitor<ServerSettings> settings)
 		{
-			this.Settings = Settings;
+			_settings = settings;
 		}
 		
 		/// <summary>
@@ -54,8 +38,8 @@ namespace Horde.Build.Controllers
 		[Route("/api/v1/server/version")]
 		public ActionResult GetVersionAsync()
 		{
-			FileVersionInfo FileVersionInfo = FileVersionInfo.GetVersionInfo(Assembly.GetExecutingAssembly().Location);
-			return Ok(FileVersionInfo.ProductVersion);
+			FileVersionInfo fileVersionInfo = FileVersionInfo.GetVersionInfo(Assembly.GetExecutingAssembly().Location);
+			return Ok(fileVersionInfo.ProductVersion);
 		}		
 
 		/// <summary>
@@ -66,7 +50,7 @@ namespace Horde.Build.Controllers
 		[ProducesResponseType(typeof(GetServerInfoResponse), 200)]
 		public ActionResult<GetServerInfoResponse> GetServerInfo()
 		{
-			return new GetServerInfoResponse(this.Settings.CurrentValue.SingleInstance);
+			return new GetServerInfoResponse(_settings.CurrentValue.SingleInstance);
 		}
 	}
 }
