Horde: Fix coding conventions within Horde.Build.

diff --git a/Engine/Source/Programs/Horde/Horde.Build/Controllers/StreamsController.cs b/Engine/Source/Programs/Horde/Horde.Build/Controllers/StreamsController.cs
--- a/Engine/Source/Programs/Horde/Horde.Build/Controllers/StreamsController.cs
+++ b/Engine/Source/Programs/Horde/Horde.Build/Controllers/StreamsController.cs
@@ -1,5 +1,9 @@
-﻿// Copyright Epic Games, Inc. All Rights Reserved.
+// Copyright Epic Games, Inc. All Rights Reserved.
 
+using System;
+using System.Collections.Generic;
+using System.Linq;
+using System.Threading.Tasks;
 using Horde.Build.Acls;
 using Horde.Build.Api;
 using Horde.Build.Collections;
@@ -8,12 +12,6 @@ using Horde.Build.Services;
 using Horde.Build.Utilities;
 using Microsoft.AspNetCore.Authorization;
 using Microsoft.AspNetCore.Mvc;
-using System;
-using System.Collections.Generic;
-using System.Diagnostics.CodeAnalysis;
-using System.IO;
-using System.Linq;
-using System.Threading.Tasks;
 
 namespace Horde.Build.Controllers
 {
@@ -29,277 +27,269 @@ namespace Horde.Build.Controllers
 	[Route("[controller]")]
 	public class StreamsController : HordeControllerBase
 	{
-		private readonly AclService AclService;
-		private readonly ProjectService ProjectService;
-		private readonly StreamService StreamService;
-		private readonly ITemplateCollection TemplateCollection;
-		private readonly JobService JobService;
-		private readonly IJobStepRefCollection JobStepRefCollection;
-		private readonly IPerforceService PerforceService;
-		private readonly IUserCollection UserCollection;
+		private readonly StreamService _streamService;
+		private readonly ITemplateCollection _templateCollection;
+		private readonly IJobStepRefCollection _jobStepRefCollection;
+		private readonly IPerforceService _perforceService;
 
 		/// <summary>
 		/// Constructor
 		/// </summary>
-		public StreamsController(AclService AclService, ProjectService ProjectService, StreamService StreamService, ITemplateCollection TemplateCollection, JobService JobService, IJobStepRefCollection JobStepRefCollection, IPerforceService PerforceService, IUserCollection UserCollection)
+		public StreamsController(StreamService streamService, ITemplateCollection templateCollection, JobService jobService, IJobStepRefCollection jobStepRefCollection, IPerforceService perforceService)
 		{
-			this.AclService = AclService;
-			this.ProjectService = ProjectService;
-			this.StreamService = StreamService;
-			this.TemplateCollection = TemplateCollection;
-			this.JobService = JobService;
-			this.JobStepRefCollection = JobStepRefCollection;
-			this.PerforceService = PerforceService;
-			this.UserCollection = UserCollection;
+			_streamService = streamService;
+			_templateCollection = templateCollection;
+			_jobStepRefCollection = jobStepRefCollection;
+			_perforceService = perforceService;
 		}
 
 		/// <summary>
 		/// Query all the streams for a particular project.
 		/// </summary>
-		/// <param name="ProjectIds">Unique id of the project to query</param>
-		/// <param name="Filter">Filter for the properties to return</param>
+		/// <param name="projectIds">Unique id of the project to query</param>
+		/// <param name="filter">Filter for the properties to return</param>
 		/// <returns>Information about all the projects</returns>
 		[HttpGet]
 		[Route("/api/v1/streams")]
 		[ProducesResponseType(typeof(List<GetStreamResponse>), 200)]
-		public async Task<ActionResult<List<object>>> GetStreamsAsync([FromQuery(Name = "ProjectId")] string[] ProjectIds, [FromQuery] PropertyFilter? Filter = null)
+		public async Task<ActionResult<List<object>>> GetStreamsAsync([FromQuery(Name = "ProjectId")] string[] projectIds, [FromQuery] PropertyFilter? filter = null)
 		{
-			ProjectId[] ProjectIdValues = Array.ConvertAll(ProjectIds, x => new ProjectId(x));
+			ProjectId[] projectIdValues = Array.ConvertAll(projectIds, x => new ProjectId(x));
 
-			List<IStream> Streams = await StreamService.GetStreamsAsync(ProjectIdValues);
-			ProjectPermissionsCache PermissionsCache = new ProjectPermissionsCache();
+			List<IStream> streams = await _streamService.GetStreamsAsync(projectIdValues);
+			ProjectPermissionsCache permissionsCache = new ProjectPermissionsCache();
 
-			List<GetStreamResponse> Responses = new List<GetStreamResponse>();
-			foreach (IStream Stream in Streams)
+			List<GetStreamResponse> responses = new List<GetStreamResponse>();
+			foreach (IStream stream in streams)
 			{
-				if (await StreamService.AuthorizeAsync(Stream, AclAction.ViewStream, User, PermissionsCache))
+				if (await _streamService.AuthorizeAsync(stream, AclAction.ViewStream, User, permissionsCache))
 				{
-					GetStreamResponse Response = await CreateGetStreamResponse(Stream, PermissionsCache);
-					Responses.Add(Response);
+					GetStreamResponse response = await CreateGetStreamResponse(stream, permissionsCache);
+					responses.Add(response);
 				}
 			}
-			return Responses.OrderBy(x => x.Order).ThenBy(x => x.Name).Select(x => PropertyFilter.Apply(x, Filter)).ToList();
+			return responses.OrderBy(x => x.Order).ThenBy(x => x.Name).Select(x => PropertyFilter.Apply(x, filter)).ToList();
 		}
 
 		/// <summary>
 		/// Retrieve information about a specific stream.
 		/// </summary>
-		/// <param name="StreamId">Id of the stream to get information about</param>
-		/// <param name="Filter">Filter for the properties to return</param>
+		/// <param name="streamId">Id of the stream to get information about</param>
+		/// <param name="filter">Filter for the properties to return</param>
 		/// <returns>Information about the requested project</returns>
 		[HttpGet]
-		[Route("/api/v1/streams/{StreamId}")]
+		[Route("/api/v1/streams/{streamId}")]
 		[ProducesResponseType(typeof(GetStreamResponse), 200)]
-		public async Task<ActionResult<object>> GetStreamAsync(StreamId StreamId, [FromQuery] PropertyFilter? Filter = null)
+		public async Task<ActionResult<object>> GetStreamAsync(StreamId streamId, [FromQuery] PropertyFilter? filter = null)
 		{
-			IStream? Stream = await StreamService.GetStreamAsync(StreamId);
-			if (Stream == null)
+			IStream? stream = await _streamService.GetStreamAsync(streamId);
+			if (stream == null)
 			{
-				return NotFound(StreamId);
+				return NotFound(streamId);
 			}
 
-			ProjectPermissionsCache PermissionsCache = new ProjectPermissionsCache();
-			if (!await StreamService.AuthorizeAsync(Stream, AclAction.ViewStream, User, PermissionsCache))
+			ProjectPermissionsCache permissionsCache = new ProjectPermissionsCache();
+			if (!await _streamService.AuthorizeAsync(stream, AclAction.ViewStream, User, permissionsCache))
 			{
-				return Forbid(AclAction.ViewStream, StreamId);
+				return Forbid(AclAction.ViewStream, streamId);
 			}
 
-			return PropertyFilter.Apply(await CreateGetStreamResponse(Stream, PermissionsCache), Filter);
+			return PropertyFilter.Apply(await CreateGetStreamResponse(stream, permissionsCache), filter);
 		}
 
 		/// <summary>
 		/// Create a stream response object, including all the templates
 		/// </summary>
-		/// <param name="Stream">Stream to create response for</param>
-		/// <param name="Cache">Permissions cache</param>
+		/// <param name="stream">Stream to create response for</param>
+		/// <param name="cache">Permissions cache</param>
 		/// <returns>Response object</returns>
-		async Task<GetStreamResponse> CreateGetStreamResponse(IStream Stream, ProjectPermissionsCache Cache)
+		async Task<GetStreamResponse> CreateGetStreamResponse(IStream stream, ProjectPermissionsCache cache)
 		{
-			bool bIncludeAcl = Stream.Acl != null && await StreamService.AuthorizeAsync(Stream, AclAction.ViewPermissions, User, Cache);
+			bool bIncludeAcl = stream.Acl != null && await _streamService.AuthorizeAsync(stream, AclAction.ViewPermissions, User, cache);
 
-			List<GetTemplateRefResponse> ApiTemplateRefs = new List<GetTemplateRefResponse>();
-			foreach (KeyValuePair<TemplateRefId, TemplateRef> Pair in Stream.Templates)
+			List<GetTemplateRefResponse> apiTemplateRefs = new List<GetTemplateRefResponse>();
+			foreach (KeyValuePair<TemplateRefId, TemplateRef> pair in stream.Templates)
 			{
-				if (await StreamService.AuthorizeAsync(Stream, Pair.Value, AclAction.ViewTemplate, User, Cache))
+				if (await _streamService.AuthorizeAsync(stream, pair.Value, AclAction.ViewTemplate, User, cache))
 				{
-					ITemplate? Template = await TemplateCollection.GetAsync(Pair.Value.Hash);
-					if (Template != null)
+					ITemplate? template = await _templateCollection.GetAsync(pair.Value.Hash);
+					if (template != null)
 					{
-						bool bIncludeTemplateAcl = Pair.Value.Acl != null && await StreamService.AuthorizeAsync(Stream, Pair.Value, AclAction.ViewPermissions, User, Cache);
-						ApiTemplateRefs.Add(new GetTemplateRefResponse(Pair.Key, Pair.Value, Template, bIncludeTemplateAcl));
+						bool bIncludeTemplateAcl = pair.Value.Acl != null && await _streamService.AuthorizeAsync(stream, pair.Value, AclAction.ViewPermissions, User, cache);
+						apiTemplateRefs.Add(new GetTemplateRefResponse(pair.Key, pair.Value, template, bIncludeTemplateAcl));
 					}
 				}
 			}
 
-			return Stream.ToApiResponse(bIncludeAcl, ApiTemplateRefs);
+			return stream.ToApiResponse(bIncludeAcl, apiTemplateRefs);
 		}
 
 		/// <summary>
 		/// Gets a list of changes for a stream
 		/// </summary>
-		/// <param name="StreamId">The stream id</param>
-		/// <param name="Min">The starting changelist number</param>
-		/// <param name="Max">The ending changelist number</param>
-		/// <param name="Results">Number of results to return</param>
-		/// <param name="Filter">The filter to apply to the results</param>
+		/// <param name="streamId">The stream id</param>
+		/// <param name="min">The starting changelist number</param>
+		/// <param name="max">The ending changelist number</param>
+		/// <param name="results">Number of results to return</param>
+		/// <param name="filter">The filter to apply to the results</param>
 		/// <returns>Http result code</returns>
 		[HttpGet]
-		[Route("/api/v1/streams/{StreamId}/changes")]
+		[Route("/api/v1/streams/{streamId}/changes")]
 		[ProducesResponseType(typeof(List<GetChangeSummaryResponse>), 200)]
-		public async Task<ActionResult<List<object>>> GetChangesAsync(StreamId StreamId, [FromQuery] int? Min = null, [FromQuery] int? Max = null, [FromQuery] int Results = 50, PropertyFilter? Filter = null)
+		public async Task<ActionResult<List<object>>> GetChangesAsync(StreamId streamId, [FromQuery] int? min = null, [FromQuery] int? max = null, [FromQuery] int results = 50, PropertyFilter? filter = null)
 		{
-			IStream? Stream = await StreamService.GetStreamAsync(StreamId);
-			if (Stream == null)
+			IStream? stream = await _streamService.GetStreamAsync(streamId);
+			if (stream == null)
 			{
-				return NotFound(StreamId);
+				return NotFound(streamId);
 			}
-			if (!await StreamService.AuthorizeAsync(Stream, AclAction.ViewChanges, User, null))
+			if (!await _streamService.AuthorizeAsync(stream, AclAction.ViewChanges, User, null))
 			{
-				return Forbid(AclAction.ViewChanges, StreamId);
+				return Forbid(AclAction.ViewChanges, streamId);
 			}
 
-			string? PerforceUser = User.GetPerforceUser();
-			if(PerforceUser == null)
+			string? perforceUser = User.GetPerforceUser();
+			if(perforceUser == null)
 			{
 				return BadRequest("Current user does not have an associated Perforce user");
 			}
 
-			List<ChangeSummary> Commits = await PerforceService.GetChangesAsync(Stream.ClusterName, Stream.Name, Min, Max, Results, PerforceUser);
+			List<ChangeSummary> commits = await _perforceService.GetChangesAsync(stream.ClusterName, stream.Name, min, max, results, perforceUser);
 
-			List<GetChangeSummaryResponse> Responses = new List<GetChangeSummaryResponse>();
-			foreach (ChangeSummary Commit in Commits)
+			List<GetChangeSummaryResponse> responses = new List<GetChangeSummaryResponse>();
+			foreach (ChangeSummary commit in commits)
 			{
-				Responses.Add(new GetChangeSummaryResponse(Commit));
+				responses.Add(new GetChangeSummaryResponse(commit));
 			}
-			return Responses.ConvertAll(x => PropertyFilter.Apply(x, Filter));
+			return responses.ConvertAll(x => PropertyFilter.Apply(x, filter));
 		}
 
 		/// <summary>
 		/// Gets a list of changes for a stream
 		/// </summary>
-		/// <param name="StreamId">The stream id</param>
-		/// <param name="ChangeNumber">The changelist number</param>
-		/// <param name="Filter">The filter to apply to the results</param>
+		/// <param name="streamId">The stream id</param>
+		/// <param name="changeNumber">The changelist number</param>
+		/// <param name="filter">The filter to apply to the results</param>
 		/// <returns>Http result code</returns>
 		[HttpGet]
-		[Route("/api/v1/streams/{StreamId}/changes/{ChangeNumber}")]
+		[Route("/api/v1/streams/{streamId}/changes/{changeNumber}")]
 		[ProducesResponseType(typeof(GetChangeDetailsResponse), 200)]
-		public async Task<ActionResult<object>> GetChangeDetailsAsync(StreamId StreamId, int ChangeNumber, PropertyFilter? Filter = null)
+		public async Task<ActionResult<object>> GetChangeDetailsAsync(StreamId streamId, int changeNumber, PropertyFilter? filter = null)
 		{
-			IStream? Stream = await StreamService.GetStreamAsync(StreamId);
-			if (Stream == null)
+			IStream? stream = await _streamService.GetStreamAsync(streamId);
+			if (stream == null)
 			{
-				return NotFound(StreamId);
+				return NotFound(streamId);
 			}
-			if (!await StreamService.AuthorizeAsync(Stream, AclAction.ViewChanges, User, null))
+			if (!await _streamService.AuthorizeAsync(stream, AclAction.ViewChanges, User, null))
 			{
-				return Forbid(AclAction.ViewChanges, StreamId);
+				return Forbid(AclAction.ViewChanges, streamId);
 			}
 
-			string? PerforceUser = User.GetPerforceUser();
-			if(PerforceUser == null)
+			string? perforceUser = User.GetPerforceUser();
+			if(perforceUser == null)
 			{
 				return BadRequest("Current user does not have an associated Perforce user");
 			}
 
-			ChangeDetails? ChangeDetails = await PerforceService.GetChangeDetailsAsync(Stream.ClusterName, Stream.Name, ChangeNumber, PerforceUser);
-			if(ChangeDetails == null)
+			ChangeDetails? changeDetails = await _perforceService.GetChangeDetailsAsync(stream.ClusterName, stream.Name, changeNumber, perforceUser);
+			if(changeDetails == null)
 			{
-				return NotFound("CL {Change} not found in stream {StreamId}", ChangeNumber, StreamId);
+				return NotFound("CL {Change} not found in stream {StreamId}", changeNumber, streamId);
 			}
 
-			return PropertyFilter.Apply(new GetChangeDetailsResponse(ChangeDetails), Filter);
+			return PropertyFilter.Apply(new GetChangeDetailsResponse(changeDetails), filter);
 		}
 
 		/// <summary>
 		/// Gets the history of a step in the stream
 		/// </summary>
-		/// <param name="StreamId">The stream id</param>
-		/// <param name="TemplateId"></param>
-		/// <param name="Step">Name of the step to search for</param>
-		/// <param name="Change">Maximum changelist number to return</param>
-		/// <param name="Count">Number of results to return</param>
-		/// <param name="Filter">The filter to apply to the results</param>
+		/// <param name="streamId">The stream id</param>
+		/// <param name="templateId"></param>
+		/// <param name="step">Name of the step to search for</param>
+		/// <param name="change">Maximum changelist number to return</param>
+		/// <param name="count">Number of results to return</param>
+		/// <param name="filter">The filter to apply to the results</param>
 		/// <returns>Http result code</returns>
 		[HttpGet]
-		[Route("/api/v1/streams/{StreamId}/history")]
+		[Route("/api/v1/streams/{streamId}/history")]
 		[ProducesResponseType(typeof(List<GetJobStepRefResponse>), 200)]
-		public async Task<ActionResult<List<object>>> GetStepHistoryAsync(StreamId StreamId, [FromQuery] string TemplateId, [FromQuery] string Step, [FromQuery] int? Change = null, [FromQuery] int Count = 10, [FromQuery] PropertyFilter? Filter = null)
+		public async Task<ActionResult<List<object>>> GetStepHistoryAsync(StreamId streamId, [FromQuery] string templateId, [FromQuery] string step, [FromQuery] int? change = null, [FromQuery] int count = 10, [FromQuery] PropertyFilter? filter = null)
 		{
-			IStream? Stream = await StreamService.GetStreamAsync(StreamId);
-			if (Stream == null)
+			IStream? stream = await _streamService.GetStreamAsync(streamId);
+			if (stream == null)
 			{
-				return NotFound(StreamId);
+				return NotFound(streamId);
 			}
-			if (!await StreamService.AuthorizeAsync(Stream, AclAction.ViewJob, User, null))
+			if (!await _streamService.AuthorizeAsync(stream, AclAction.ViewJob, User, null))
 			{
-				return Forbid(AclAction.ViewJob, StreamId);
+				return Forbid(AclAction.ViewJob, streamId);
 			}
 
-			TemplateRefId TemplateIdValue = new TemplateRefId(TemplateId);
+			TemplateRefId templateIdValue = new TemplateRefId(templateId);
 
-			List<IJobStepRef> Steps = await JobStepRefCollection.GetStepsForNodeAsync(StreamId, TemplateIdValue, Step, Change, true, Count);
-			return Steps.ConvertAll(x => PropertyFilter.Apply(new GetJobStepRefResponse(x), Filter));
+			List<IJobStepRef> steps = await _jobStepRefCollection.GetStepsForNodeAsync(streamId, templateIdValue, step, change, true, count);
+			return steps.ConvertAll(x => PropertyFilter.Apply(new GetJobStepRefResponse(x), filter));
 		}
 
 		/// <summary>
 		/// Deletes a stream
 		/// </summary>
-		/// <param name="StreamId">Id of the stream to update.</param>
+		/// <param name="streamId">Id of the stream to update.</param>
 		/// <returns>Http result code</returns>
 		[HttpDelete]
-		[Route("/api/v1/streams/{StreamId}")]
-		public async Task<ActionResult> DeleteStreamAsync(StreamId StreamId)
+		[Route("/api/v1/streams/{streamId}")]
+		public async Task<ActionResult> DeleteStreamAsync(StreamId streamId)
 		{
-			IStream? Stream = await StreamService.GetStreamAsync(StreamId);
-			if (Stream == null)
+			IStream? stream = await _streamService.GetStreamAsync(streamId);
+			if (stream == null)
 			{
-				return NotFound(StreamId);
+				return NotFound(streamId);
 			}
-			if (!await StreamService.AuthorizeAsync(Stream, AclAction.DeleteStream, User, null))
+			if (!await _streamService.AuthorizeAsync(stream, AclAction.DeleteStream, User, null))
 			{
-				return Forbid(AclAction.DeleteStream, StreamId);
+				return Forbid(AclAction.DeleteStream, streamId);
 			}
 
-			await StreamService.DeleteStreamAsync(StreamId);
+			await _streamService.DeleteStreamAsync(streamId);
 			return new OkResult();
 		}
 
 		/// <summary>
 		/// Gets all the templates for a stream
 		/// </summary>
-		/// <param name="StreamId">Unique id of the stream to query</param>
-		/// <param name="TemplateId">Unique id of the template to query</param>
-		/// <param name="Filter">Filter for properties to return</param>
+		/// <param name="streamId">Unique id of the stream to query</param>
+		/// <param name="templateId">Unique id of the template to query</param>
+		/// <param name="filter">Filter for properties to return</param>
 		/// <returns>Information about all the templates</returns>
 		[HttpGet]
-		[Route("/api/v1/streams/{StreamId}/templates/{TemplateId}")]
+		[Route("/api/v1/streams/{streamId}/templates/{templateId}")]
 		[ProducesResponseType(typeof(List<GetTemplateResponse>), 200)]
-		public async Task<ActionResult<object>> GetTemplateAsync(StreamId StreamId, TemplateRefId TemplateId, [FromQuery] PropertyFilter? Filter = null)
+		public async Task<ActionResult<object>> GetTemplateAsync(StreamId streamId, TemplateRefId templateId, [FromQuery] PropertyFilter? filter = null)
 		{
-			IStream? Stream = await StreamService.GetStreamAsync(StreamId);
-			if (Stream == null)
+			IStream? stream = await _streamService.GetStreamAsync(streamId);
+			if (stream == null)
 			{
-				return NotFound(StreamId);
+				return NotFound(streamId);
 			}
 
-			TemplateRef? TemplateRef;
-			if (!Stream.Templates.TryGetValue(TemplateId, out TemplateRef))
+			TemplateRef? templateRef;
+			if (!stream.Templates.TryGetValue(templateId, out templateRef))
 			{
-				return NotFound(StreamId, TemplateId);
+				return NotFound(streamId, templateId);
 			}
-			if (!await StreamService.AuthorizeAsync(Stream, TemplateRef, AclAction.ViewTemplate, User, null))
+			if (!await _streamService.AuthorizeAsync(stream, templateRef, AclAction.ViewTemplate, User, null))
 			{
-				return Forbid(AclAction.ViewTemplate, StreamId);
+				return Forbid(AclAction.ViewTemplate, streamId);
 			}
 
-			ITemplate? Template = await TemplateCollection.GetAsync(TemplateRef.Hash);
-			if(Template == null)
+			ITemplate? template = await _templateCollection.GetAsync(templateRef.Hash);
+			if(template == null)
 			{
 				return NotFound();
 			}
 
-			return new GetTemplateResponse(Template).ApplyFilter(Filter);
+			return new GetTemplateResponse(template).ApplyFilter(filter);
 		}
 	}
 }
diff --git a/Engine/Source/Programs/Horde/Horde.Build/Controllers/SubscriptionsController.cs b/Engine/Source/Programs/Horde/Horde.Build/Controllers/SubscriptionsController.cs
--- a/Engine/Source/Programs/Horde/Horde.Build/Controllers/SubscriptionsController.cs
+++ b/Engine/Source/Programs/Horde/Horde.Build/Controllers/SubscriptionsController.cs
@@ -1,19 +1,15 @@
-﻿// Copyright Epic Games, Inc. All Rights Reserved.
+// Copyright Epic Games, Inc. All Rights Reserved.
 
+using System;
+using System.Collections.Generic;
+using System.Threading.Tasks;
 using Horde.Build.Acls;
 using Horde.Build.Api;
 using Horde.Build.Collections;
 using Horde.Build.Models;
-using Horde.Build.Services;
 using Horde.Build.Utilities;
 using Microsoft.AspNetCore.Authorization;
 using Microsoft.AspNetCore.Mvc;
-using MongoDB.Bson;
-using MongoDB.Driver;
-using System;
-using System.Collections.Generic;
-using System.Linq;
-using System.Threading.Tasks;
 
 namespace Horde.Build.Controllers
 {
@@ -30,163 +26,162 @@ namespace Horde.Build.Controllers
 		/// <summary>
 		/// The ACL service singleton
 		/// </summary>
-		AclService AclService;
+		readonly AclService _aclService;
 
 		/// <summary>
 		/// Collection of subscription documents
 		/// </summary>
-		ISubscriptionCollection SubscriptionCollection;
+		readonly ISubscriptionCollection _subscriptionCollection;
 
 		/// <summary>
 		/// Constructor
 		/// </summary>
-		/// <param name="AclService">The acl service singleton</param>
-		/// <param name="SubscriptionCollection">The collection of subscription documents</param>
-		public SubscriptionsController(AclService AclService, ISubscriptionCollection SubscriptionCollection)
+		/// <param name="aclService">The acl service singleton</param>
+		/// <param name="subscriptionCollection">The collection of subscription documents</param>
+		public SubscriptionsController(AclService aclService, ISubscriptionCollection subscriptionCollection)
 		{
-			this.AclService = AclService;
-			this.SubscriptionCollection = SubscriptionCollection;
+			_aclService = aclService;
+			_subscriptionCollection = subscriptionCollection;
 		}
 
 		/// <summary>
 		/// Find subscriptions matching a criteria
 		/// </summary>
-		/// <param name="UserId">Name of the user</param>
-		/// <param name="Filter">Filter for properties to return</param>
+		/// <param name="userId">Name of the user</param>
+		/// <param name="filter">Filter for properties to return</param>
 		/// <returns>List of subscriptions</returns>
 		[HttpGet]
 		[Route("/api/v1/subscriptions")]
 		[ProducesResponseType(typeof(List<GetSubscriptionResponse>), 200)]
-		public async Task<ActionResult<List<object>>> GetSubscriptionsAsync([FromQuery] string UserId, [FromQuery] PropertyFilter? Filter = null)
+		public async Task<ActionResult<List<object>>> GetSubscriptionsAsync([FromQuery] string userId, [FromQuery] PropertyFilter? filter = null)
 		{
-			UserId UserIdValue;
-			if (!TryParseUserId(UserId, out UserIdValue))
+			UserId userIdValue;
+			if (!TryParseUserId(userId, out userIdValue))
 			{
 				return BadRequest("Invalid user id");
 			}
-			if (!await AclService.AuthorizeAsUserAsync(User, UserIdValue))
+			if (!await _aclService.AuthorizeAsUserAsync(User, userIdValue))
 			{
 				return Forbid();
 			}
 
-			List<ISubscription> Results = await SubscriptionCollection.FindSubscriptionsAsync(UserIdValue);
-			return Results.ConvertAll(x => PropertyFilter.Apply(new GetSubscriptionResponse(x), Filter));
+			List<ISubscription> results = await _subscriptionCollection.FindSubscriptionsAsync(userIdValue);
+			return results.ConvertAll(x => PropertyFilter.Apply(new GetSubscriptionResponse(x), filter));
 		}
 
 		/// <summary>
 		/// Find subscriptions matching a criteria
 		/// </summary>
-		/// <param name="SubscriptionId">The subscription id</param>
-		/// <param name="Filter">Filter for properties to return</param>
+		/// <param name="subscriptionId">The subscription id</param>
+		/// <param name="filter">Filter for properties to return</param>
 		/// <returns>List of subscriptions</returns>
 		[HttpGet]
-		[Route("/api/v1/subscriptions/{SubscriptionId}")]
+		[Route("/api/v1/subscriptions/{subscriptionId}")]
 		[ProducesResponseType(typeof(GetSubscriptionResponse), 200)]
-		public async Task<ActionResult<object>> GetSubscriptionAsync(string SubscriptionId, [FromQuery] PropertyFilter? Filter = null)
+		public async Task<ActionResult<object>> GetSubscriptionAsync(string subscriptionId, [FromQuery] PropertyFilter? filter = null)
 		{
-			ISubscription? Subscription = await SubscriptionCollection.GetAsync(SubscriptionId);
-			if (Subscription == null)
+			ISubscription? subscription = await _subscriptionCollection.GetAsync(subscriptionId);
+			if (subscription == null)
 			{
 				return NotFound();
 			}
-			if (!await AclService.AuthorizeAsUserAsync(User, Subscription.UserId))
+			if (!await _aclService.AuthorizeAsUserAsync(User, subscription.UserId))
 			{
 				return Forbid();
 			}
 
-			return PropertyFilter.Apply(new GetSubscriptionResponse(Subscription), Filter);
+			return PropertyFilter.Apply(new GetSubscriptionResponse(subscription), filter);
 		}
 
-
 		/// <summary>
 		/// Remove a subscription
 		/// </summary>
-		/// <param name="SubscriptionId">The subscription id</param>
+		/// <param name="subscriptionId">The subscription id</param>
 		/// <returns>Async task</returns>
 		[HttpDelete]
-		[Route("/api/v1/subscriptions/{SubscriptionId}")]
+		[Route("/api/v1/subscriptions/{subscriptionId}")]
 		[ProducesResponseType(typeof(List<GetSubscriptionResponse>), 200)]
-		public async Task<ActionResult> DeleteSubscriptionAsync(string SubscriptionId)
+		public async Task<ActionResult> DeleteSubscriptionAsync(string subscriptionId)
 		{
-			ISubscription? Subscription = await SubscriptionCollection.GetAsync(SubscriptionId);
-			if (Subscription == null)
+			ISubscription? subscription = await _subscriptionCollection.GetAsync(subscriptionId);
+			if (subscription == null)
 			{
 				return NotFound();
 			}
-			if (!await AclService.AuthorizeAsUserAsync(User, Subscription.UserId))
+			if (!await _aclService.AuthorizeAsUserAsync(User, subscription.UserId))
 			{
 				return Forbid();
 			}
 
-			await SubscriptionCollection.RemoveAsync(new[] { Subscription });
+			await _subscriptionCollection.RemoveAsync(new[] { subscription });
 			return Ok();
 		}
 
 		/// <summary>
 		/// Find subscriptions matching a criteria
 		/// </summary>
-		/// <param name="Subscriptions">The new subscriptions to create</param>
+		/// <param name="subscriptions">The new subscriptions to create</param>
 		/// <returns>List of subscriptions</returns>
 		[HttpPost]
 		[Route("/api/v1/subscriptions")]
-		public async Task<ActionResult<List<CreateSubscriptionResponse>>> CreateSubscriptionsAsync(List<CreateSubscriptionRequest> Subscriptions)
+		public async Task<ActionResult<List<CreateSubscriptionResponse>>> CreateSubscriptionsAsync(List<CreateSubscriptionRequest> subscriptions)
 		{
-			HashSet<UserId> AuthorizedUsers = new HashSet<UserId>();
+			HashSet<UserId> authorizedUsers = new HashSet<UserId>();
 
-			UserId? CurrentUserId = User.GetUserId();
-			if(CurrentUserId != null)
+			UserId? currentUserId = User.GetUserId();
+			if(currentUserId != null)
 			{
-				AuthorizedUsers.Add(CurrentUserId.Value);
+				authorizedUsers.Add(currentUserId.Value);
 			}
 
-			GlobalPermissionsCache Cache = new GlobalPermissionsCache();
+			GlobalPermissionsCache cache = new GlobalPermissionsCache();
 
-			List<NewSubscription> NewSubscriptions = new List<NewSubscription>();
-			foreach (CreateSubscriptionRequest Subscription in Subscriptions)
+			List<NewSubscription> newSubscriptions = new List<NewSubscription>();
+			foreach (CreateSubscriptionRequest subscription in subscriptions)
 			{
-				UserId NewUserId;
-				if (!TryParseUserId(Subscription.UserId, out NewUserId))
+				UserId newUserId;
+				if (!TryParseUserId(subscription.UserId, out newUserId))
 				{
-					return BadRequest($"Invalid user id: '{Subscription.UserId}'.");
+					return BadRequest($"Invalid user id: '{subscription.UserId}'.");
 				}
-				if (AuthorizedUsers.Add(NewUserId) && !await AclService.AuthorizeAsync(AclAction.Impersonate, User, Cache))
+				if (authorizedUsers.Add(newUserId) && !await _aclService.AuthorizeAsync(AclAction.Impersonate, User, cache))
 				{
 					return Forbid();
 				}
-				NewSubscriptions.Add(new NewSubscription(Subscription.Event, NewUserId, Subscription.NotificationType));
+				newSubscriptions.Add(new NewSubscription(subscription.Event, newUserId, subscription.NotificationType));
 			}
 
-			List<ISubscription> Results = await SubscriptionCollection.AddAsync(NewSubscriptions);
-			return Results.ConvertAll(x => new CreateSubscriptionResponse(x));
+			List<ISubscription> results = await _subscriptionCollection.AddAsync(newSubscriptions);
+			return results.ConvertAll(x => new CreateSubscriptionResponse(x));
 		}
 
 		/// <summary>
 		/// Parse a user id from a string. Allows passing the user's name as well as their objectid value.
 		/// </summary>
-		/// <param name="UserName"></param>
-		/// <param name="ObjectId"></param>
+		/// <param name="userName"></param>
+		/// <param name="objectId"></param>
 		/// <returns></returns>
-		bool TryParseUserId(string UserName, out UserId ObjectId)
+		bool TryParseUserId(string userName, out UserId objectId)
 		{
-			UserId NewObjectId;
-			if (UserId.TryParse(UserName, out NewObjectId))
+			UserId newObjectId;
+			if (UserId.TryParse(userName, out newObjectId))
 			{
-				ObjectId = NewObjectId;
+				objectId = newObjectId;
 				return true;
 			}
 
-			string? CurrentUserName = User.GetUserName();
-			if (CurrentUserName != null && String.Equals(UserName, CurrentUserName, StringComparison.OrdinalIgnoreCase))
+			string? currentUserName = User.GetUserName();
+			if (currentUserName != null && String.Equals(userName, currentUserName, StringComparison.OrdinalIgnoreCase))
 			{
-				UserId? CurrentUserId = User.GetUserId();
-				if (CurrentUserId != null)
+				UserId? currentUserId = User.GetUserId();
+				if (currentUserId != null)
 				{
-					ObjectId = CurrentUserId.Value;
+					objectId = currentUserId.Value;
 					return true;
 				}
 			}
 
-			ObjectId = default;
+			objectId = default;
 			return false;
 		}
 	}
diff --git a/Engine/Source/Programs/Horde/Horde.Build/Controllers/SwarmController.cs b/Engine/Source/Programs/Horde/Horde.Build/Controllers/SwarmController.cs
--- a/Engine/Source/Programs/Horde/Horde.Build/Controllers/SwarmController.cs
+++ b/Engine/Source/Programs/Horde/Horde.Build/Controllers/SwarmController.cs
@@ -1,12 +1,7 @@
 // Copyright Epic Games, Inc. All Rights Reserved.
 
-using Amazon.Runtime.Internal.Util;
 using Microsoft.AspNetCore.Mvc;
 using Microsoft.Extensions.Logging;
-using System;
-using System.Collections.Generic;
-using System.Linq;
-using System.Threading.Tasks;
 
 namespace Horde.Build.Controllers
 {
@@ -20,14 +15,14 @@ namespace Horde.Build.Controllers
 		/// <summary>
 		/// Logger instance
 		/// </summary>
-		ILogger<SwarmController> Logger;
+		readonly ILogger<SwarmController> _logger;
 
 		/// <summary>
 		/// Constructor
 		/// </summary>
-		public SwarmController(ILogger<SwarmController> Logger)
+		public SwarmController(ILogger<SwarmController> logger)
 		{
-			this.Logger = Logger;
+			_logger = logger;
 		}
 
 		/// <summary>
@@ -36,9 +31,9 @@ namespace Horde.Build.Controllers
 		/// <returns>Result code</returns>
 		[HttpPost]
 		[Route("/api/v1/swarm/reviews")]
-		public ActionResult AddReview(object Request)
+		public ActionResult AddReview(object request)
 		{
-			Logger.LogInformation("Added Swarm review: {Request}", Request.ToString());
+			_logger.LogInformation("Added Swarm review: {Request}", request.ToString());
 			return Ok();
 		}
 	}
diff --git a/Engine/Source/Programs/Horde/Horde.Build/Controllers/TestDataController.cs b/Engine/Source/Programs/Horde/Horde.Build/Controllers/TestDataController.cs
--- a/Engine/Source/Programs/Horde/Horde.Build/Controllers/TestDataController.cs
+++ b/Engine/Source/Programs/Horde/Horde.Build/Controllers/TestDataController.cs
@@ -1,23 +1,16 @@
 ﻿// Copyright Epic Games, Inc. All Rights Reserved.
 
+using System.Collections.Generic;
+using System.Threading.Tasks;
 using Horde.Build.Acls;
 using Horde.Build.Api;
 using Horde.Build.Collections;
-using HordeCommon;
 using Horde.Build.Models;
 using Horde.Build.Services;
 using Horde.Build.Utilities;
 using Microsoft.AspNetCore.Authorization;
 using Microsoft.AspNetCore.Mvc;
-using Microsoft.AspNetCore.Server.Kestrel.Core.Features;
 using MongoDB.Bson;
-using MongoDB.Bson.Serialization;
-using System;
-using System.Collections.Generic;
-using System.ComponentModel.DataAnnotations;
-using System.IO;
-using System.Linq;
-using System.Threading.Tasks;
 
 namespace Horde.Build.Controllers
 {
@@ -35,22 +28,22 @@ namespace Horde.Build.Controllers
 		/// <summary>
 		/// Collection of job documents
 		/// </summary>
-		private JobService JobService;
+		private readonly JobService _jobService;
 
 		/// <summary>
 		/// Collection of test data documents
 		/// </summary>
-		private readonly ITestDataCollection TestDataCollection;
+		private readonly ITestDataCollection _testDataCollection;
 
 		/// <summary>
 		/// Constructor
 		/// </summary>
-		/// <param name="JobService">The job service singleton</param>
-		/// <param name="TestDataCollection">Collection of test data documents</param>
-		public TestDataController(JobService JobService, ITestDataCollection TestDataCollection)
+		/// <param name="jobService">The job service singleton</param>
+		/// <param name="testDataCollection">Collection of test data documents</param>
+		public TestDataController(JobService jobService, ITestDataCollection testDataCollection)
 		{
-			this.JobService = JobService;
-			this.TestDataCollection = TestDataCollection;
+			_jobService = jobService;
+			_testDataCollection = testDataCollection;
 		}
 
 		/// <summary>
@@ -59,90 +52,90 @@ namespace Horde.Build.Controllers
 		/// <returns>The stream document</returns>
 		[HttpPost]
 		[Route("/api/v1/testdata")]
-		public async Task<ActionResult<CreateTestDataResponse>> CreateAsync(CreateTestDataRequest Request)
+		public async Task<ActionResult<CreateTestDataResponse>> CreateAsync(CreateTestDataRequest request)
 		{
-			IJob? Job = await JobService.GetJobAsync(new JobId(Request.JobId));
-			if (Job == null)
+			IJob? job = await _jobService.GetJobAsync(new JobId(request.JobId));
+			if (job == null)
 			{
 				return NotFound();
 			}
-			if (!await JobService.AuthorizeAsync(Job, AclAction.UpdateJob, User, null))
+			if (!await _jobService.AuthorizeAsync(job, AclAction.UpdateJob, User, null))
 			{
 				return Forbid();
 			}
 
-			IJobStep? JobStep;
-			if (!Job.TryGetStep(Request.StepId.ToSubResourceId(), out JobStep))
+			IJobStep? jobStep;
+			if (!job.TryGetStep(request.StepId.ToSubResourceId(), out jobStep))
 			{
 				return NotFound();
 			}
 
-			ITestData TestData = await TestDataCollection.AddAsync(Job, JobStep, Request.Key, new BsonDocument(Request.Data));
-			return new CreateTestDataResponse(TestData.Id.ToString());
+			ITestData testData = await _testDataCollection.AddAsync(job, jobStep, request.Key, new BsonDocument(request.Data));
+			return new CreateTestDataResponse(testData.Id.ToString());
 		}
 
 		/// <summary>
 		/// Searches for test data that matches a set of criteria
 		/// </summary>
-		/// <param name="StreamId">The stream id</param>
-		/// <param name="MinChange">The minimum changelist number to return (inclusive)</param>
-		/// <param name="MaxChange">The maximum changelist number to return (inclusive)</param>
-		/// <param name="JobId">The job id</param>
-		/// <param name="JobStepId">The unique step id</param>
-		/// <param name="Key">Key identifying the result to return</param>
-		/// <param name="Index">Offset within the results to return</param>
-		/// <param name="Count">Number of results to return</param>
-		/// <param name="Filter">Filter for properties to return</param>
+		/// <param name="streamId">The stream id</param>
+		/// <param name="minChange">The minimum changelist number to return (inclusive)</param>
+		/// <param name="maxChange">The maximum changelist number to return (inclusive)</param>
+		/// <param name="jobId">The job id</param>
+		/// <param name="jobStepId">The unique step id</param>
+		/// <param name="key">Key identifying the result to return</param>
+		/// <param name="index">Offset within the results to return</param>
+		/// <param name="count">Number of results to return</param>
+		/// <param name="filter">Filter for properties to return</param>
 		/// <returns>The stream document</returns>
 		[HttpGet]
 		[Route("/api/v1/testdata")]
 		[ProducesResponseType(typeof(List<GetTestDataResponse>), 200)]
-		public async Task<ActionResult<List<object>>> FindTestDataAsync([FromQuery] string? StreamId = null, [FromQuery] int? MinChange = null, [FromQuery] int? MaxChange = null, string? JobId = null, string? JobStepId = null, string? Key = null, int Index = 0, int Count = 10, PropertyFilter? Filter = null)
+		public async Task<ActionResult<List<object>>> FindTestDataAsync([FromQuery] string? streamId = null, [FromQuery] int? minChange = null, [FromQuery] int? maxChange = null, string? jobId = null, string? jobStepId = null, string? key = null, int index = 0, int count = 10, PropertyFilter? filter = null)
 		{
-			StreamId? StreamIdValue = null;
-			if(StreamId != null)
+			StreamId? streamIdValue = null;
+			if(streamId != null)
 			{
-				StreamIdValue = new StreamId(StreamId);
+				streamIdValue = new StreamId(streamId);
 			}
 
-			JobPermissionsCache Cache = new JobPermissionsCache();
+			JobPermissionsCache cache = new JobPermissionsCache();
 
-			List<object> Results = new List<object>();
+			List<object> results = new List<object>();
 
-			List<ITestData> Documents = await TestDataCollection.FindAsync(StreamIdValue, MinChange, MaxChange, JobId?.ToObjectId<IJob>(), JobStepId?.ToSubResourceId(), Key, Index, Count);
-			foreach (ITestData Document in Documents)
+			List<ITestData> documents = await _testDataCollection.FindAsync(streamIdValue, minChange, maxChange, jobId?.ToObjectId<IJob>(), jobStepId?.ToSubResourceId(), key, index, count);
+			foreach (ITestData document in documents)
 			{
-				if (await JobService.AuthorizeAsync(Document.JobId, AclAction.ViewJob, User, Cache))
+				if (await _jobService.AuthorizeAsync(document.JobId, AclAction.ViewJob, User, cache))
 				{
-					Results.Add(PropertyFilter.Apply(new GetTestDataResponse(Document), Filter));
+					results.Add(PropertyFilter.Apply(new GetTestDataResponse(document), filter));
 				}
 			}
 
-			return Results;
+			return results;
 		}
 
 		/// <summary>
 		/// Retrieve information about a specific issue
 		/// </summary>
-		/// <param name="TestDataId">Id of the document to get information about</param>
-		/// <param name="Filter">Filter for the properties to return</param>
+		/// <param name="testDataId">Id of the document to get information about</param>
+		/// <param name="filter">Filter for the properties to return</param>
 		/// <returns>List of matching agents</returns>
 		[HttpGet]
-		[Route("/api/v1/testdata/{TestDataId}")]
+		[Route("/api/v1/testdata/{testDataId}")]
 		[ProducesResponseType(typeof(GetTestDataResponse), 200)]
-		public async Task<ActionResult<object>> GetTestDataAsync(string TestDataId, [FromQuery] PropertyFilter? Filter = null)
+		public async Task<ActionResult<object>> GetTestDataAsync(string testDataId, [FromQuery] PropertyFilter? filter = null)
 		{
-			ITestData? TestData = await TestDataCollection.GetAsync(TestDataId.ToObjectId());
-			if (TestData == null)
+			ITestData? testData = await _testDataCollection.GetAsync(testDataId.ToObjectId());
+			if (testData == null)
 			{
 				return NotFound();
 			}
-			if (!await JobService.AuthorizeAsync(TestData.JobId, AclAction.ViewJob, User, null))
+			if (!await _jobService.AuthorizeAsync(testData.JobId, AclAction.ViewJob, User, null))
 			{
 				return Forbid();
 			}
 
-			return PropertyFilter.Apply(new GetTestDataResponse(TestData), Filter);
+			return PropertyFilter.Apply(new GetTestDataResponse(testData), filter);
 		}
 	}
 }
diff --git a/Engine/Source/Programs/Horde/Horde.Build/Controllers/UgsController.cs b/Engine/Source/Programs/Horde/Horde.Build/Controllers/UgsController.cs
--- a/Engine/Source/Programs/Horde/Horde.Build/Controllers/UgsController.cs
+++ b/Engine/Source/Programs/Horde/Horde.Build/Controllers/UgsController.cs
@@ -1,30 +1,22 @@
 // Copyright Epic Games, Inc. All Rights Reserved.
 
+using System;
+using System.Collections.Generic;
+using System.Linq;
+using System.Threading.Tasks;
 using Horde.Build.Api;
 using Horde.Build.Collections;
-using HordeCommon;
 using Horde.Build.Models;
 using Horde.Build.Services;
 using Horde.Build.Utilities;
-using Microsoft.AspNetCore.Authorization;
 using Microsoft.AspNetCore.Mvc;
-using Microsoft.AspNetCore.Server.Kestrel.Core.Features;
-using Microsoft.Extensions.Caching.Memory;
+using Microsoft.Extensions.Logging;
 using Microsoft.Extensions.Options;
 using MongoDB.Bson;
-using Serilog.Events;
-using System;
-using System.Collections.Generic;
-using System.ComponentModel.DataAnnotations;
-using System.IO;
-using System.Linq;
-using System.Threading.Tasks;
-using Microsoft.Extensions.Logging;
 
 namespace Horde.Build.Controllers
 {
 	using LogId = ObjectId<ILogFile>;
-	using StreamId = StringId<IStream>;
 	using UserId = ObjectId<IUser>;
 
 	/// <summary>
@@ -37,71 +29,53 @@ namespace Horde.Build.Controllers
 		/// <summary>
 		/// Singleton instance of the issue service
 		/// </summary>
-		private readonly IIssueService IssueService;
-
-		/// <summary>
-		/// Singleton instance of the job service
-		/// </summary>
-		private readonly JobService JobService;
-
-		/// <summary>
-		/// Collection of stream documents
-		/// </summary>
-		private readonly StreamService StreamService;
+		private readonly IIssueService _issueService;
 
 		/// <summary>
 		/// Collection of metadata documents
 		/// </summary>
-		private readonly IUgsMetadataCollection UgsMetadataCollection;
-
-		/// <summary>
-		/// Collection of log events
-		/// </summary>
-		private readonly ILogEventCollection LogEventCollection;
+		private readonly IUgsMetadataCollection _ugsMetadataCollection;
 
 		/// <summary>
 		/// Collection of users
 		/// </summary>
-		private readonly IUserCollection UserCollection;
+		private readonly IUserCollection _userCollection;
 
 		/// <summary>
 		/// The log file service
 		/// </summary>
-		private readonly ILogFileService LogFileService;
+		private readonly ILogFileService _logFileService;
 
 		/// <summary>
 		/// Server settings
 		/// </summary>
-		private readonly ServerSettings Settings;
+		private readonly ServerSettings _settings;
 
 		/// <summary>
 		/// Logger 
 		/// </summary>
-		ILogger Logger;
+		readonly ILogger _logger;
 
 		/// <summary>
 		/// Constructor
 		/// </summary>
-		/// <param name="IssueService">The issue service</param>
-		/// <param name="JobService">The job service</param>
-		/// <param name="StreamService">Collection of stream documents</param>
-		/// <param name="UgsMetadataCollection">Collection of UGS metadata documents</param>
-		/// <param name="LogEventCollection">Collection of log event documents</param>
-		/// <param name="UserCollection">Collection of user documents</param>
-		/// <param name="LogFileService">The log file service</param>
-		/// <param name="OptionsMonitor">The server settings</param>
-		/// <param name="Logger">Logger</param>
-		public UgsController(IIssueService IssueService, JobService JobService, StreamService StreamService, IUgsMetadataCollection UgsMetadataCollection, ILogEventCollection LogEventCollection, IUserCollection UserCollection, ILogFileService LogFileService, IOptionsMonitor<ServerSettings> OptionsMonitor, ILogger<UgsController> Logger)
+		/// <param name="issueService">The issue service</param>
+		/// <param name="jobService">The job service</param>
+		/// <param name="streamService">Collection of stream documents</param>
+		/// <param name="ugsMetadataCollection">Collection of UGS metadata documents</param>
+		/// <param name="logEventCollection">Collection of log event documents</param>
+		/// <param name="userCollection">Collection of user documents</param>
+		/// <param name="logFileService">The log file service</param>
+		/// <param name="optionsMonitor">The server settings</param>
+		/// <param name="logger">Logger</param>
+		public UgsController(IIssueService issueService, JobService jobService, StreamService streamService, IUgsMetadataCollection ugsMetadataCollection, ILogEventCollection logEventCollection, IUserCollection userCollection, ILogFileService logFileService, IOptionsMonitor<ServerSettings> optionsMonitor, ILogger<UgsController> logger)
 		{
-			this.IssueService = IssueService;
-			this.JobService = JobService;
-			this.StreamService = StreamService;
-			this.UgsMetadataCollection = UgsMetadataCollection;
-			this.LogEventCollection = LogEventCollection;
-			this.UserCollection = UserCollection;
-			this.LogFileService = LogFileService;
-			this.Settings = OptionsMonitor.CurrentValue;
-			this.Logger = Logger;
+			_issueService = issueService;
+			_ugsMetadataCollection = ugsMetadataCollection;
+			_userCollection = userCollection;
+			_logFileService = logFileService;
+			_settings = optionsMonitor.CurrentValue;
+			_logger = logger;
 		}
 
 		/// <summary>
@@ -118,26 +92,26 @@ namespace Horde.Build.Controllers
 		/// <summary>
 		/// Adds new metadata to the database
 		/// </summary>
-		/// <param name="Request">Request object</param>
+		/// <param name="request">Request object</param>
 		/// <returns>Result code</returns>
 		[HttpPost]
 		[Route("/ugs/api/metadata")]
-		public async Task<ActionResult> AddMetadataAsync(AddUgsMetadataRequest Request)
+		public async Task<ActionResult> AddMetadataAsync(AddUgsMetadataRequest request)
 		{
-			IUgsMetadata Metadata = await UgsMetadataCollection.FindOrAddAsync(Request.Stream, Request.Change, Request.Project);
-			if (Request.Synced != null || Request.Vote != null || Request.Investigating != null || Request.Starred != null || Request.Comment != null)
+			IUgsMetadata metadata = await _ugsMetadataCollection.FindOrAddAsync(request.Stream, request.Change, request.Project);
+			if (request.Synced != null || request.Vote != null || request.Investigating != null || request.Starred != null || request.Comment != null)
 			{
-				if (Request.UserName == null)
+				if (request.UserName == null)
 				{
 					return BadRequest("Missing UserName field on request body");
 				}
-				Metadata = await UgsMetadataCollection.UpdateUserAsync(Metadata, Request.UserName, Request.Synced, Request.Vote, Request.Investigating, Request.Starred, Request.Comment);
+				metadata = await _ugsMetadataCollection.UpdateUserAsync(metadata, request.UserName, request.Synced, request.Vote, request.Investigating, request.Starred, request.Comment);
 			}
-			if (Request.Badges != null)
+			if (request.Badges != null)
 			{
-				foreach (AddUgsBadgeRequest Badge in Request.Badges)
+				foreach (AddUgsBadgeRequest badge in request.Badges)
 				{
-					Metadata = await UgsMetadataCollection.UpdateBadgeAsync(Metadata, Badge.Name, Badge.Url, Badge.State);
+					metadata = await _ugsMetadataCollection.UpdateBadgeAsync(metadata, badge.Name, badge.Url, badge.State);
 				}
 			}
 			return Ok();
@@ -146,272 +120,268 @@ namespace Horde.Build.Controllers
 		/// <summary>
 		/// Searches for metadata updates
 		/// </summary>
-		/// <param name="Stream">THe stream to search for</param>
-		/// <param name="MinChange">Minimum changelist number</param>
-		/// <param name="MaxChange">Maximum changelist number</param>
-		/// <param name="Project">The project identifiers to search for</param>
-		/// <param name="Sequence">Last sequence number</param>
+		/// <param name="stream">THe stream to search for</param>
+		/// <param name="minChange">Minimum changelist number</param>
+		/// <param name="maxChange">Maximum changelist number</param>
+		/// <param name="project">The project identifiers to search for</param>
+		/// <param name="sequence">Last sequence number</param>
 		/// <returns>List of metadata updates</returns>
 		[HttpGet]
 		[Route("/ugs/api/metadata")]
-		public async Task<GetUgsMetadataListResponse> FindMetadataAsync([FromQuery] string Stream, [FromQuery] int MinChange, [FromQuery] int? MaxChange = null, [FromQuery] string? Project = null, [FromQuery] long? Sequence = null)
+		public async Task<GetUgsMetadataListResponse> FindMetadataAsync([FromQuery] string stream, [FromQuery] int minChange, [FromQuery] int? maxChange = null, [FromQuery] string? project = null, [FromQuery] long? sequence = null)
 		{
-			List<IUgsMetadata> MetadataList = await UgsMetadataCollection.FindAsync(Stream, MinChange, MaxChange, Sequence);
+			List<IUgsMetadata> metadataList = await _ugsMetadataCollection.FindAsync(stream, minChange, maxChange, sequence);
 		
-			GetUgsMetadataListResponse Response = new GetUgsMetadataListResponse();
-			if(Sequence != null)
+			GetUgsMetadataListResponse response = new GetUgsMetadataListResponse();
+			if(sequence != null)
 			{
-				Response.SequenceNumber = Sequence.Value;
+				response.SequenceNumber = sequence.Value;
 			}
 
-			foreach (IUgsMetadata Metadata in MetadataList)
+			foreach (IUgsMetadata metadata in metadataList)
 			{
-				if (Metadata.UpdateTicks > Response.SequenceNumber)
+				if (metadata.UpdateTicks > response.SequenceNumber)
 				{
-					Response.SequenceNumber = Metadata.UpdateTicks;
+					response.SequenceNumber = metadata.UpdateTicks;
 				}
-				if (String.IsNullOrEmpty(Metadata.Project) || Metadata.Project.Equals(Project, StringComparison.OrdinalIgnoreCase))
+				if (String.IsNullOrEmpty(metadata.Project) || metadata.Project.Equals(project, StringComparison.OrdinalIgnoreCase))
 				{
-					Response.Items.Add(new GetUgsMetadataResponse(Metadata));
+					response.Items.Add(new GetUgsMetadataResponse(metadata));
 				}
 			}
-			return Response;
+			return response;
 		}
 
 		/// <summary>
 		/// Retrieve information about open issues
 		/// </summary>
-		/// <param name="User"></param>
-		/// <param name="IncludeResolved">Whether to include resolved issues</param>
-		/// <param name="MaxResults">Maximum number of results to return</param>
+		/// <param name="user"></param>
+		/// <param name="includeResolved">Whether to include resolved issues</param>
+		/// <param name="maxResults">Maximum number of results to return</param>
 		/// <returns>List of matching agents</returns>
 		[HttpGet]
 		[Route("/ugs/api/issues")]
 		[ProducesResponseType(typeof(GetUgsIssueResponse), 200)]
-		[System.Diagnostics.CodeAnalysis.SuppressMessage("Design", "CA1801:Remove unused parameter", Justification = "<Pending>")]
-		[System.Diagnostics.CodeAnalysis.SuppressMessage("Style", "IDE0060:Remove unused parameter", Justification = "<Pending>")]
-		public async Task<ActionResult<List<GetUgsIssueResponse>>> GetIssuesAsync([FromQuery] string? User = null, [FromQuery] bool IncludeResolved = false, [FromQuery] int MaxResults = 100)
+		public async Task<ActionResult<List<GetUgsIssueResponse>>> GetIssuesAsync([FromQuery] string? user = null, [FromQuery] bool includeResolved = false, [FromQuery] int maxResults = 100)
 		{
-			IUser? UserInfo = (User != null) ? await UserCollection.FindUserByLoginAsync(User) : null;
+			IUser? userInfo = (user != null) ? await _userCollection.FindUserByLoginAsync(user) : null;
 
-			List<GetUgsIssueResponse> Responses = new List<GetUgsIssueResponse>();
-			if (IncludeResolved)
+			List<GetUgsIssueResponse> responses = new List<GetUgsIssueResponse>();
+			if (includeResolved)
 			{
-				List<IIssue> Issues = await IssueService.FindIssuesAsync(null, Resolved: null, Count: MaxResults);
-				foreach(IIssue Issue in Issues)
+				List<IIssue> issues = await _issueService.FindIssuesAsync(null, resolved: null, count: maxResults);
+				foreach(IIssue issue in issues)
 				{
-					IIssueDetails Details = await IssueService.GetIssueDetailsAsync(Issue);
-					bool bNotify = UserInfo != null && Details.Suspects.Any(x => x.AuthorId == UserInfo.Id);
-					Responses.Add(await CreateIssueResponseAsync(Details, bNotify));
+					IIssueDetails details = await _issueService.GetIssueDetailsAsync(issue);
+					bool bNotify = userInfo != null && details.Suspects.Any(x => x.AuthorId == userInfo.Id);
+					responses.Add(await CreateIssueResponseAsync(details, bNotify));
 				}
 			}
 			else
 			{
-				foreach (IIssueDetails CachedOpenIssue in IssueService.CachedOpenIssues)
+				foreach (IIssueDetails cachedOpenIssue in _issueService.CachedOpenIssues)
 				{
-					if (Responses.Count >= MaxResults)
+					if (responses.Count >= maxResults)
 					{
 						break;
 					}
 
-					if (CachedOpenIssue.ShowNotifications())
+					if (cachedOpenIssue.ShowNotifications())
 					{
-						bool bNotify = UserInfo != null && CachedOpenIssue.IncludeForUser(UserInfo.Id);
-						Responses.Add(await CreateIssueResponseAsync(CachedOpenIssue, bNotify));
+						bool bNotify = userInfo != null && cachedOpenIssue.IncludeForUser(userInfo.Id);
+						responses.Add(await CreateIssueResponseAsync(cachedOpenIssue, bNotify));
 					}
 				}
 			}
-			return Responses;
+			return responses;
 		}
 
 		/// <summary>
 		/// Retrieve information about a specific issue
 		/// </summary>
-		/// <param name="IssueId">Id of the issue to get information about</param>
-		/// <param name="Filter">Filter for the properties to return</param>
+		/// <param name="issueId">Id of the issue to get information about</param>
+		/// <param name="filter">Filter for the properties to return</param>
 		/// <returns>List of matching agents</returns>
 		[HttpGet]
-		[Route("/ugs/api/issues/{IssueId}")]
+		[Route("/ugs/api/issues/{issueId}")]
 		[ProducesResponseType(typeof(GetUgsIssueBuildResponse), 200)]
-		public async Task<ActionResult<object>> GetIssueAsync(int IssueId, [FromQuery] PropertyFilter? Filter = null)
+		public async Task<ActionResult<object>> GetIssueAsync(int issueId, [FromQuery] PropertyFilter? filter = null)
 		{
-			IIssueDetails? Issue = await IssueService.GetIssueDetailsAsync(IssueId);
-			if (Issue == null)
+			IIssueDetails? issue = await _issueService.GetIssueDetailsAsync(issueId);
+			if (issue == null)
 			{
 				return NotFound();
 			}
 
-			return PropertyFilter.Apply(await CreateIssueResponseAsync(Issue, false), Filter);
+			return PropertyFilter.Apply(await CreateIssueResponseAsync(issue, false), filter);
 		}
 
 		/// <summary>
 		/// Retrieve information about builds for a specific issue
 		/// </summary>
-		/// <param name="IssueId">Id of the issue to get information about</param>
-		/// <param name="Filter">Filter for the properties to return</param>
+		/// <param name="issueId">Id of the issue to get information about</param>
+		/// <param name="filter">Filter for the properties to return</param>
 		/// <returns>List of matching agents</returns>
 		[HttpGet]
-		[Route("/ugs/api/issues/{IssueId}/builds")]
+		[Route("/ugs/api/issues/{issueId}/builds")]
 		[ProducesResponseType(typeof(List<GetUgsIssueBuildResponse>), 200)]
-		public async Task<ActionResult<List<object>>> GetIssueBuildsAsync(int IssueId, [FromQuery] PropertyFilter? Filter = null)
+		public async Task<ActionResult<List<object>>> GetIssueBuildsAsync(int issueId, [FromQuery] PropertyFilter? filter = null)
 		{
-			IIssueDetails? Issue = await IssueService.GetCachedIssueDetailsAsync(IssueId);
-			if (Issue == null)
+			IIssueDetails? issue = await _issueService.GetCachedIssueDetailsAsync(issueId);
+			if (issue == null)
 			{
 				return NotFound();
 			}
 
-			List<GetUgsIssueBuildResponse> Responses = new List<GetUgsIssueBuildResponse>();
-			foreach (IIssueSpan Span in Issue.Spans)
+			List<GetUgsIssueBuildResponse> responses = new List<GetUgsIssueBuildResponse>();
+			foreach (IIssueSpan span in issue.Spans)
 			{
-				if (Span.LastSuccess != null)
+				if (span.LastSuccess != null)
 				{
-					Responses.Add(CreateBuildResponse(Span, Span.LastSuccess, IssueBuildOutcome.Success));
+					responses.Add(CreateBuildResponse(span, span.LastSuccess, IssueBuildOutcome.Success));
 				}
-				foreach (IIssueStep Step in Issue.Steps)
+				foreach (IIssueStep step in issue.Steps)
 				{
-					Responses.Add(CreateBuildResponse(Span, Step, IssueBuildOutcome.Error));
+					responses.Add(CreateBuildResponse(span, step, IssueBuildOutcome.Error));
 				}
-				if (Span.NextSuccess != null)
+				if (span.NextSuccess != null)
 				{
-					Responses.Add(CreateBuildResponse(Span, Span.NextSuccess, IssueBuildOutcome.Success));
+					responses.Add(CreateBuildResponse(span, span.NextSuccess, IssueBuildOutcome.Success));
 				}
 			}
 
-			return Responses.ConvertAll(x => PropertyFilter.Apply(x, Filter));
+			return responses.ConvertAll(x => PropertyFilter.Apply(x, filter));
 		}
 
 		/// <summary>
 		/// Retrieve information about builds for a specific issue
 		/// </summary>
-		/// <param name="IssueId">Id of the issue to get information about</param>
+		/// <param name="issueId">Id of the issue to get information about</param>
 		/// <returns>List of matching agents</returns>
 		[HttpGet]
-		[Route("/ugs/api/issues/{IssueId}/diagnostics")]
+		[Route("/ugs/api/issues/{issueId}/diagnostics")]
 		[ProducesResponseType(typeof(List<GetUgsIssueDiagnosticResponse>), 200)]
-		public async Task<ActionResult<List<GetUgsIssueDiagnosticResponse>>> GetIssueDiagnosticsAsync(int IssueId)
+		public async Task<ActionResult<List<GetUgsIssueDiagnosticResponse>>> GetIssueDiagnosticsAsync(int issueId)
 		{
-			List<GetUgsIssueDiagnosticResponse> Diagnostics = new List<GetUgsIssueDiagnosticResponse>();
-
-			Dictionary<LogId, ILogFile?> LogFiles = new Dictionary<LogId, ILogFile?>();
+			List<GetUgsIssueDiagnosticResponse> diagnostics = new List<GetUgsIssueDiagnosticResponse>();
 
-			JobPermissionsCache PermissionsCache = new JobPermissionsCache();
+			Dictionary<LogId, ILogFile?> logFiles = new Dictionary<LogId, ILogFile?>();
 
-			List<ILogEvent> Events = await IssueService.FindEventsForIssueAsync(IssueId, Count: 10);
-			foreach (ILogEvent Event in Events)
+			List<ILogEvent> events = await _issueService.FindEventsForIssueAsync(issueId, count: 10);
+			foreach (ILogEvent logEvent in events)
 			{
-				ILogFile? LogFile;
-				if(!LogFiles.TryGetValue(Event.LogId, out LogFile))
+				ILogFile? logFile;
+				if(!logFiles.TryGetValue(logEvent.LogId, out logFile))
 				{
-					LogFile = await LogFileService.GetLogFileAsync(Event.LogId);
-					LogFiles.Add(Event.LogId, LogFile);
+					logFile = await _logFileService.GetLogFileAsync(logEvent.LogId);
+					logFiles.Add(logEvent.LogId, logFile);
 				}
-				if (LogFile != null)
+				if (logFile != null)
 				{
-					ILogEventData EventData = await LogFileService.GetEventDataAsync(LogFile, Event.LineIndex, Event.LineCount);
-					long BuildId = Event.LogId.GetHashCode();
-					Uri Url = new Uri(Settings.DashboardUrl, $"log/{Event.LogId}?lineindex={Event.LineIndex}");
+					ILogEventData eventData = await _logFileService.GetEventDataAsync(logFile, logEvent.LineIndex, logEvent.LineCount);
+					long buildId = logEvent.LogId.GetHashCode();
+					Uri url = new Uri(_settings.DashboardUrl, $"log/{logEvent.LogId}?lineindex={logEvent.LineIndex}");
 
-					GetUgsIssueDiagnosticResponse Diagnostic = new GetUgsIssueDiagnosticResponse(BuildId, EventData.Message, Url);
-					Diagnostics.Add(Diagnostic);
+					GetUgsIssueDiagnosticResponse diagnostic = new GetUgsIssueDiagnosticResponse(buildId, eventData.Message, url);
+					diagnostics.Add(diagnostic);
 				}
 			}
 
-			return Diagnostics;
+			return diagnostics;
 		}
 
 		/// <summary>
 		/// Gets the URL for a failing step in the
 		/// </summary>
-		/// <param name="Details">The issue to get a URL for</param>
+		/// <param name="details">The issue to get a URL for</param>
 		/// <param name="bNotify">Whether to show notifications for this issue</param>
 		/// <returns>The issue response</returns>
-		async Task<GetUgsIssueResponse> CreateIssueResponseAsync(IIssueDetails Details, bool bNotify)
+		async Task<GetUgsIssueResponse> CreateIssueResponseAsync(IIssueDetails details, bool bNotify)
 		{
-			Uri? BuildUrl = GetIssueBuildUrl(Details);
+			Uri? buildUrl = GetIssueBuildUrl(details);
 
-			IUser? Owner = Details.Issue.OwnerId.HasValue ? await UserCollection.GetCachedUserAsync(Details.Issue.OwnerId.Value) : null;
-			IUser? NominatedBy = Details.Issue.NominatedById.HasValue ? await UserCollection.GetCachedUserAsync(Details.Issue.NominatedById.Value) : null;
+			IUser? owner = details.Issue.OwnerId.HasValue ? await _userCollection.GetCachedUserAsync(details.Issue.OwnerId.Value) : null;
+			IUser? nominatedBy = details.Issue.NominatedById.HasValue ? await _userCollection.GetCachedUserAsync(details.Issue.NominatedById.Value) : null;
 
-			return new GetUgsIssueResponse(Details, Owner, NominatedBy, bNotify, BuildUrl);
+			return new GetUgsIssueResponse(details, owner, nominatedBy, bNotify, buildUrl);
 		}
 
 		/// <summary>
 		/// Gets the URL for a failing step in the given issue
 		/// </summary>
-		/// <param name="Issue">The issue details</param>
+		/// <param name="issue">The issue details</param>
 		/// <returns>The build URL</returns>
-		Uri? GetIssueBuildUrl(IIssueDetails Issue)
+		Uri? GetIssueBuildUrl(IIssueDetails issue)
 		{
-			HashSet<ObjectId> UnresolvedSpans = new HashSet<ObjectId>(Issue.Spans.Where(x => x.NextSuccess == null).Select(x => x.Id));
+			HashSet<ObjectId> unresolvedSpans = new HashSet<ObjectId>(issue.Spans.Where(x => x.NextSuccess == null).Select(x => x.Id));
 
-			IIssueStep? Step = Issue.Steps.OrderByDescending(x => UnresolvedSpans.Contains(x.SpanId)).ThenByDescending(x => x.Change).FirstOrDefault();
-			if (Step == null)
+			IIssueStep? step = issue.Steps.OrderByDescending(x => unresolvedSpans.Contains(x.SpanId)).ThenByDescending(x => x.Change).FirstOrDefault();
+			if (step == null)
 			{
 				return null;
 			}
 
-			return new Uri(Settings.DashboardUrl, $"job/{Step.JobId}?step={Step.StepId}");
+			return new Uri(_settings.DashboardUrl, $"job/{step.JobId}?step={step.StepId}");
 		}
 
 		/// <summary>
 		/// Creates the response for a particular build
 		/// </summary>
-		/// <param name="Span">Span containing the step</param>
-		/// <param name="Step">The step to describe</param>
-		/// <param name="Outcome">Outcome of this step</param>
+		/// <param name="span">Span containing the step</param>
+		/// <param name="step">The step to describe</param>
+		/// <param name="outcome">Outcome of this step</param>
 		/// <returns>Response object</returns>
-		GetUgsIssueBuildResponse CreateBuildResponse(IIssueSpan Span, IIssueStep Step, IssueBuildOutcome Outcome)
+		GetUgsIssueBuildResponse CreateBuildResponse(IIssueSpan span, IIssueStep step, IssueBuildOutcome outcome)
 		{
-			GetUgsIssueBuildResponse Response = new GetUgsIssueBuildResponse(Span.StreamName, Step.Change, Outcome);
-			Response.Id = Step.LogId.GetHashCode();
-			Response.JobName = $"{Step.JobName}: {Span.NodeName}";
-			Response.JobUrl = new Uri(Settings.DashboardUrl, $"job/{Step.JobId}");
-			Response.JobStepName = Span.NodeName;
-			Response.JobStepUrl = new Uri(Settings.DashboardUrl, $"job/{Step.JobId}?step={Step.StepId}");
-			Response.ErrorUrl = Response.JobStepUrl;
-			return Response;
+			GetUgsIssueBuildResponse response = new GetUgsIssueBuildResponse(span.StreamName, step.Change, outcome);
+			response.Id = step.LogId.GetHashCode();
+			response.JobName = $"{step.JobName}: {span.NodeName}";
+			response.JobUrl = new Uri(_settings.DashboardUrl, $"job/{step.JobId}");
+			response.JobStepName = span.NodeName;
+			response.JobStepUrl = new Uri(_settings.DashboardUrl, $"job/{step.JobId}?step={step.StepId}");
+			response.ErrorUrl = response.JobStepUrl;
+			return response;
 		}
 
 		/// <summary>
 		/// Update an issue
 		/// </summary>
-		/// <param name="IssueId">Id of the issue to get information about</param>
-		/// <param name="Request">The update information</param>
+		/// <param name="issueId">Id of the issue to get information about</param>
+		/// <param name="request">The update information</param>
 		/// <returns>List of matching agents</returns>
 		[HttpPut]
-		[Route("/ugs/api/issues/{IssueId}")]
-		public async Task<ActionResult> UpdateIssueAsync(int IssueId, [FromBody] UpdateUgsIssueRequest Request)
+		[Route("/ugs/api/issues/{issueId}")]
+		public async Task<ActionResult> UpdateIssueAsync(int issueId, [FromBody] UpdateUgsIssueRequest request)
 		{
-			UserId? NewOwnerId = null;
-			if (!String.IsNullOrEmpty(Request.Owner))
+			UserId? newOwnerId = null;
+			if (!String.IsNullOrEmpty(request.Owner))
 			{
-				NewOwnerId = (await UserCollection.FindOrAddUserByLoginAsync(Request.Owner))?.Id;
+				newOwnerId = (await _userCollection.FindOrAddUserByLoginAsync(request.Owner))?.Id;
 			}
 
-			UserId? NewNominatedById = null;
-			if (!String.IsNullOrEmpty(Request.NominatedBy))
+			UserId? newNominatedById = null;
+			if (!String.IsNullOrEmpty(request.NominatedBy))
 			{
-				NewNominatedById = (await UserCollection.FindOrAddUserByLoginAsync(Request.NominatedBy))?.Id;
+				newNominatedById = (await _userCollection.FindOrAddUserByLoginAsync(request.NominatedBy))?.Id;
 			}
 
-			UserId? NewDeclinedById = null;
-			if (!String.IsNullOrEmpty(Request.DeclinedBy))
+			UserId? newDeclinedById = null;
+			if (!String.IsNullOrEmpty(request.DeclinedBy))
 			{
-				NewDeclinedById = (await UserCollection.FindOrAddUserByLoginAsync(Request.DeclinedBy))?.Id;
+				newDeclinedById = (await _userCollection.FindOrAddUserByLoginAsync(request.DeclinedBy))?.Id;
 			}
 
-			UserId? NewResolvedById = null;
-			if (!String.IsNullOrEmpty(Request.ResolvedBy))
+			UserId? newResolvedById = null;
+			if (!String.IsNullOrEmpty(request.ResolvedBy))
 			{
-				NewResolvedById = (await UserCollection.FindOrAddUserByLoginAsync(Request.ResolvedBy))?.Id;
+				newResolvedById = (await _userCollection.FindOrAddUserByLoginAsync(request.ResolvedBy))?.Id;
 			}
-			if (NewResolvedById == null && Request.Resolved.HasValue)
+			if (newResolvedById == null && request.Resolved.HasValue)
 			{
-				NewResolvedById = Request.Resolved.Value ? IIssue.ResolvedByUnknownId : UserId.Empty;
+				newResolvedById = request.Resolved.Value ? IIssue.ResolvedByUnknownId : UserId.Empty;
 			}
 
-			if (!await IssueService.UpdateIssueAsync(IssueId, OwnerId: NewOwnerId, NominatedById: NewNominatedById, Acknowledged: Request.Acknowledged, DeclinedById: NewDeclinedById, FixChange: Request.FixChange, ResolvedById: NewResolvedById))
+			if (!await _issueService.UpdateIssueAsync(issueId, ownerId: newOwnerId, nominatedById: newNominatedById, acknowledged: request.Acknowledged, declinedById: newDeclinedById, fixChange: request.FixChange, resolvedById: newResolvedById))
 			{
 				return NotFound();
 			}
@@ -423,9 +393,9 @@ namespace Horde.Build.Controllers
 		/// </summary>
 		[HttpPost]
 		[Route("/ugs/api/netcore")]
-		public ActionResult<object> PostNetCoreInfo(string? User = null, string? Machine = null, bool NetCore = false)
+		public ActionResult<object> PostNetCoreInfo(string? user = null, string? machine = null, bool netCore = false)
 		{
-			Logger.LogInformation("NetCore: User={User}, Machine={Machine}, NetCore={NetCore}", User, Machine, NetCore);
+			_logger.LogInformation("NetCore: User={User}, Machine={Machine}, NetCore={NetCore}", user, machine, netCore);
 			return new { };
 		}
 	}
diff --git a/Engine/Source/Programs/Horde/Horde.Build/Controllers/UserController.cs b/Engine/Source/Programs/Horde/Horde.Build/Controllers/UserController.cs
--- a/Engine/Source/Programs/Horde/Horde.Build/Controllers/UserController.cs
+++ b/Engine/Source/Programs/Horde/Horde.Build/Controllers/UserController.cs
@@ -1,5 +1,8 @@
 // Copyright Epic Games, Inc. All Rights Reserved.
 
+using System.Collections.Generic;
+using System.Linq;
+using System.Threading.Tasks;
 using Horde.Build.Api;
 using Horde.Build.Collections;
 using Horde.Build.Models;
@@ -7,11 +10,6 @@ using Horde.Build.Services;
 using Horde.Build.Utilities;
 using Microsoft.AspNetCore.Authorization;
 using Microsoft.AspNetCore.Mvc;
-using MongoDB.Bson;
-using System;
-using System.Collections.Generic;
-using System.Linq;
-using System.Threading.Tasks;
 
 namespace Horde.Build.Controllers
 {
@@ -39,12 +37,12 @@ namespace Horde.Build.Controllers
 		/// <summary>
 		/// Constructor
 		/// </summary>
-		/// <param name="UserCollection"></param>
-		/// <param name="AvatarService"></param>
-		public UserController(IUserCollection UserCollection, IAvatarService? AvatarService)
+		/// <param name="userCollection"></param>
+		/// <param name="avatarService"></param>
+		public UserController(IUserCollection userCollection, IAvatarService? avatarService)
 		{
-			this.UserCollection = UserCollection;
-			this.AvatarService = AvatarService;
+			UserCollection = userCollection;
+			AvatarService = avatarService;
 		}
 
 		/// <summary>
@@ -54,18 +52,18 @@ namespace Horde.Build.Controllers
 		[HttpGet]
 		[Route("/api/v1/user")]
 		[ProducesResponseType(typeof(List<GetUserResponse>), 200)]
-		public async Task<ActionResult<object>> GetUserAsync([FromQuery] PropertyFilter? Filter = null)
+		public async Task<ActionResult<object>> GetUserAsync([FromQuery] PropertyFilter? filter = null)
 		{
-			IUser? InternalUser = await UserCollection.GetUserAsync(User);
-			if (InternalUser == null)
+			IUser? internalUser = await UserCollection.GetUserAsync(User);
+			if (internalUser == null)
 			{
 				return NotFound();
 			}
 
-			IAvatar? Avatar = (AvatarService == null)? (IAvatar?)null : await AvatarService.GetAvatarAsync(InternalUser);
-			IUserClaims Claims = await UserCollection.GetClaimsAsync(InternalUser.Id);
-			IUserSettings Settings = await UserCollection.GetSettingsAsync(InternalUser.Id);
-			return PropertyFilter.Apply(new GetUserResponse(InternalUser, Avatar, Claims, Settings), Filter);
+			IAvatar? avatar = (AvatarService == null)? (IAvatar?)null : await AvatarService.GetAvatarAsync(internalUser);
+			IUserClaims claims = await UserCollection.GetClaimsAsync(internalUser.Id);
+			IUserSettings settings = await UserCollection.GetSettingsAsync(internalUser.Id);
+			return PropertyFilter.Apply(new GetUserResponse(internalUser, avatar, claims, settings), filter);
 		}
 
 		/// <summary>
@@ -74,14 +72,14 @@ namespace Horde.Build.Controllers
 		/// <returns>Http result code</returns>
 		[HttpPut]
 		[Route("/api/v1/user")]
-		public async Task<ActionResult> UpdateUserAsync(UpdateUserRequest Request)
+		public async Task<ActionResult> UpdateUserAsync(UpdateUserRequest request)
 		{
-			UserId? UserId = User.GetUserId();
-			if(UserId == null)
+			UserId? userId = User.GetUserId();
+			if(userId == null)
 			{
 				return BadRequest("Current user does not have a registered profile");
 			}
-			await UserCollection.UpdateSettingsAsync(UserId.Value, Request.EnableExperimentalFeatures, Request.DashboardSettings?.ToBsonValue(), Request.AddPinnedJobIds?.Select(x => new JobId(x)), Request.RemovePinnedJobIds?.Select(x => new JobId(x)));
+			await UserCollection.UpdateSettingsAsync(userId.Value, request.EnableExperimentalFeatures, request.DashboardSettings?.ToBsonValue(), request.AddPinnedJobIds?.Select(x => new JobId(x)), request.RemovePinnedJobIds?.Select(x => new JobId(x)));
 			return Ok();
 		}
 	}
diff --git a/Engine/Source/Programs/Horde/Horde.Build/Controllers/UsersController.cs b/Engine/Source/Programs/Horde/Horde.Build/Controllers/UsersController.cs
--- a/Engine/Source/Programs/Horde/Horde.Build/Controllers/UsersController.cs
+++ b/Engine/Source/Programs/Horde/Horde.Build/Controllers/UsersController.cs
@@ -1,19 +1,16 @@
 // Copyright Epic Games, Inc. All Rights Reserved.
 
+using System;
+using System.Collections.Generic;
+using System.Linq;
+using System.Threading.Tasks;
 using Horde.Build.Api;
 using Horde.Build.Collections;
 using Horde.Build.Models;
-using Horde.Build.Notifications.Impl;
 using Horde.Build.Services;
 using Horde.Build.Utilities;
 using Microsoft.AspNetCore.Authorization;
 using Microsoft.AspNetCore.Mvc;
-using Microsoft.Extensions.Options;
-using MongoDB.Bson;
-using System;
-using System.Collections.Generic;
-using System.Linq;
-using System.Threading.Tasks;
 
 namespace Horde.Build.Controllers
 {
@@ -40,12 +37,12 @@ namespace Horde.Build.Controllers
 		/// <summary>
 		/// Constructor
 		/// </summary>
-		/// <param name="UserCollection"></param>
-		/// <param name="AvatarService"></param>
-		public UsersController(IUserCollection UserCollection, IAvatarService? AvatarService)
+		/// <param name="userCollection"></param>
+		/// <param name="avatarService"></param>
+		public UsersController(IUserCollection userCollection, IAvatarService? avatarService)
 		{
-			this.UserCollection = UserCollection;
-			this.AvatarService = AvatarService;
+			UserCollection = userCollection;
+			AvatarService = avatarService;
 		}
 
 		/// <summary>
@@ -53,26 +50,26 @@ namespace Horde.Build.Controllers
 		/// </summary>
 		/// <returns>Http result code</returns>
 		[HttpGet]
-		[Route("/api/v1/users/{Id}")]
+		[Route("/api/v1/users/{id}")]
 		[ProducesResponseType(typeof(List<GetUserResponse>), 200)]
-		public async Task<ActionResult<object>> GetUserAsync(string Id, [FromQuery] PropertyFilter? Filter = null)
+		public async Task<ActionResult<object>> GetUserAsync(string id, [FromQuery] PropertyFilter? filter = null)
 		{
-			UserId? UserId = ParseUserId(Id);
-			if(UserId == null)
+			UserId? userId = ParseUserId(id);
+			if(userId == null)
 			{
-				return BadRequest("Invalid user id '{Id}'", Id);
+				return BadRequest("Invalid user id '{Id}'", id);
 			}
 
-			IUser? User = await UserCollection.GetUserAsync(UserId.Value);
-			if (User == null)
+			IUser? user = await UserCollection.GetUserAsync(userId.Value);
+			if (user == null)
 			{
-				return NotFound(UserId.Value);
+				return NotFound(userId.Value);
 			}
 
-			IAvatar? Avatar = (AvatarService == null) ? (IAvatar?)null : await AvatarService.GetAvatarAsync(User);
-			IUserClaims? Claims = await UserCollection.GetClaimsAsync(User.Id);
-			IUserSettings? Settings = await UserCollection.GetSettingsAsync(User.Id);
-			return PropertyFilter.Apply(new GetUserResponse(User, Avatar, Claims, Settings), Filter);
+			IAvatar? avatar = (AvatarService == null) ? (IAvatar?)null : await AvatarService.GetAvatarAsync(user);
+			IUserClaims? claims = await UserCollection.GetClaimsAsync(user.Id);
+			IUserSettings? settings = await UserCollection.GetSettingsAsync(user.Id);
+			return PropertyFilter.Apply(new GetUserResponse(user, avatar, claims, settings), filter);
 		}
 
 		/// <summary>
@@ -83,42 +80,42 @@ namespace Horde.Build.Controllers
 		[Route("/api/v1/users")]
 		[ProducesResponseType(typeof(List<GetUserResponse>), 200)]
 		public async Task<ActionResult<List<GetUserResponse>>> FindUsersAsync(
-			[FromQuery] string[]? Ids = null,
-			[FromQuery] string? NameRegex = null,
-			[FromQuery] int Index = 0,
-			[FromQuery] int Count = 100,
-			[FromQuery] bool IncludeClaims = false,
-			[FromQuery] bool IncludeAvatar = false)
+			[FromQuery] string[]? ids = null,
+			[FromQuery] string? nameRegex = null,
+			[FromQuery] int index = 0,
+			[FromQuery] int count = 100,
+			[FromQuery] bool includeClaims = false,
+			[FromQuery] bool includeAvatar = false)
 		{
 
-			UserId[]? UserIds = null;
-			if (Ids != null && Ids.Length > 0)
+			UserId[]? userIds = null;
+			if (ids != null && ids.Length > 0)
 			{
-				UserIds = Ids.Select(x => new UserId(x)).ToArray();
+				userIds = ids.Select(x => new UserId(x)).ToArray();
 			}
 
-			List<IUser> Users = await UserCollection.FindUsersAsync(UserIds, NameRegex, Index, Count);
+			List<IUser> users = await UserCollection.FindUsersAsync(userIds, nameRegex, index, count);
 
-			List<GetUserResponse> Response = new List<GetUserResponse>();
-			foreach (IUser User in Users)
+			List<GetUserResponse> response = new List<GetUserResponse>();
+			foreach (IUser user in users)
 			{
-				IAvatar? Avatar = (AvatarService == null || !IncludeAvatar) ? (IAvatar?)null : await AvatarService.GetAvatarAsync(User);
-				IUserClaims? Claims = (!IncludeClaims) ? null : await UserCollection.GetClaimsAsync(User.Id);
-				Response.Add(new GetUserResponse(User, Avatar, Claims, null));
+				IAvatar? avatar = (AvatarService == null || !includeAvatar) ? (IAvatar?)null : await AvatarService.GetAvatarAsync(user);
+				IUserClaims? claims = (!includeClaims) ? null : await UserCollection.GetClaimsAsync(user.Id);
+				response.Add(new GetUserResponse(user, avatar, claims, null));
 			}
 
-			return Response;
+			return response;
 		}
 
-		UserId? ParseUserId(string Id)
+		UserId? ParseUserId(string id)
 		{
-			if (Id.Equals("current", StringComparison.OrdinalIgnoreCase))
+			if (id.Equals("current", StringComparison.OrdinalIgnoreCase))
 			{
 				return User.GetUserId();
 			}
-			else if(UserId.TryParse(Id, out UserId Result))
+			else if(UserId.TryParse(id, out UserId result))
 			{
-				return Result;
+				return result;
 			}
 			return null;
 		}
diff --git a/Engine/Source/Programs/Horde/Horde.Build/Devices/IDevice.cs b/Engine/Source/Programs/Horde/Horde.Build/Devices/IDevice.cs
--- a/Engine/Source/Programs/Horde/Horde.Build/Devices/IDevice.cs
+++ b/Engine/Source/Programs/Horde/Horde.Build/Devices/IDevice.cs
@@ -1,20 +1,19 @@
-﻿// Copyright Epic Games, Inc. All Rights Reserved.
+// Copyright Epic Games, Inc. All Rights Reserved.
 
 using System;
-using MongoDB.Bson;
-using MongoDB.Bson.Serialization.Attributes;
-
 using System.Collections.Generic;
 using Horde.Build.Acls;
 using Horde.Build.Utilities;
+using MongoDB.Bson;
+using MongoDB.Bson.Serialization.Attributes;
 
 namespace Horde.Build.Models
 {
 	using DeviceId = StringId<IDevice>;
 	using DevicePlatformId = StringId<IDevicePlatform>;
 	using DevicePoolId = StringId<IDevicePool>;
-	using UserId = ObjectId<IUser>;
 	using ProjectId = StringId<IProject>;
+	using UserId = ObjectId<IUser>;
 
 	/// <summary>
 	/// A reservation containing one or more devices
@@ -76,14 +75,12 @@ namespace Horde.Build.Models
 		/// </summary>
 		public List<string> RequestedDevicePlatforms { get; }
 
-
 		/// <summary>
 		/// The legacy reservation system guid, to be removed once can update Gauntlet client in all streams
 		/// </summary>
 		public string LegacyGuid { get; }
 	}
 
-
 	/// <summary>
 	/// A device platform 
 	/// </summary>
@@ -103,7 +100,6 @@ namespace Horde.Build.Models
 		/// A list of valid models for the platform
 		/// </summary>
 		public IReadOnlyList<string>? Models { get; }
-
 	}
 
     /// <summary>
@@ -120,7 +116,6 @@ namespace Horde.Build.Models
         /// Shared by users with remote checking and checkouts
         /// </summary>
         Shared
-
     }
 
     /// <summary>
@@ -190,15 +185,13 @@ namespace Horde.Build.Models
 		/// <summary>
 		/// Constructor
 		/// </summary>
-		/// <param name="ReservationStartUtc"></param>
-		public DeviceUtilizationTelemetry( DateTime ReservationStartUtc )
+		/// <param name="reservationStartUtc"></param>
+		public DeviceUtilizationTelemetry( DateTime reservationStartUtc )
 		{
-			this.ReservationStartUtc = ReservationStartUtc;
+			ReservationStartUtc = reservationStartUtc;
 		}
-
 	}
 
-
 	/// <summary>
 	/// A physical device
 	/// </summary>
@@ -229,7 +222,6 @@ namespace Horde.Build.Models
 		/// </summary>
 		public string? ModelId { get; }
 
-
 		/// <summary>
 		/// Address if the device supports shared network connections
 		/// </summary>
@@ -255,7 +247,6 @@ namespace Horde.Build.Models
 		/// </summary>
 		public DateTime? CheckOutTime { get; }
 
-
 		/// <summary>
 		/// The last time a problem was reported
 		/// </summary>
@@ -282,10 +273,3 @@ namespace Horde.Build.Models
 		public Acl? Acl { get; }
 	}
 }
-
-
-
-
-
-
-
diff --git a/Engine/Source/Programs/Horde/Horde.Build/Devices/IDeviceCollection.cs b/Engine/Source/Programs/Horde/Horde.Build/Devices/IDeviceCollection.cs
--- a/Engine/Source/Programs/Horde/Horde.Build/Devices/IDeviceCollection.cs
+++ b/Engine/Source/Programs/Horde/Horde.Build/Devices/IDeviceCollection.cs
@@ -1,18 +1,18 @@
 // Copyright Epic Games, Inc. All Rights Reserved.
 
+using System.Collections.Generic;
+using System.Threading.Tasks;
 using Horde.Build.Models;
 using Horde.Build.Utilities;
 using MongoDB.Bson;
-using System.Collections.Generic;
-using System.Threading.Tasks;
 
 namespace Horde.Build.Collections
 {
 	using DeviceId = StringId<IDevice>;
 	using DevicePlatformId = StringId<IDevicePlatform>;
 	using DevicePoolId = StringId<IDevicePool>;
-	using UserId = ObjectId<IUser>;
 	using ProjectId = StringId<IProject>;
+	using UserId = ObjectId<IUser>;
 
 	/// <summary>
 	/// Device reservation request data
@@ -42,12 +42,12 @@ namespace Horde.Build.Collections
 		/// <summary>
 		/// Constructor
 		/// </summary>
-		public DeviceRequestData(DevicePlatformId PlatformId, string RequestedPlatform, List<string>? IncludeModels = null, List<string>? ExcludeModels = null)
+		public DeviceRequestData(DevicePlatformId platformId, string requestedPlatform, List<string>? includeModels = null, List<string>? excludeModels = null)
 		{
-			this.PlatformId = PlatformId;
-			this.IncludeModels = IncludeModels ?? new List<string>();
-			this.ExcludeModels = ExcludeModels ?? new List<string>();
-			this.RequestedPlatform = RequestedPlatform;
+			PlatformId = platformId;
+			IncludeModels = includeModels ?? new List<string>();
+			ExcludeModels = excludeModels ?? new List<string>();
+			RequestedPlatform = requestedPlatform;
 		}
 	}
 
@@ -61,9 +61,9 @@ namespace Horde.Build.Collections
 		/// <summary>
 		/// Add a platform 
 		/// </summary>
-		/// <param name="Id">The id of the platform</param>
-		/// <param name="Name">The friendly name of the platform</param>
-		Task<IDevicePlatform?> TryAddPlatformAsync(DevicePlatformId Id, string Name);
+		/// <param name="id">The id of the platform</param>
+		/// <param name="name">The friendly name of the platform</param>
+		Task<IDevicePlatform?> TryAddPlatformAsync(DevicePlatformId id, string name);
 
 		/// <summary>
 		/// Get a list of all available device platforms
@@ -73,37 +73,37 @@ namespace Horde.Build.Collections
 		/// <summary>
 		/// Get a specific platform by id
 		/// </summary>
-		Task<IDevicePlatform?> GetPlatformAsync(DevicePlatformId PlatformId);
+		Task<IDevicePlatform?> GetPlatformAsync(DevicePlatformId platformId);
 
 		/// <summary>
 		/// Update a device platform
 		/// </summary>
 		/// <param name="platformId">The id of the device</param>
-		/// <param name="ModelIds">The available model ids for the platform</param>
-		Task<bool> UpdatePlatformAsync(DevicePlatformId platformId, string[]? ModelIds);
+		/// <param name="modelIds">The available model ids for the platform</param>
+		Task<bool> UpdatePlatformAsync(DevicePlatformId platformId, string[]? modelIds);
 
 		// POOLS
 
 		/// <summary>
 		/// Add a new device pool to the collection
 		/// </summary>
-		/// <param name="Id">The id of the new pool</param>
-		/// <param name="Name">The friendly name of the new pool</param>
-		/// <param name="PoolType">The pool type</param>
-		/// <param name="ProjectIds">Projects associated with this pool</param>
-		Task<IDevicePool?> TryAddPoolAsync(DevicePoolId Id, string Name, DevicePoolType PoolType, List<ProjectId>? ProjectIds );
+		/// <param name="id">The id of the new pool</param>
+		/// <param name="name">The friendly name of the new pool</param>
+		/// <param name="poolType">The pool type</param>
+		/// <param name="projectIds">Projects associated with this pool</param>
+		Task<IDevicePool?> TryAddPoolAsync(DevicePoolId id, string name, DevicePoolType poolType, List<ProjectId>? projectIds );
 
 		/// <summary>
 		/// Update a device pool
 		/// </summary>
-		/// <param name="Id">The id of the device pool to update</param>
-		/// <param name="ProjectIds">Associated project ids</param>
-		Task UpdatePoolAsync(DevicePoolId Id, List<ProjectId>? ProjectIds);
+		/// <param name="id">The id of the device pool to update</param>
+		/// <param name="projectIds">Associated project ids</param>
+		Task UpdatePoolAsync(DevicePoolId id, List<ProjectId>? projectIds);
 
 		/// <summary>
 		/// Get a pool by id
 		/// </summary>
-		Task<IDevicePool?> GetPoolAsync(DevicePoolId PoolId);
+		Task<IDevicePool?> GetPoolAsync(DevicePoolId poolId);
 
 		/// <summary>
 		/// Gets a list of existing device pools
@@ -115,85 +115,85 @@ namespace Horde.Build.Collections
 		/// <summary>
 		/// Get a device by id
 		/// </summary>
-		Task<IDevice?> GetDeviceAsync(DeviceId DeviceId);
+		Task<IDevice?> GetDeviceAsync(DeviceId deviceId);
 
 		/// <summary>
 		/// Get a device by name
 		/// </summary>
-		Task<IDevice?> GetDeviceByNameAsync(string DeviceName);
+		Task<IDevice?> GetDeviceByNameAsync(string deviceName);
 
 		/// <summary>
 		/// Get a list of all devices
 		/// </summary>
-		/// <param name="DeviceIds">Optional list of device ids to get</param>
-		Task<List<IDevice>> FindAllDevicesAsync(List<DeviceId>? DeviceIds);
+		/// <param name="deviceIds">Optional list of device ids to get</param>
+		Task<List<IDevice>> FindAllDevicesAsync(List<DeviceId>? deviceIds);
 
 		/// <summary>
 		/// Adds a new device
 		/// </summary>
-		/// <param name="Id">The device id</param>
-		/// <param name="Name">The name of the device (unique)</param>
-		/// <param name="PlatformId">The platform of the device</param>
-		/// <param name="PoolId">Which pool to add it to</param>
-		/// <param name="Enabled">Whather the device is enabled by default</param>
-		/// <param name="Address">The network address or hostname the device can be reached at</param>
-		/// <param name="ModelId">The vendor model id of the device</param>
-        /// <param name="UserId">The user adding the device</param>
-		Task<IDevice?> TryAddDeviceAsync(DeviceId Id, string Name, DevicePlatformId PlatformId, DevicePoolId PoolId, bool? Enabled, string? Address, string? ModelId, UserId? UserId);
+		/// <param name="id">The device id</param>
+		/// <param name="name">The name of the device (unique)</param>
+		/// <param name="platformId">The platform of the device</param>
+		/// <param name="poolId">Which pool to add it to</param>
+		/// <param name="enabled">Whather the device is enabled by default</param>
+		/// <param name="address">The network address or hostname the device can be reached at</param>
+		/// <param name="modelId">The vendor model id of the device</param>
+        /// <param name="userId">The user adding the device</param>
+		Task<IDevice?> TryAddDeviceAsync(DeviceId id, string name, DevicePlatformId platformId, DevicePoolId poolId, bool? enabled, string? address, string? modelId, UserId? userId);
 
 		/// <summary>
 		/// Update a device
 		/// </summary>
-		/// <param name="DeviceId">The id of the device to update</param>
-		/// <param name="NewPoolId">The new pool to assign</param>
-		/// <param name="NewName">The new name to assign</param>		
-		/// <param name="NewAddress">The new device address or hostname</param>
-		/// <param name="NewModelId">The new model id</param>
-		/// <param name="NewNotes">The devices markdown notes</param>
-		/// <param name="NewEnabled">Whether the device is enabled or not</param>
-		/// <param name="NewProblem">Whether to set or clear problem state</param>
-		/// <param name="NewMaintenance">Whether to set or clear maintenance state</param>
-        /// <param name="ModifiedByUserId">The user who is updating the device</param>
-		Task UpdateDeviceAsync(DeviceId DeviceId, DevicePoolId? NewPoolId, string? NewName, string? NewAddress, string? NewModelId, string? NewNotes, bool? NewEnabled, bool? NewProblem, bool? NewMaintenance, UserId? ModifiedByUserId = null);
+		/// <param name="deviceId">The id of the device to update</param>
+		/// <param name="newPoolId">The new pool to assign</param>
+		/// <param name="newName">The new name to assign</param>		
+		/// <param name="newAddress">The new device address or hostname</param>
+		/// <param name="newModelId">The new model id</param>
+		/// <param name="newNotes">The devices markdown notes</param>
+		/// <param name="newEnabled">Whether the device is enabled or not</param>
+		/// <param name="newProblem">Whether to set or clear problem state</param>
+		/// <param name="newMaintenance">Whether to set or clear maintenance state</param>
+        /// <param name="modifiedByUserId">The user who is updating the device</param>
+		Task UpdateDeviceAsync(DeviceId deviceId, DevicePoolId? newPoolId, string? newName, string? newAddress, string? newModelId, string? newNotes, bool? newEnabled, bool? newProblem, bool? newMaintenance, UserId? modifiedByUserId = null);
 
 		/// <summary>
 		/// Delete a device from the collection
 		/// </summary>
-		/// <param name="DeviceId">The id of the device to delete</param>
-		Task<bool> DeleteDeviceAsync(DeviceId DeviceId);
+		/// <param name="deviceId">The id of the device to delete</param>
+		Task<bool> DeleteDeviceAsync(DeviceId deviceId);
 
 		/// <summary>
         /// Checkout or checkin the specified device
         /// </summary>
-        /// <param name="DeviceId"></param>
-        /// <param name="CheckedOutByUserId"></param>
+        /// <param name="deviceId"></param>
+        /// <param name="checkedOutByUserId"></param>
         /// <returns></returns>
-        Task CheckoutDeviceAsync(DeviceId DeviceId, UserId? CheckedOutByUserId);
+        Task CheckoutDeviceAsync(DeviceId deviceId, UserId? checkedOutByUserId);
 
         // RESERVATIONS
 
         /// <summary>
         /// Create a new reseveration in the pool with the specified devices
         /// </summary>
-        /// <param name="PoolId">The pool of devices to use for the new reservation</param>
-        /// <param name="Request">The requested devices for the reservation</param>
-        /// <param name="Hostname">The hostname of the machine making the reservation</param>
-        /// <param name="ReservationDetails">The details of the reservation</param>
-        /// <param name="JobId">The Job Id associated with the job</param>
-        /// <param name="StepId">The Step Id associated with the job</param>
-        Task<IDeviceReservation?> TryAddReservationAsync(DevicePoolId PoolId, List<DeviceRequestData> Request, string? Hostname, string? ReservationDetails, string? JobId, string? StepId);
+        /// <param name="poolId">The pool of devices to use for the new reservation</param>
+        /// <param name="request">The requested devices for the reservation</param>
+        /// <param name="hostname">The hostname of the machine making the reservation</param>
+        /// <param name="reservationDetails">The details of the reservation</param>
+        /// <param name="jobId">The Job Id associated with the job</param>
+        /// <param name="stepId">The Step Id associated with the job</param>
+        Task<IDeviceReservation?> TryAddReservationAsync(DevicePoolId poolId, List<DeviceRequestData> request, string? hostname, string? reservationDetails, string? jobId, string? stepId);
 
 		/// <summary>
 		/// Gets a reservation by guid for legacy clients
 		/// </summary>
-		/// <param name="LegacyGuid">YThe legacy guid of the reservation</param>
-		Task<IDeviceReservation?> TryGetReservationFromLegacyGuidAsync(string LegacyGuid);
+		/// <param name="legacyGuid">YThe legacy guid of the reservation</param>
+		Task<IDeviceReservation?> TryGetReservationFromLegacyGuidAsync(string legacyGuid);
 
         /// <summary>
         /// Gets a reservation by device id
         /// </summary>
-        /// <param name="Id">A device contained in reservation</param>
-        Task<IDeviceReservation?> TryGetDeviceReservationAsync(DeviceId Id);
+        /// <param name="id">A device contained in reservation</param>
+        Task<IDeviceReservation?> TryGetDeviceReservationAsync(DeviceId id);
 
         /// <summary>
         /// Get a list of all reservations
@@ -203,14 +203,14 @@ namespace Horde.Build.Collections
 		/// <summary>
 		/// Updates a reservation to the current time, for expiration
 		/// </summary>
-		/// <param name="Id">The id of the reservation to update</param>
-		public Task<bool> TryUpdateReservationAsync(ObjectId Id);
+		/// <param name="id">The id of the reservation to update</param>
+		public Task<bool> TryUpdateReservationAsync(ObjectId id);
 
 		/// <summary>
 		/// Deletes a reservation and releases reserved devices
 		/// </summary>
-		/// <param name="Id">The id of the reservation to delete</param>
-		public Task<bool> DeleteReservationAsync(ObjectId Id);
+		/// <param name="id">The id of the reservation to delete</param>
+		public Task<bool> DeleteReservationAsync(ObjectId id);
 
 		/// <summary>
 		/// Deletes expired reservations
@@ -226,7 +226,5 @@ namespace Horde.Build.Collections
 		/// Gets a list of users to notify whose device checkout is about to expire
 		/// </summary>
 		public Task<List<(UserId, IDevice)>?> ExpireNotificatonsAsync();
-
 	}
-
 }
\ No newline at end of file
