Horde: Fix coding conventions within Horde.Build.

diff --git a/Engine/Source/Programs/Horde/Horde.Build/Devices/Impl/DeviceCollection.cs b/Engine/Source/Programs/Horde/Horde.Build/Devices/Impl/DeviceCollection.cs
--- a/Engine/Source/Programs/Horde/Horde.Build/Devices/Impl/DeviceCollection.cs
+++ b/Engine/Source/Programs/Horde/Horde.Build/Devices/Impl/DeviceCollection.cs
@@ -1,26 +1,24 @@
-﻿// Copyright Epic Games, Inc. All Rights Reserved.
+// Copyright Epic Games, Inc. All Rights Reserved.
 
-using Horde.Build.Services;
-using MongoDB.Bson.Serialization.Attributes;
+using System;
+using System.Collections.Generic;
+using System.Linq;
 using System.Threading.Tasks;
-
 using Horde.Build.Acls;
 using Horde.Build.Models;
-using MongoDB.Driver;
-using System;
+using Horde.Build.Services;
 using Horde.Build.Utilities;
-using System.Linq;
-
 using MongoDB.Bson;
-using System.Collections.Generic;
+using MongoDB.Bson.Serialization.Attributes;
+using MongoDB.Driver;
 
 namespace Horde.Build.Collections.Impl
 {
 	using DeviceId = StringId<IDevice>;
 	using DevicePlatformId = StringId<IDevicePlatform>;
 	using DevicePoolId = StringId<IDevicePool>;
-	using UserId = ObjectId<IUser>;
 	using ProjectId = StringId<IProject>;
+	using UserId = ObjectId<IUser>;
 
 	/// <summary>
 	/// Collection of device documents
@@ -48,12 +46,11 @@ namespace Horde.Build.Collections.Impl
 				Name = null!;
 			}
 
-			public DevicePlatformDocument(DevicePlatformId Id, string Name)
+			public DevicePlatformDocument(DevicePlatformId id, string name)
 			{
-				this.Id = Id;
-				this.Name = Name;
+				Id = id;
+				Name = name;
 			}
-
 		}
 
 		/// <summary>
@@ -81,15 +78,13 @@ namespace Horde.Build.Collections.Impl
 			{
             }
 
-			public DevicePoolDocument(DevicePoolId Id, string Name, DevicePoolType PoolType, List<ProjectId>? ProjectIds)
+			public DevicePoolDocument(DevicePoolId id, string name, DevicePoolType poolType, List<ProjectId>? projectIds)
 			{
-				this.Id = Id;
-				this.Name = Name;
-                this.PoolType = PoolType;
-				this.ProjectIds = ProjectIds;
+				Id = id;
+				Name = name;
+                PoolType = poolType;
+				ProjectIds = projectIds;
             }
-
-
 		}
 
 		/// <summary>
@@ -150,26 +145,23 @@ namespace Horde.Build.Collections.Impl
 
 			}
 
-			public DeviceReservationDocument(ObjectId Id, DevicePoolId PoolId, List<DeviceId> Devices, List<string> RequestedDevicePlatforms, DateTime CreateTimeUtc, string? Hostname, string? ReservationDetails, string? JobId, string? StepId)
+			public DeviceReservationDocument(ObjectId id, DevicePoolId poolId, List<DeviceId> devices, List<string> requestedDevicePlatforms, DateTime createTimeUtc, string? hostname, string? reservationDetails, string? jobId, string? stepId)
 			{
-				this.Id = Id;
-				this.PoolId = PoolId;
-				this.Devices = Devices;
-				this.RequestedDevicePlatforms = RequestedDevicePlatforms;
-				this.CreateTimeUtc = CreateTimeUtc;
-				this.UpdateTimeUtc = CreateTimeUtc;
-				this.Hostname = Hostname;
-				this.ReservationDetails = ReservationDetails;
-                this.JobId = JobId;
-                this.StepId = StepId;
-
-                this.LegacyGuid = Guid.NewGuid().ToString();
+				Id = id;
+				PoolId = poolId;
+				Devices = devices;
+				RequestedDevicePlatforms = requestedDevicePlatforms;
+				CreateTimeUtc = createTimeUtc;
+				UpdateTimeUtc = createTimeUtc;
+				Hostname = hostname;
+				ReservationDetails = reservationDetails;
+                JobId = jobId;
+                StepId = stepId;
 
+                LegacyGuid = Guid.NewGuid().ToString();
 			}
-
 		}
 
-
 		/// <summary>
 		/// Concrete implementation of an device document
 		/// </summary>
@@ -228,73 +220,70 @@ namespace Horde.Build.Collections.Impl
 			[BsonConstructor]
 			private DeviceDocument()
 			{
-
 			}
 
-			public DeviceDocument(DeviceId Id, DevicePlatformId PlatformId, DevicePoolId PoolId, string Name, bool Enabled, string? Address, string? ModelId, UserId? UserId)
+			public DeviceDocument(DeviceId id, DevicePlatformId platformId, DevicePoolId poolId, string name, bool enabled, string? address, string? modelId, UserId? userId)
 			{
-				this.Id = Id;
-				this.PlatformId = PlatformId;
-				this.PoolId = PoolId;
-				this.Name = Name;
-				this.Enabled = Enabled;
-				this.Address = Address;
-				this.ModelId = ModelId;
-                this.ModifiedByUser = UserId?.ToString();
+				Id = id;
+				PlatformId = platformId;
+				PoolId = poolId;
+				Name = name;
+				Enabled = enabled;
+				Address = address;
+				ModelId = modelId;
+                ModifiedByUser = userId?.ToString();
             }
-
 		}
 
-		readonly IMongoCollection<DevicePlatformDocument> Platforms;
+		readonly IMongoCollection<DevicePlatformDocument> _platforms;
 
-		readonly IMongoCollection<DeviceDocument> Devices;
+		readonly IMongoCollection<DeviceDocument> _devices;
 
-		readonly IMongoCollection<DevicePoolDocument> Pools;
+		readonly IMongoCollection<DevicePoolDocument> _pools;
 
-		readonly IMongoCollection<DeviceReservationDocument> Reservations;
+		readonly IMongoCollection<DeviceReservationDocument> _reservations;
 
-		readonly int CheckoutDays = 7;
+		readonly int _checkoutDays = 7;
 
 		/// <summary>
 		/// Constructor
 		/// </summary>
-		public DeviceCollection(DatabaseService DatabaseService)
+		public DeviceCollection(DatabaseService databaseService)
 		{
-			Devices = DatabaseService.GetCollection<DeviceDocument>("Devices");
-			Platforms = DatabaseService.GetCollection<DevicePlatformDocument>("Devices.Platforms");
-			Pools = DatabaseService.GetCollection<DevicePoolDocument>("Devices.Pools");
-			Reservations = DatabaseService.GetCollection<DeviceReservationDocument>("Devices.Reservations");
+			_devices = databaseService.GetCollection<DeviceDocument>("Devices");
+			_platforms = databaseService.GetCollection<DevicePlatformDocument>("Devices.Platforms");
+			_pools = databaseService.GetCollection<DevicePoolDocument>("Devices.Pools");
+			_reservations = databaseService.GetCollection<DeviceReservationDocument>("Devices.Reservations");
 
-			if (!DatabaseService.ReadOnlyMode)
+			if (!databaseService.ReadOnlyMode)
 			{
-				Platforms.Indexes.CreateOne(new CreateIndexModel<DevicePlatformDocument>(Builders<DevicePlatformDocument>.IndexKeys.Ascending(x => x.Name), new CreateIndexOptions { Unique = true }));
-				Pools.Indexes.CreateOne(new CreateIndexModel<DevicePoolDocument>(Builders<DevicePoolDocument>.IndexKeys.Ascending(x => x.Name), new CreateIndexOptions { Unique = true }));
-				Devices.Indexes.CreateOne(new CreateIndexModel<DeviceDocument>(Builders<DeviceDocument>.IndexKeys.Ascending(x => x.Name), new CreateIndexOptions { Unique = true }));
+				_platforms.Indexes.CreateOne(new CreateIndexModel<DevicePlatformDocument>(Builders<DevicePlatformDocument>.IndexKeys.Ascending(x => x.Name), new CreateIndexOptions { Unique = true }));
+				_pools.Indexes.CreateOne(new CreateIndexModel<DevicePoolDocument>(Builders<DevicePoolDocument>.IndexKeys.Ascending(x => x.Name), new CreateIndexOptions { Unique = true }));
+				_devices.Indexes.CreateOne(new CreateIndexModel<DeviceDocument>(Builders<DeviceDocument>.IndexKeys.Ascending(x => x.Name), new CreateIndexOptions { Unique = true }));
 			}
 		}
 
 		/// <inheritdoc/>
-		public async Task<IDevice?> TryAddDeviceAsync(DeviceId Id, string Name, DevicePlatformId PlatformId, DevicePoolId PoolId, bool? Enabled, string? Address, string? ModelId, UserId? UserId)
+		public async Task<IDevice?> TryAddDeviceAsync(DeviceId id, string name, DevicePlatformId platformId, DevicePoolId poolId, bool? enabled, string? address, string? modelId, UserId? userId)
 		{
-			DeviceDocument NewDevice = new DeviceDocument(Id, PlatformId, PoolId, Name, Enabled ?? true, Address, ModelId, UserId);
-			await Devices.InsertOneAsync(NewDevice);
-			return NewDevice;
+			DeviceDocument newDevice = new DeviceDocument(id, platformId, poolId, name, enabled ?? true, address, modelId, userId);
+			await _devices.InsertOneAsync(newDevice);
+			return newDevice;
 		}
 
-
 		/// <inheritdoc/>
-		public async Task<IDevicePlatform?> TryAddPlatformAsync(DevicePlatformId Id, string Name)
+		public async Task<IDevicePlatform?> TryAddPlatformAsync(DevicePlatformId id, string name)
 		{
-			DevicePlatformDocument NewPlatform = new DevicePlatformDocument(Id, Name);
+			DevicePlatformDocument newPlatform = new DevicePlatformDocument(id, name);
 
 			try
 			{
-				await Platforms.InsertOneAsync(NewPlatform);
-				return NewPlatform;
+				await _platforms.InsertOneAsync(newPlatform);
+				return newPlatform;
 			}
-			catch (MongoWriteException Ex)
+			catch (MongoWriteException ex)
 			{
-				if (Ex.WriteError.Category == ServerErrorCategory.DuplicateKey)
+				if (ex.WriteError.Category == ServerErrorCategory.DuplicateKey)
 				{
 					return null;
 				}
@@ -308,25 +297,25 @@ namespace Horde.Build.Collections.Impl
 		/// <inheritdoc/>
 		public async Task<List<IDevicePlatform>> FindAllPlatformsAsync()
 		{
-			List<DevicePlatformDocument> Results = await Platforms.Find(x => true).ToListAsync();
-			return Results.OrderBy(x => x.Name).Select<DevicePlatformDocument, IDevicePlatform>(x => x).ToList();
+			List<DevicePlatformDocument> results = await _platforms.Find(x => true).ToListAsync();
+			return results.OrderBy(x => x.Name).Select<DevicePlatformDocument, IDevicePlatform>(x => x).ToList();
 		}
 
 		/// <inheritdoc/>
-		public async Task<bool> UpdatePlatformAsync(DevicePlatformId PlatformId, string[]? ModelIds)
+		public async Task<bool> UpdatePlatformAsync(DevicePlatformId platformId, string[]? modelIds)
 		{
-			UpdateDefinitionBuilder<DevicePlatformDocument> UpdateBuilder = Builders<DevicePlatformDocument>.Update;
+			UpdateDefinitionBuilder<DevicePlatformDocument> updateBuilder = Builders<DevicePlatformDocument>.Update;
 
-			List<UpdateDefinition<DevicePlatformDocument>> Updates = new List<UpdateDefinition<DevicePlatformDocument>>();
+			List<UpdateDefinition<DevicePlatformDocument>> updates = new List<UpdateDefinition<DevicePlatformDocument>>();
 
-			if (ModelIds != null)
+			if (modelIds != null)
 			{
-				Updates.Add(UpdateBuilder.Set(x => x.Models, ModelIds.ToList()));
+				updates.Add(updateBuilder.Set(x => x.Models, modelIds.ToList()));
 			}
 
-			if (Updates.Count > 0)
+			if (updates.Count > 0)
 			{
-				await Platforms.FindOneAndUpdateAsync<DevicePlatformDocument>(x => x.Id == PlatformId, UpdateBuilder.Combine(Updates));
+				await _platforms.FindOneAndUpdateAsync<DevicePlatformDocument>(x => x.Id == platformId, updateBuilder.Combine(updates));
 			}
 
 			return true;
@@ -334,18 +323,18 @@ namespace Horde.Build.Collections.Impl
 		}
 
 		/// <inheritdoc/>
-		public async Task<IDevicePool?> TryAddPoolAsync(DevicePoolId Id, string Name, DevicePoolType PoolType, List<ProjectId>? ProjectIds)
+		public async Task<IDevicePool?> TryAddPoolAsync(DevicePoolId id, string name, DevicePoolType poolType, List<ProjectId>? projectIds)
 		{
-			DevicePoolDocument NewPool = new DevicePoolDocument(Id, Name, PoolType, ProjectIds);
+			DevicePoolDocument newPool = new DevicePoolDocument(id, name, poolType, projectIds);
 
 			try
 			{
-				await Pools.InsertOneAsync(NewPool);
-				return NewPool;
+				await _pools.InsertOneAsync(newPool);
+				return newPool;
 			}
-			catch (MongoWriteException Ex)
+			catch (MongoWriteException ex)
 			{
-				if (Ex.WriteError.Category == ServerErrorCategory.DuplicateKey)
+				if (ex.WriteError.Category == ServerErrorCategory.DuplicateKey)
 				{
 					return null;
 				}
@@ -357,333 +346,329 @@ namespace Horde.Build.Collections.Impl
 		}
 
 		/// <inheritdoc/>
-		public async Task UpdatePoolAsync(DevicePoolId Id, List<ProjectId>? ProjectIds)
+		public async Task UpdatePoolAsync(DevicePoolId id, List<ProjectId>? projectIds)
 		{
-			UpdateDefinitionBuilder<DevicePoolDocument> UpdateBuilder = Builders<DevicePoolDocument>.Update;
+			UpdateDefinitionBuilder<DevicePoolDocument> updateBuilder = Builders<DevicePoolDocument>.Update;
 
-			List<UpdateDefinition<DevicePoolDocument>> Updates = new List<UpdateDefinition<DevicePoolDocument>>();
+			List<UpdateDefinition<DevicePoolDocument>> updates = new List<UpdateDefinition<DevicePoolDocument>>();
 
-			if (ProjectIds != null)
+			if (projectIds != null)
 			{
-				Updates.Add(UpdateBuilder.Set(x => x.ProjectIds, ProjectIds));
+				updates.Add(updateBuilder.Set(x => x.ProjectIds, projectIds));
 			}
 
-			if (Updates.Count > 0)
+			if (updates.Count > 0)
 			{
-				await Pools.FindOneAndUpdateAsync<DevicePoolDocument>(x => x.Id == Id, UpdateBuilder.Combine(Updates));
+				await _pools.FindOneAndUpdateAsync<DevicePoolDocument>(x => x.Id == id, updateBuilder.Combine(updates));
 			}
-
 		}
 
 		/// <inheritdoc/>
-		public async Task<List<IDevice>> FindAllDevicesAsync(List<DeviceId>? DeviceIds = null)
+		public async Task<List<IDevice>> FindAllDevicesAsync(List<DeviceId>? deviceIds = null)
 		{
-			FilterDefinition<DeviceDocument> Filter = Builders<DeviceDocument>.Filter.Empty;
-			if (DeviceIds != null)
+			FilterDefinition<DeviceDocument> filter = Builders<DeviceDocument>.Filter.Empty;
+			if (deviceIds != null)
 			{
-				Filter &= Builders<DeviceDocument>.Filter.In(x => x.Id, DeviceIds);
+				filter &= Builders<DeviceDocument>.Filter.In(x => x.Id, deviceIds);
 			}
-			List<DeviceDocument> Results = await Devices.Find(Filter).ToListAsync();
-			return Results.OrderBy(x => x.Name).Select<DeviceDocument, IDevice>(x => x).ToList();
+			List<DeviceDocument> results = await _devices.Find(filter).ToListAsync();
+			return results.OrderBy(x => x.Name).Select<DeviceDocument, IDevice>(x => x).ToList();
 		}
 
 		/// <inheritdoc/>
 		public async Task<List<IDevicePool>> FindAllPoolsAsync()
 		{
-			List<DevicePoolDocument> Results = await Pools.Find(x => true).ToListAsync();
-			return Results.OrderBy(x => x.Name).Select<DevicePoolDocument, IDevicePool>(x => x).ToList();
+			List<DevicePoolDocument> results = await _pools.Find(x => true).ToListAsync();
+			return results.OrderBy(x => x.Name).Select<DevicePoolDocument, IDevicePool>(x => x).ToList();
 		}
 
 		/// <inheritdoc/>
-		public async Task<IDevicePlatform?> GetPlatformAsync(DevicePlatformId PlatformId)
+		public async Task<IDevicePlatform?> GetPlatformAsync(DevicePlatformId platformId)
 		{
-			return await Platforms.Find<DevicePlatformDocument>(x => x.Id == PlatformId).FirstOrDefaultAsync();
+			return await _platforms.Find<DevicePlatformDocument>(x => x.Id == platformId).FirstOrDefaultAsync();
 		}
 
 		/// <inheritdoc/>
-		public async Task<IDevicePool?> GetPoolAsync(DevicePoolId PoolId)
+		public async Task<IDevicePool?> GetPoolAsync(DevicePoolId poolId)
 		{
-			return await Pools.Find<DevicePoolDocument>(x => x.Id == PoolId).FirstOrDefaultAsync();
+			return await _pools.Find<DevicePoolDocument>(x => x.Id == poolId).FirstOrDefaultAsync();
 		}
 
 		/// <inheritdoc/>
-		public async Task<IDevice?> GetDeviceAsync(DeviceId DeviceId)
+		public async Task<IDevice?> GetDeviceAsync(DeviceId deviceId)
 		{
-			return await Devices.Find<DeviceDocument>(x => x.Id == DeviceId).FirstOrDefaultAsync();
+			return await _devices.Find<DeviceDocument>(x => x.Id == deviceId).FirstOrDefaultAsync();
 		}
 
 		/// <inheritdoc/>
-		public async Task<IDevice?> GetDeviceByNameAsync(string DeviceName)
+		public async Task<IDevice?> GetDeviceByNameAsync(string deviceName)
 		{
-			return await Devices.Find<DeviceDocument>(x => x.Name == DeviceName).FirstOrDefaultAsync();
+			return await _devices.Find<DeviceDocument>(x => x.Name == deviceName).FirstOrDefaultAsync();
 		}
 
 		/// <inheritdoc/>
-		public async Task CheckoutDeviceAsync(DeviceId DeviceId, UserId? CheckedOutByUserId)
+		public async Task CheckoutDeviceAsync(DeviceId deviceId, UserId? checkedOutByUserId)
 		{
-			UpdateDefinitionBuilder<DeviceDocument> UpdateBuilder = Builders<DeviceDocument>.Update;
+			UpdateDefinitionBuilder<DeviceDocument> updateBuilder = Builders<DeviceDocument>.Update;
 
-			List<UpdateDefinition<DeviceDocument>> Updates = new List<UpdateDefinition<DeviceDocument>>();
+			List<UpdateDefinition<DeviceDocument>> updates = new List<UpdateDefinition<DeviceDocument>>();
 
-            string? UserId = CheckedOutByUserId?.ToString();
+            string? userId = checkedOutByUserId?.ToString();
 
-            Updates.Add(UpdateBuilder.Set(x => x.CheckedOutByUser, string.IsNullOrEmpty(UserId) ? null : UserId));			
+            updates.Add(updateBuilder.Set(x => x.CheckedOutByUser, String.IsNullOrEmpty(userId) ? null : userId));			
 
-			if (CheckedOutByUserId != null)
+			if (checkedOutByUserId != null)
 			{
-				Updates.Add(UpdateBuilder.Set(x => x.CheckOutTime, DateTime.UtcNow));
-				Updates.Add(UpdateBuilder.Set(x => x.CheckoutExpiringNotificationSent, false));
+				updates.Add(updateBuilder.Set(x => x.CheckOutTime, DateTime.UtcNow));
+				updates.Add(updateBuilder.Set(x => x.CheckoutExpiringNotificationSent, false));
 			}
 			else
 			{
-				Updates.Add(UpdateBuilder.Set(x => x.CheckoutExpiringNotificationSent, true));
+				updates.Add(updateBuilder.Set(x => x.CheckoutExpiringNotificationSent, true));
 			}
 
-			await Devices.FindOneAndUpdateAsync<DeviceDocument>(x => x.Id == DeviceId, UpdateBuilder.Combine(Updates));
+			await _devices.FindOneAndUpdateAsync<DeviceDocument>(x => x.Id == deviceId, updateBuilder.Combine(updates));
 		}
 
-
 		/// <inheritdoc/>
-		public async Task UpdateDeviceAsync(DeviceId DeviceId, DevicePoolId? NewPoolId, string? NewName, string? NewAddress, string? NewModelId, string? NewNotes, bool? NewEnabled, bool? NewProblem, bool? NewMaintenance, UserId? ModifiedByUserId = null)
+		public async Task UpdateDeviceAsync(DeviceId deviceId, DevicePoolId? newPoolId, string? newName, string? newAddress, string? newModelId, string? newNotes, bool? newEnabled, bool? newProblem, bool? newMaintenance, UserId? modifiedByUserId = null)
 		{
-			UpdateDefinitionBuilder<DeviceDocument> UpdateBuilder = Builders<DeviceDocument>.Update;
+			UpdateDefinitionBuilder<DeviceDocument> updateBuilder = Builders<DeviceDocument>.Update;
 
-			List<UpdateDefinition<DeviceDocument>> Updates = new List<UpdateDefinition<DeviceDocument>>();
+			List<UpdateDefinition<DeviceDocument>> updates = new List<UpdateDefinition<DeviceDocument>>();
 
-			DateTime UtcNow = DateTime.UtcNow;
+			DateTime utcNow = DateTime.UtcNow;
 
-			if (ModifiedByUserId != null)
+			if (modifiedByUserId != null)
 			{
-				Updates.Add(UpdateBuilder.Set(x => x.ModifiedByUser, ModifiedByUserId.ToString()));
+				updates.Add(updateBuilder.Set(x => x.ModifiedByUser, modifiedByUserId.ToString()));
 			}
 
-			if (NewPoolId != null)
+			if (newPoolId != null)
 			{
-				Updates.Add(UpdateBuilder.Set(x => x.PoolId, NewPoolId.Value));
+				updates.Add(updateBuilder.Set(x => x.PoolId, newPoolId.Value));
 			}
 
-			if (NewName != null)
+			if (newName != null)
 			{
-				Updates.Add(UpdateBuilder.Set(x => x.Name, NewName));
+				updates.Add(updateBuilder.Set(x => x.Name, newName));
 			}
 
-			if (NewEnabled != null)
+			if (newEnabled != null)
 			{
-				Updates.Add(UpdateBuilder.Set(x => x.Enabled, NewEnabled));
+				updates.Add(updateBuilder.Set(x => x.Enabled, newEnabled));
 			}
 
-			if (NewAddress != null)
+			if (newAddress != null)
 			{
-				Updates.Add(UpdateBuilder.Set(x => x.Address, NewAddress));
+				updates.Add(updateBuilder.Set(x => x.Address, newAddress));
 			}
 
-			if (!string.IsNullOrEmpty(NewModelId))
+			if (!String.IsNullOrEmpty(newModelId))
 			{
-				if (NewModelId == "Base")
+				if (newModelId == "Base")
 				{
-					Updates.Add(UpdateBuilder.Set(x => x.ModelId, null));
+					updates.Add(updateBuilder.Set(x => x.ModelId, null));
 				}
 				else
 				{
-					Updates.Add(UpdateBuilder.Set(x => x.ModelId, NewModelId));
+					updates.Add(updateBuilder.Set(x => x.ModelId, newModelId));
 				}				
 			}
 
-			if (NewNotes != null)
+			if (newNotes != null)
 			{
-				Updates.Add(UpdateBuilder.Set(x => x.Notes, NewNotes));
+				updates.Add(updateBuilder.Set(x => x.Notes, newNotes));
 			}
 
-			if (NewProblem.HasValue)
+			if (newProblem.HasValue)
 			{
-				DateTime? ProblemTime = null;
-				if (NewProblem.Value)
+				DateTime? problemTime = null;
+				if (newProblem.Value)
 				{
-					ProblemTime = UtcNow;
+					problemTime = utcNow;
 				}
 
-				Updates.Add(UpdateBuilder.Set(x => x.ProblemTimeUtc, ProblemTime));
+				updates.Add(updateBuilder.Set(x => x.ProblemTimeUtc, problemTime));
 			}
 
-			if (NewMaintenance.HasValue)
+			if (newMaintenance.HasValue)
 			{
-				DateTime? MaintenanceTime = null;
-				if (NewMaintenance.Value)
+				DateTime? maintenanceTime = null;
+				if (newMaintenance.Value)
 				{
-					MaintenanceTime = UtcNow;
+					maintenanceTime = utcNow;
 				}
 
-				Updates.Add(UpdateBuilder.Set(x => x.MaintenanceTimeUtc, MaintenanceTime));
+				updates.Add(updateBuilder.Set(x => x.MaintenanceTimeUtc, maintenanceTime));
 			}
 
-
-			if (Updates.Count > 0)
+			if (updates.Count > 0)
 			{
-				await Devices.FindOneAndUpdateAsync<DeviceDocument>(x => x.Id == DeviceId, UpdateBuilder.Combine(Updates));
+				await _devices.FindOneAndUpdateAsync<DeviceDocument>(x => x.Id == deviceId, updateBuilder.Combine(updates));
 			}
 		}
 
 		/// <inheritdoc/>
-		public async Task<bool> DeleteDeviceAsync(DeviceId DeviceId)
+		public async Task<bool> DeleteDeviceAsync(DeviceId deviceId)
 		{
-			FilterDefinition<DeviceDocument> Filter = Builders<DeviceDocument>.Filter.Eq(x => x.Id, DeviceId);
-			DeleteResult Result = await Devices.DeleteOneAsync(Filter);
-			return Result.DeletedCount > 0;
+			FilterDefinition<DeviceDocument> filter = Builders<DeviceDocument>.Filter.Eq(x => x.Id, deviceId);
+			DeleteResult result = await _devices.DeleteOneAsync(filter);
+			return result.DeletedCount > 0;
 
 		}
 
 		/// <inheritdoc/>
-		public async Task<List<IDeviceReservation>> FindAllDeviceReservationsAsync(DevicePoolId? PoolId = null)
+		public async Task<List<IDeviceReservation>> FindAllDeviceReservationsAsync(DevicePoolId? poolId = null)
 		{
-			return await Reservations.Find(A => PoolId == null || A.PoolId == PoolId).ToListAsync<DeviceReservationDocument, IDeviceReservation>();
+			return await _reservations.Find(a => poolId == null || a.PoolId == poolId).ToListAsync<DeviceReservationDocument, IDeviceReservation>();
 		}
 
 		/// <inheritdoc/>
-		public async Task<IDeviceReservation?> TryAddReservationAsync(DevicePoolId PoolId, List<DeviceRequestData> Request, string? Hostname, string? ReservationDetails, string? JobId, string? StepId)
+		public async Task<IDeviceReservation?> TryAddReservationAsync(DevicePoolId poolId, List<DeviceRequestData> request, string? hostname, string? reservationDetails, string? jobId, string? stepId)
 		{
 
-			if (Request.Count == 0)
+			if (request.Count == 0)
 			{
 				return null;
 			}
 
-			DevicePoolDocument? Pool = await Pools.Find<DevicePoolDocument>(x => x.Id == PoolId).FirstOrDefaultAsync();
+			DevicePoolDocument? pool = await _pools.Find<DevicePoolDocument>(x => x.Id == poolId).FirstOrDefaultAsync();
 
-			if (Pool == null || Pool.PoolType != DevicePoolType.Automation)
+			if (pool == null || pool.PoolType != DevicePoolType.Automation)
 			{
                 return null;
             }
 			
-			HashSet<DeviceId> Allocated = new HashSet<DeviceId>();
-			Dictionary<DeviceId, string> PlatformRequestMap = new Dictionary<DeviceId, string>();
+			HashSet<DeviceId> allocated = new HashSet<DeviceId>();
+			Dictionary<DeviceId, string> platformRequestMap = new Dictionary<DeviceId, string>();
 
-			List<DeviceReservationDocument> PoolReservations = await Reservations.Find(x => x.PoolId == PoolId).ToListAsync();
+			List<DeviceReservationDocument> poolReservations = await _reservations.Find(x => x.PoolId == poolId).ToListAsync();
 
-			DateTime ReservationTimeUtc = DateTime.UtcNow;
+			DateTime reservationTimeUtc = DateTime.UtcNow;
 
 			// Get available devices
-			List<DeviceDocument> PoolDevices = await Devices.Find(x =>
-				x.PoolId == PoolId &&
+			List<DeviceDocument> poolDevices = await _devices.Find(x =>
+				x.PoolId == poolId &&
 				x.Enabled &&
 				x.MaintenanceTimeUtc == null).ToListAsync();
 
 			// filter out problem devices
-			PoolDevices = PoolDevices.FindAll(x => (x.ProblemTimeUtc == null || ((ReservationTimeUtc - x.ProblemTimeUtc).Value.TotalMinutes > 30)));
+			poolDevices = poolDevices.FindAll(x => (x.ProblemTimeUtc == null || ((reservationTimeUtc - x.ProblemTimeUtc).Value.TotalMinutes > 30)));
 
 			// filter out currently reserved devices
-			PoolDevices = PoolDevices.FindAll(x => PoolReservations.FirstOrDefault(p => p.Devices.Contains(x.Id)) == null);
+			poolDevices = poolDevices.FindAll(x => poolReservations.FirstOrDefault(p => p.Devices.Contains(x.Id)) == null);
 
 			// sort to use last reserved first to cycle devices
-			PoolDevices.Sort((A, B) =>
+			poolDevices.Sort((a, b) =>
 			{
-				DateTime? ATime = A.ReservationTimeUtc;
-				DateTime? BTime = B.ReservationTimeUtc;
+				DateTime? aTime = a.ReservationTimeUtc;
+				DateTime? bTime = b.ReservationTimeUtc;
 
-				if (ATime == BTime)
+				if (aTime == bTime)
 				{
 					return 0;
 				}
 
-				if (ATime == null && BTime != null)
+				if (aTime == null && bTime != null)
 				{
 					return -1;
 				}
-				if (ATime != null && BTime == null)
+				if (aTime != null && bTime == null)
 				{
 					return 1;
 				}
 
-				return ATime < BTime ? -1 : 1;
+				return aTime < bTime ? -1 : 1;
 
 			});
 
-			foreach (DeviceRequestData Data in Request)
+			foreach (DeviceRequestData data in request)
 			{
-				DeviceDocument? Device = PoolDevices.FirstOrDefault(A =>
+				DeviceDocument? device = poolDevices.FirstOrDefault(a =>
 				{
 
-					if (Allocated.Contains(A.Id) || A.PlatformId != Data.PlatformId)
+					if (allocated.Contains(a.Id) || a.PlatformId != data.PlatformId)
 					{
 						return false;
 					}
 
-					if (Data.IncludeModels.Count > 0 && (A.ModelId == null || !Data.IncludeModels.Contains(A.ModelId)))
+					if (data.IncludeModels.Count > 0 && (a.ModelId == null || !data.IncludeModels.Contains(a.ModelId)))
 					{
 						return false;
 					}
 
-					if (Data.ExcludeModels.Count > 0 && (A.ModelId != null && Data.ExcludeModels.Contains(A.ModelId)))
+					if (data.ExcludeModels.Count > 0 && (a.ModelId != null && data.ExcludeModels.Contains(a.ModelId)))
 					{
 						return false;
 					}
 
-
 					return true;
 				});
 
-				if (Device == null)
+				if (device == null)
 				{
 					// can't fulfill request
 					return null;
 				}
 
-				Allocated.Add(Device.Id);
-				PlatformRequestMap.Add(Device.Id, Data.RequestedPlatform);
+				allocated.Add(device.Id);
+				platformRequestMap.Add(device.Id, data.RequestedPlatform);
 			}
 
 			// update reservation time and utilization for allocated devices 
-			foreach (DeviceId Id in Allocated)
+			foreach (DeviceId id in allocated)
 			{
-				DeviceDocument Device = PoolDevices.First((Device) => Device.Id == Id);
+				DeviceDocument device = poolDevices.First((device) => device.Id == id);
 
-				List<DeviceUtilizationTelemetry>? Utilization = Device.Utilization;
+				List<DeviceUtilizationTelemetry>? utilization = device.Utilization;
 
-				if (Utilization == null)
+				if (utilization == null)
 				{
-					Utilization = new List<DeviceUtilizationTelemetry>();
+					utilization = new List<DeviceUtilizationTelemetry>();
 				}
 
 				// keep up to 100, maintaining order
-				if (Utilization.Count > 99)
+				if (utilization.Count > 99)
 				{
-					Utilization = Utilization.GetRange(0, 99);
+					utilization = utilization.GetRange(0, 99);
 				}
 
-				Utilization.Insert(0, new DeviceUtilizationTelemetry(ReservationTimeUtc) { JobId = JobId, StepId = StepId });
+				utilization.Insert(0, new DeviceUtilizationTelemetry(reservationTimeUtc) { JobId = jobId, StepId = stepId });
 
-				UpdateDefinitionBuilder<DeviceDocument> DeviceBuilder = Builders<DeviceDocument>.Update;
-				List<UpdateDefinition<DeviceDocument>> DeviceUpdates = new List<UpdateDefinition<DeviceDocument>>();
+				UpdateDefinitionBuilder<DeviceDocument> deviceBuilder = Builders<DeviceDocument>.Update;
+				List<UpdateDefinition<DeviceDocument>> deviceUpdates = new List<UpdateDefinition<DeviceDocument>>();
 
-				DeviceUpdates.Add(DeviceBuilder.Set(x => x.ReservationTimeUtc, ReservationTimeUtc));
-				DeviceUpdates.Add(DeviceBuilder.Set(x => x.Utilization, Utilization));
+				deviceUpdates.Add(deviceBuilder.Set(x => x.ReservationTimeUtc, reservationTimeUtc));
+				deviceUpdates.Add(deviceBuilder.Set(x => x.Utilization, utilization));
 
-				await Devices.FindOneAndUpdateAsync<DeviceDocument>(x => x.Id == Id, DeviceBuilder.Combine(DeviceUpdates));
+				await _devices.FindOneAndUpdateAsync<DeviceDocument>(x => x.Id == id, deviceBuilder.Combine(deviceUpdates));
 			}
 
-			List<DeviceId> DeviceIds = Allocated.ToList();
-			List<string> RequestedPlatforms = DeviceIds.Select(x => PlatformRequestMap[x]).ToList();
+			List<DeviceId> deviceIds = allocated.ToList();
+			List<string> requestedPlatforms = deviceIds.Select(x => platformRequestMap[x]).ToList();
 
 			// Create new reservation
-			DeviceReservationDocument NewReservation = new DeviceReservationDocument(ObjectId.GenerateNewId(), PoolId, DeviceIds, RequestedPlatforms, ReservationTimeUtc, Hostname, ReservationDetails, JobId, StepId);
-			await Reservations.InsertOneAsync(NewReservation);
+			DeviceReservationDocument newReservation = new DeviceReservationDocument(ObjectId.GenerateNewId(), poolId, deviceIds, requestedPlatforms, reservationTimeUtc, hostname, reservationDetails, jobId, stepId);
+			await _reservations.InsertOneAsync(newReservation);
 
-			return NewReservation;
+			return newReservation;
 
 		}
 
 		/// <inheritdoc/>
-		public async Task<bool> TryUpdateReservationAsync(ObjectId Id)
+		public async Task<bool> TryUpdateReservationAsync(ObjectId id)
 		{
-			UpdateResult Result = await Reservations.UpdateOneAsync(x => x.Id == Id, Builders<DeviceReservationDocument>.Update.Set(x => x.UpdateTimeUtc, DateTime.UtcNow));
-			return Result.ModifiedCount == 1;
+			UpdateResult result = await _reservations.UpdateOneAsync(x => x.Id == id, Builders<DeviceReservationDocument>.Update.Set(x => x.UpdateTimeUtc, DateTime.UtcNow));
+			return result.ModifiedCount == 1;
 		}
 
 		/// <inheritdoc/>
-		public async Task<bool> DeleteReservationAsync(ObjectId Id)
+		public async Task<bool> DeleteReservationAsync(ObjectId id)
 		{
-			FilterDefinition<DeviceReservationDocument> Filter = Builders<DeviceReservationDocument>.Filter.Eq(x => x.Id, Id);
-			DeleteResult Result = await Reservations.DeleteOneAsync(Filter);
-			return Result.DeletedCount > 0;
+			FilterDefinition<DeviceReservationDocument> filter = Builders<DeviceReservationDocument>.Filter.Eq(x => x.Id, id);
+			DeleteResult result = await _reservations.DeleteOneAsync(filter);
+			return result.DeletedCount > 0;
 		}
 
 		/// <summary>
@@ -691,16 +676,16 @@ namespace Horde.Build.Collections.Impl
 		/// </summary>
 		public async Task<bool> ExpireReservationsAsync()
 		{
-			List<IDeviceReservation> Reserves = await Reservations.Find(A => true).ToListAsync<DeviceReservationDocument, IDeviceReservation>();
+			List<IDeviceReservation> reserves = await _reservations.Find(a => true).ToListAsync<DeviceReservationDocument, IDeviceReservation>();
 
-			DateTime UtcNow = DateTime.UtcNow;
+			DateTime utcNow = DateTime.UtcNow;
 
-			Reserves = Reserves.FindAll(R => (UtcNow - R.UpdateTimeUtc).TotalMinutes > 10).ToList();
+			reserves = reserves.FindAll(r => (utcNow - r.UpdateTimeUtc).TotalMinutes > 10).ToList();
 
-			if (Reserves.Count > 0)
+			if (reserves.Count > 0)
 			{
-				DeleteResult Result = await Reservations.DeleteManyAsync(Builders<DeviceReservationDocument>.Filter.In(x => x.Id, Reserves.Select(y => y.Id)));
-				return (Result.DeletedCount > 0);
+				DeleteResult result = await _reservations.DeleteManyAsync(Builders<DeviceReservationDocument>.Filter.In(x => x.Id, reserves.Select(y => y.Id)));
+				return (result.DeletedCount > 0);
 			}
 
 			return false;
@@ -712,33 +697,32 @@ namespace Horde.Build.Collections.Impl
 		/// </summary>
 		public async Task<List<(UserId, IDevice)>?> ExpireCheckedOutAsync()
 		{
-			FilterDefinition<DeviceDocument> Filter = Builders<DeviceDocument>.Filter.Empty;
-			Filter &= Builders<DeviceDocument>.Filter.Where(x => x.CheckedOutByUser != null && x.CheckOutTime != null);
+			FilterDefinition<DeviceDocument> filter = Builders<DeviceDocument>.Filter.Empty;
+			filter &= Builders<DeviceDocument>.Filter.Where(x => x.CheckedOutByUser != null && x.CheckOutTime != null);
 
-			List<DeviceDocument> CheckedOutDevices = await Devices.Find(Filter).ToListAsync();
+			List<DeviceDocument> checkedOutDevices = await _devices.Find(filter).ToListAsync();
 
-			DateTime UtcNow = DateTime.UtcNow;
-			List<DeviceDocument> ExpiredDevices = CheckedOutDevices.FindAll(x => (UtcNow - x.CheckOutTime!.Value).TotalDays >= CheckoutDays).ToList();
+			DateTime utcNow = DateTime.UtcNow;
+			List<DeviceDocument> expiredDevices = checkedOutDevices.FindAll(x => (utcNow - x.CheckOutTime!.Value).TotalDays >= _checkoutDays).ToList();
 
-			if (ExpiredDevices.Count > 0)
+			if (expiredDevices.Count > 0)
 			{
 
-				FilterDefinition<DeviceDocument> UpdateFilter = Builders<DeviceDocument>.Filter.In(x => x.Id, ExpiredDevices.Select(y => y.Id));
+				FilterDefinition<DeviceDocument> updateFilter = Builders<DeviceDocument>.Filter.In(x => x.Id, expiredDevices.Select(y => y.Id));
 
-				UpdateDefinitionBuilder<DeviceDocument> DeviceBuilder = Builders<DeviceDocument>.Update;
-				List<UpdateDefinition<DeviceDocument>> DeviceUpdates = new List<UpdateDefinition<DeviceDocument>>();
+				UpdateDefinitionBuilder<DeviceDocument> deviceBuilder = Builders<DeviceDocument>.Update;
+				List<UpdateDefinition<DeviceDocument>> deviceUpdates = new List<UpdateDefinition<DeviceDocument>>();
 
-				DeviceUpdates.Add(DeviceBuilder.Set(x => x.CheckedOutByUser, null));
-				DeviceUpdates.Add(DeviceBuilder.Set(x => x.CheckOutTime, null));
-				DeviceUpdates.Add(DeviceBuilder.Set(x => x.CheckoutExpiringNotificationSent, true));
+				deviceUpdates.Add(deviceBuilder.Set(x => x.CheckedOutByUser, null));
+				deviceUpdates.Add(deviceBuilder.Set(x => x.CheckOutTime, null));
+				deviceUpdates.Add(deviceBuilder.Set(x => x.CheckoutExpiringNotificationSent, true));
 
-				UpdateResult Result = await Devices.UpdateManyAsync(Builders<DeviceDocument>.Filter.In(x => x.Id, ExpiredDevices.Select(y => y.Id)), DeviceBuilder.Combine(DeviceUpdates));
+				UpdateResult result = await _devices.UpdateManyAsync(Builders<DeviceDocument>.Filter.In(x => x.Id, expiredDevices.Select(y => y.Id)), deviceBuilder.Combine(deviceUpdates));
 
-				if (Result.ModifiedCount > 0)
+				if (result.ModifiedCount > 0)
 				{
-					return ExpiredDevices.Select(x => (UserId.Parse(x.CheckedOutByUser!), (IDevice) x)).ToList();
+					return expiredDevices.Select(x => (UserId.Parse(x.CheckedOutByUser!), (IDevice) x)).ToList();
 				}
-
 			}
 
 			return null;
@@ -750,59 +734,52 @@ namespace Horde.Build.Collections.Impl
 		/// </summary>
 		public async Task<List<(UserId, IDevice)>?> ExpireNotificatonsAsync()
 		{
-			FilterDefinition<DeviceDocument> Filter = Builders<DeviceDocument>.Filter.Empty;
-			Filter &= Builders<DeviceDocument>.Filter.Where(x => x.CheckedOutByUser != null && x.CheckoutExpiringNotificationSent != true);
+			FilterDefinition<DeviceDocument> filter = Builders<DeviceDocument>.Filter.Empty;
+			filter &= Builders<DeviceDocument>.Filter.Where(x => x.CheckedOutByUser != null && x.CheckoutExpiringNotificationSent != true);
 
-			List<DeviceDocument> CheckedOutDevices = await Devices.Find(Filter).ToListAsync();
+			List<DeviceDocument> checkedOutDevices = await _devices.Find(filter).ToListAsync();
 
-			DateTime UtcNow = DateTime.UtcNow;
-			List<DeviceDocument> ExpiredDevices = CheckedOutDevices.FindAll(x => (UtcNow - x.CheckOutTime!.Value).TotalDays >= (CheckoutDays - 1)).ToList();
+			DateTime utcNow = DateTime.UtcNow;
+			List<DeviceDocument> expiredDevices = checkedOutDevices.FindAll(x => (utcNow - x.CheckOutTime!.Value).TotalDays >= (_checkoutDays - 1)).ToList();
 
-			if (ExpiredDevices.Count > 0)
+			if (expiredDevices.Count > 0)
 			{
-				FilterDefinition<DeviceDocument> UpdateFilter = Builders<DeviceDocument>.Filter.In(x => x.Id, ExpiredDevices.Select(y => y.Id));
+				FilterDefinition<DeviceDocument> updateFilter = Builders<DeviceDocument>.Filter.In(x => x.Id, expiredDevices.Select(y => y.Id));
 
-				UpdateDefinitionBuilder<DeviceDocument> DeviceBuilder = Builders<DeviceDocument>.Update;
-				List<UpdateDefinition<DeviceDocument>> DeviceUpdates = new List<UpdateDefinition<DeviceDocument>>();
+				UpdateDefinitionBuilder<DeviceDocument> deviceBuilder = Builders<DeviceDocument>.Update;
+				List<UpdateDefinition<DeviceDocument>> deviceUpdates = new List<UpdateDefinition<DeviceDocument>>();
 
-				DeviceUpdates.Add(DeviceBuilder.Set(x => x.CheckoutExpiringNotificationSent, true));
+				deviceUpdates.Add(deviceBuilder.Set(x => x.CheckoutExpiringNotificationSent, true));
 
-				UpdateResult Result = await Devices.UpdateManyAsync(Builders<DeviceDocument>.Filter.In(x => x.Id, ExpiredDevices.Select(y => y.Id)), DeviceBuilder.Combine(DeviceUpdates));
+				UpdateResult result = await _devices.UpdateManyAsync(Builders<DeviceDocument>.Filter.In(x => x.Id, expiredDevices.Select(y => y.Id)), deviceBuilder.Combine(deviceUpdates));
 
-				if (Result.ModifiedCount > 0)
+				if (result.ModifiedCount > 0)
 				{
-					return ExpiredDevices.Select(x => (UserId.Parse(x.CheckedOutByUser!), (IDevice)x)).ToList();
+					return expiredDevices.Select(x => (UserId.Parse(x.CheckedOutByUser!), (IDevice)x)).ToList();
 				}
-
 			}
 
 			return null;
-
 		}
 
-
 		/// <inheritdoc/>
 		public async Task<List<IDeviceReservation>> FindAllReservationsAsync()
 		{
-			List<DeviceReservationDocument> Results = await Reservations.Find(x => true).ToListAsync();
-			return Results.OrderBy(x => x.CreateTimeUtc.Ticks).Select<DeviceReservationDocument, IDeviceReservation>(x => x).ToList();
+			List<DeviceReservationDocument> results = await _reservations.Find(x => true).ToListAsync();
+			return results.OrderBy(x => x.CreateTimeUtc.Ticks).Select<DeviceReservationDocument, IDeviceReservation>(x => x).ToList();
 		}
 
-
 		/// <inheritdoc/>
-		public async Task<IDeviceReservation?> TryGetReservationFromLegacyGuidAsync(string LegacyGuid)
+		public async Task<IDeviceReservation?> TryGetReservationFromLegacyGuidAsync(string legacyGuid)
 		{
-			return await Reservations.Find<DeviceReservationDocument>(R => R.LegacyGuid == LegacyGuid).FirstOrDefaultAsync();
+			return await _reservations.Find<DeviceReservationDocument>(r => r.LegacyGuid == legacyGuid).FirstOrDefaultAsync();
 		}
 
 		/// <inheritdoc/>
-		public async Task<IDeviceReservation?> TryGetDeviceReservationAsync(DeviceId Id)
+		public async Task<IDeviceReservation?> TryGetDeviceReservationAsync(DeviceId id)
 		{
-			List<DeviceReservationDocument> Results = await Reservations.Find(x => true).ToListAsync();
-			return Results.FirstOrDefault(R => R.Devices.Contains(Id));
+			List<DeviceReservationDocument> results = await _reservations.Find(x => true).ToListAsync();
+			return results.FirstOrDefault(r => r.Devices.Contains(id));
         }
-
 	}
-
 }
-
diff --git a/Engine/Source/Programs/Horde/Horde.Build/Devices/Impl/DeviceService.cs b/Engine/Source/Programs/Horde/Horde.Build/Devices/Impl/DeviceService.cs
--- a/Engine/Source/Programs/Horde/Horde.Build/Devices/Impl/DeviceService.cs
+++ b/Engine/Source/Programs/Horde/Horde.Build/Devices/Impl/DeviceService.cs
@@ -1,24 +1,21 @@
-﻿// Copyright Epic Games, Inc. All Rights Reserved.
+// Copyright Epic Games, Inc. All Rights Reserved.
 
+using System;
+using System.Collections.Generic;
+using System.Linq;
+using System.Security.Claims;
+using System.Threading;
 using System.Threading.Tasks;
-
 using Horde.Build.Acls;
-using Horde.Build.Models;
+using Horde.Build.Api;
 using Horde.Build.Collections;
-using System;
+using Horde.Build.Models;
+using Horde.Build.Notifications;
 using Horde.Build.Utilities;
-using System.Threading;
-
-using Microsoft.Extensions.Logging;
 using HordeCommon;
-using MongoDB.Bson;
-using System.Collections.Generic;
-using System.Security.Claims;
-using System.Diagnostics.CodeAnalysis;
-using Horde.Build.Notifications;
-using System.Linq;
-using Horde.Build.Api;
 using Microsoft.Extensions.Hosting;
+using Microsoft.Extensions.Logging;
+using MongoDB.Bson;
 
 namespace Horde.Build.Services
 {
@@ -26,8 +23,8 @@ namespace Horde.Build.Services
 	using DevicePlatformId = StringId<IDevicePlatform>;
 	using DevicePoolId = StringId<IDevicePool>;
 	using JobId = ObjectId<IJob>;
-	using UserId = ObjectId<IUser>;
 	using ProjectId = StringId<IProject>;
+	using UserId = ObjectId<IUser>;
 
 	/// <summary>
 	///  Device Pool Authorization (convenience class for pool ACL's)
@@ -52,14 +49,14 @@ namespace Horde.Build.Services
 		/// <summary>
 		/// Constructor
 		/// </summary>
-		/// <param name="Pool"></param>
-		/// <param name="Read"></param>
-		/// <param name="Write"></param>
-		public DevicePoolAuthorization(IDevicePool Pool, bool Read, bool Write)
+		/// <param name="pool"></param>
+		/// <param name="read"></param>
+		/// <param name="write"></param>
+		public DevicePoolAuthorization(IDevicePool pool, bool read, bool write)
 		{
-			this.Pool = Pool;
-			this.Read = Read;
-			this.Write = Write;				
+			Pool = pool;
+			Read = read;
+			Write = write;				
 		}
 	}
 
@@ -86,7 +83,6 @@ namespace Horde.Build.Services
 
 	}
 
-
 	/// <summary>
 	/// Device management service
 	/// </summary>
@@ -95,27 +91,27 @@ namespace Horde.Build.Services
 		/// <summary>
 		/// The ACL service instance
 		/// </summary>
-		AclService AclService;
+		readonly AclService _aclService;
 
 		/// <summary>
 		/// Instance of the notification service
 		/// </summary>
-		INotificationService NotificationService;
+		readonly INotificationService _notificationService;
 
 		/// <summary>
 		/// Singleton instance of the job service
 		/// </summary>
-		JobService JobService;
+		readonly JobService _jobService;
 
 		/// <summary>
 		/// Singleton instance of the stream service
 		/// </summary>
-		StreamService StreamService;
+		readonly StreamService _streamService;
 
 		/// <summary>
 		/// Singleton instance of the project service
 		/// </summary>
-		ProjectService ProjectService;
+		readonly ProjectService _projectService;
 
 		/// <summary>
 		/// The user collection instance
@@ -125,85 +121,83 @@ namespace Horde.Build.Services
 		/// <summary>
 		/// Log output writer
 		/// </summary>
-		ILogger<DeviceService> Logger;
+		readonly ILogger<DeviceService> _logger;
 
 		/// <summary>
 		/// Device collection
 		/// </summary>
-		IDeviceCollection Devices;
-
-		ITicker Ticker;
+		readonly IDeviceCollection _devices;
+		readonly ITicker _ticker;
 
 		/// <summary>
 		/// Platform map V1 singleton
 		/// </summary>
-		ISingletonDocument<DevicePlatformMapV1> PlatformMapSingleton;
+		readonly ISingletonDocument<DevicePlatformMapV1> _platformMapSingleton;
 
 		/// <summary>
 		/// Device service constructor
 		/// </summary>
-		public DeviceService(IDeviceCollection Devices, ISingletonDocument<DevicePlatformMapV1> PlatformMapSingleton, IUserCollection UserCollection, JobService JobService, ProjectService ProjectService, StreamService StreamService, AclService AclService, INotificationService NotificationService, IClock Clock, ILogger<DeviceService> Logger)
+		public DeviceService(IDeviceCollection devices, ISingletonDocument<DevicePlatformMapV1> platformMapSingleton, IUserCollection userCollection, JobService jobService, ProjectService projectService, StreamService streamService, AclService aclService, INotificationService notificationService, IClock clock, ILogger<DeviceService> logger)
 		{
-			this.UserCollection = UserCollection;
-			this.Devices = Devices;
-            this.JobService = JobService;
-			this.ProjectService = ProjectService;
-            this.StreamService = StreamService;
-            this.AclService = AclService;
-            this.NotificationService = NotificationService;
-			this.Ticker = Clock.AddSharedTicker<DeviceService>(TimeSpan.FromMinutes(1.0), TickAsync, Logger);
-            this.Logger = Logger;
+			UserCollection = userCollection;
+			_devices = devices;
+            _jobService = jobService;
+			_projectService = projectService;
+            _streamService = streamService;
+            _aclService = aclService;
+            _notificationService = notificationService;
+			_ticker = clock.AddSharedTicker<DeviceService>(TimeSpan.FromMinutes(1.0), TickAsync, logger);
+            _logger = logger;
 			
-			this.PlatformMapSingleton = PlatformMapSingleton;
+			_platformMapSingleton = platformMapSingleton;
 
 		}
 
 		/// <inheritdoc/>
-		public Task StartAsync(CancellationToken cancellationToken) => Ticker.StartAsync();
+		public Task StartAsync(CancellationToken cancellationToken) => _ticker.StartAsync();
 
 		/// <inheritdoc/>
-		public Task StopAsync(CancellationToken cancellationToken) => Ticker.StopAsync();
+		public Task StopAsync(CancellationToken cancellationToken) => _ticker.StopAsync();
 
 		/// <inheritdoc/>
-		public void Dispose() => Ticker.Dispose();
+		public void Dispose() => _ticker.Dispose();
 
 		/// <summary>
 		/// Ticks service
 		/// </summary>
-		async ValueTask TickAsync(CancellationToken StoppingToken)
+		async ValueTask TickAsync(CancellationToken stoppingToken)
 		{
 
-			if (!StoppingToken.IsCancellationRequested)
+			if (!stoppingToken.IsCancellationRequested)
 			{
-				await Devices.ExpireReservationsAsync();
+				await _devices.ExpireReservationsAsync();
 
-				List<(UserId, IDevice)>? ExpireNotifications = await Devices.ExpireNotificatonsAsync();
-				if (ExpireNotifications != null && ExpireNotifications.Count > 0)
+				List<(UserId, IDevice)>? expireNotifications = await _devices.ExpireNotificatonsAsync();
+				if (expireNotifications != null && expireNotifications.Count > 0)
 				{
-					foreach ((UserId, IDevice) ExpiredDevice in ExpireNotifications)
+					foreach ((UserId, IDevice) expiredDevice in expireNotifications)
 					{
-						await NotifyDeviceServiceAsync($"Device {ExpiredDevice.Item2.PlatformId.ToString().ToUpperInvariant()} / {ExpiredDevice.Item2.Name} checkout will expire in 24 hours.  Please visit https://horde.devtools.epicgames.com/devices to renew the checkout if needed.", null, null, null, ExpiredDevice.Item1);
+						await NotifyDeviceServiceAsync($"Device {expiredDevice.Item2.PlatformId.ToString().ToUpperInvariant()} / {expiredDevice.Item2.Name} checkout will expire in 24 hours.  Please visit https://horde.devtools.epicgames.com/devices to renew the checkout if needed.", null, null, null, expiredDevice.Item1);
 					}
 				}
 
-				List<(UserId, IDevice)>? ExpireCheckouts = await Devices.ExpireCheckedOutAsync();
-				if (ExpireCheckouts != null && ExpireCheckouts.Count > 0)
+				List<(UserId, IDevice)>? expireCheckouts = await _devices.ExpireCheckedOutAsync();
+				if (expireCheckouts != null && expireCheckouts.Count > 0)
 				{
-					foreach ((UserId, IDevice) ExpiredDevice in ExpireCheckouts)
+					foreach ((UserId, IDevice) expiredDevice in expireCheckouts)
 					{
-						await NotifyDeviceServiceAsync($"Device {ExpiredDevice.Item2.PlatformId.ToString().ToUpperInvariant()} / {ExpiredDevice.Item2.Name} checkout has expired.  The device has been returned to the shared pool and should no longer be accessed.  Please visit https://horde.devtools.epicgames.com/devices to checkout devices as needed.", null, null, null, ExpiredDevice.Item1);
+						await NotifyDeviceServiceAsync($"Device {expiredDevice.Item2.PlatformId.ToString().ToUpperInvariant()} / {expiredDevice.Item2.Name} checkout has expired.  The device has been returned to the shared pool and should no longer be accessed.  Please visit https://horde.devtools.epicgames.com/devices to checkout devices as needed.", null, null, null, expiredDevice.Item1);
 					}
 				}
 			}
-
 		}
 
 		/// <summary>
 		/// Create a new device platform
 		/// </summary>
-		public Task<IDevicePlatform?> TryCreatePlatformAsync(DevicePlatformId Id, string Name)
+		public Task<IDevicePlatform?> TryCreatePlatformAsync(DevicePlatformId id, string name)
 		{
-			return Devices.TryAddPlatformAsync(Id, Name);
+			return _devices.TryAddPlatformAsync(id, name);
 		}
 
 		/// <summary>
@@ -211,161 +205,160 @@ namespace Horde.Build.Services
 		/// </summary>
 		public Task<List<IDevicePlatform>> GetPlatformsAsync()
 		{
-			return Devices.FindAllPlatformsAsync();
+			return _devices.FindAllPlatformsAsync();
 		}
 
 		/// <summary>
 		/// Update an existing platform
 		/// </summary>
-		public Task<bool> UpdatePlatformAsync(DevicePlatformId PlatformId, string[]? ModelIds)
+		public Task<bool> UpdatePlatformAsync(DevicePlatformId platformId, string[]? modelIds)
 		{
-			return Devices.UpdatePlatformAsync(PlatformId, ModelIds);
+			return _devices.UpdatePlatformAsync(platformId, modelIds);
 		}
 
 		/// <summary>
 		/// Get a specific device platform
 		/// </summary>
-		public Task<IDevicePlatform?> GetPlatformAsync(DevicePlatformId Id)
+		public Task<IDevicePlatform?> GetPlatformAsync(DevicePlatformId id)
 		{
-			return Devices.GetPlatformAsync(Id);
+			return _devices.GetPlatformAsync(id);
 		}
 
 		/// <summary>
 		/// Get a device pool by id
 		/// </summary>
-		public Task<IDevicePool?> GetPoolAsync(DevicePoolId Id)
+		public Task<IDevicePool?> GetPoolAsync(DevicePoolId id)
 		{
-			return Devices.GetPoolAsync(Id);
+			return _devices.GetPoolAsync(id);
 		}
 
 		/// <summary>
 		/// Create a new device pool
 		/// </summary>
-		public Task<IDevicePool?> TryCreatePoolAsync(DevicePoolId Id, string Name, DevicePoolType PoolType, List<ProjectId>? ProjectIds)
+		public Task<IDevicePool?> TryCreatePoolAsync(DevicePoolId id, string name, DevicePoolType poolType, List<ProjectId>? projectIds)
 		{
-			return Devices.TryAddPoolAsync(Id, Name, PoolType, ProjectIds);
+			return _devices.TryAddPoolAsync(id, name, poolType, projectIds);
 		}
 
 		/// <summary>
 		/// Update a device pool
 		/// </summary>
-		public Task UpdatePoolAsync(DevicePoolId Id, List<ProjectId>? ProjectIds)
+		public Task UpdatePoolAsync(DevicePoolId id, List<ProjectId>? projectIds)
 		{
-			return Devices.UpdatePoolAsync(Id, ProjectIds);
+			return _devices.UpdatePoolAsync(id, projectIds);
 		}
 
-
 		/// <summary>
 		/// Get a list of existing device pools
 		/// </summary>
 		public Task<List<IDevicePool>> GetPoolsAsync()
 		{
-			return Devices.FindAllPoolsAsync();
+			return _devices.FindAllPoolsAsync();
 		}
 
 		/// <summary>
 		/// Get a list of devices, optionally filtered to provided ids
 		/// </summary>
-		public Task<List<IDevice>> GetDevicesAsync(List<DeviceId>? DeviceIds = null)
+		public Task<List<IDevice>> GetDevicesAsync(List<DeviceId>? deviceIds = null)
 		{
-			return Devices.FindAllDevicesAsync(DeviceIds);
+			return _devices.FindAllDevicesAsync(deviceIds);
 		}
 
 		/// <summary>
 		/// Get a specific device
 		/// </summary>
-		public Task<IDevice?> GetDeviceAsync(DeviceId Id)
+		public Task<IDevice?> GetDeviceAsync(DeviceId id)
 		{
-			return Devices.GetDeviceAsync(Id);
+			return _devices.GetDeviceAsync(id);
 		}
 
 		/// <summary>
 		/// Get a device by name
 		/// </summary>
-		public Task<IDevice?> GetDeviceByNameAsync(string DeviceName)
+		public Task<IDevice?> GetDeviceByNameAsync(string deviceName)
 		{
-			return Devices.GetDeviceByNameAsync(DeviceName);
+			return _devices.GetDeviceByNameAsync(deviceName);
 		}
 
 		/// <summary>
 		/// Delete a device
 		/// </summary>
-		public async Task<bool> DeleteDeviceAsync(DeviceId Id)
+		public async Task<bool> DeleteDeviceAsync(DeviceId id)
 		{
-			return await Devices.DeleteDeviceAsync(Id);
+			return await _devices.DeleteDeviceAsync(id);
 		}
 
 		/// <summary>
 		/// Create a new device
 		/// </summary>
-		/// <param name="Id">Unique id of the device</param>
-		/// <param name="Name">Friendly name of the device</param>
-		/// <param name="PlatformId">The device platform</param>
-		/// <param name="PoolId">Which pool to add the device</param>
-		/// <param name="Enabled">Whether the device is enabled</param>
-		/// <param name="Address">Address or hostname of device</param>
-		/// <param name="ModelId">Vendor model id</param>
-        /// <param name="UserId">User adding the device</param>
+		/// <param name="id">Unique id of the device</param>
+		/// <param name="name">Friendly name of the device</param>
+		/// <param name="platformId">The device platform</param>
+		/// <param name="poolId">Which pool to add the device</param>
+		/// <param name="enabled">Whether the device is enabled</param>
+		/// <param name="address">Address or hostname of device</param>
+		/// <param name="modelId">Vendor model id</param>
+        /// <param name="userId">User adding the device</param>
 		/// <returns></returns>
-		public Task<IDevice?> TryCreateDeviceAsync(DeviceId Id, string Name, DevicePlatformId PlatformId, DevicePoolId PoolId, bool? Enabled, string? Address, string? ModelId, UserId? UserId = null)
+		public Task<IDevice?> TryCreateDeviceAsync(DeviceId id, string name, DevicePlatformId platformId, DevicePoolId poolId, bool? enabled, string? address, string? modelId, UserId? userId = null)
 		{
-			return Devices.TryAddDeviceAsync(Id, Name, PlatformId, PoolId, Enabled, Address, ModelId, UserId);
+			return _devices.TryAddDeviceAsync(id, name, platformId, poolId, enabled, address, modelId, userId);
 		}
 
 		/// <summary>
 		/// Update a device
 		/// </summary>
-		public Task UpdateDeviceAsync(DeviceId DeviceId, DevicePoolId? NewPoolId = null, string? NewName = null, string? NewAddress = null, string? NewModelId = null, string? NewNotes = null, bool? NewEnabled = null, bool? NewProblem = null, bool? NewMaintenance = null, UserId? ModifiedByUserId = null)
+		public Task UpdateDeviceAsync(DeviceId deviceId, DevicePoolId? newPoolId = null, string? newName = null, string? newAddress = null, string? newModelId = null, string? newNotes = null, bool? newEnabled = null, bool? newProblem = null, bool? newMaintenance = null, UserId? modifiedByUserId = null)
 		{
-			return Devices.UpdateDeviceAsync(DeviceId, NewPoolId, NewName, NewAddress, NewModelId, NewNotes, NewEnabled, NewProblem, NewMaintenance, ModifiedByUserId);
+			return _devices.UpdateDeviceAsync(deviceId, newPoolId, newName, newAddress, newModelId, newNotes, newEnabled, newProblem, newMaintenance, modifiedByUserId);
 		}
 
 		/// <summary>
 		/// Checkout a device
 		/// </summary>
-		public Task CheckoutDeviceAsync(DeviceId DeviceId, UserId? UserId)
+		public Task CheckoutDeviceAsync(DeviceId deviceId, UserId? userId)
 		{
-            return Devices.CheckoutDeviceAsync(DeviceId, UserId);
+            return _devices.CheckoutDeviceAsync(deviceId, userId);
         }
 
 		/// <summary>
 		/// Try to create a reservation satisfying the specified device platforms and models
 		/// </summary>
-		public Task<IDeviceReservation?> TryCreateReservationAsync(DevicePoolId Pool, List<DeviceRequestData> Request, string? Hostname = null, string? ReservationDetails = null, string? JobId = null, string? StepId = null)
+		public Task<IDeviceReservation?> TryCreateReservationAsync(DevicePoolId pool, List<DeviceRequestData> request, string? hostname = null, string? reservationDetails = null, string? jobId = null, string? stepId = null)
 		{
-			return Devices.TryAddReservationAsync(Pool, Request, Hostname, ReservationDetails, JobId, StepId);
+			return _devices.TryAddReservationAsync(pool, request, hostname, reservationDetails, jobId, stepId);
 		}
 
 		/// <summary>
 		/// Update/renew an existing reservation
 		/// </summary>
-		public Task<bool> TryUpdateReservationAsync(ObjectId Id)
+		public Task<bool> TryUpdateReservationAsync(ObjectId id)
 		{
-			return Devices.TryUpdateReservationAsync(Id);
+			return _devices.TryUpdateReservationAsync(id);
 		}
 
 		/// <summary>
 		///  Delete an existing reservation
 		/// </summary>
-		public Task<bool> DeleteReservationAsync(ObjectId Id)
+		public Task<bool> DeleteReservationAsync(ObjectId id)
 		{
-			return Devices.DeleteReservationAsync(Id);
+			return _devices.DeleteReservationAsync(id);
 		}
 
 		/// <summary>
 		/// Get a reservation from a legacy guid
 		/// </summary>
-		public Task<IDeviceReservation?> TryGetReservationFromLegacyGuidAsync(string LegacyGuid)
+		public Task<IDeviceReservation?> TryGetReservationFromLegacyGuidAsync(string legacyGuid)
 		{
-			return Devices.TryGetReservationFromLegacyGuidAsync(LegacyGuid);
+			return _devices.TryGetReservationFromLegacyGuidAsync(legacyGuid);
 		}
 
 		/// <summary>
 		/// Get a reservation from a device id
 		/// </summary>
-		public Task<IDeviceReservation?> TryGetDeviceReservation(DeviceId DeviceId)
+		public Task<IDeviceReservation?> TryGetDeviceReservation(DeviceId deviceId)
 		{
-			return Devices.TryGetDeviceReservationAsync(DeviceId);
+			return _devices.TryGetDeviceReservationAsync(deviceId);
 		}
 
 		/// <summary>
@@ -373,103 +366,99 @@ namespace Horde.Build.Services
 		/// </summary>
 		public Task<List<IDeviceReservation>> GetReservationsAsync()
 		{
-			return Devices.FindAllReservationsAsync();
+			return _devices.FindAllReservationsAsync();
 		}
 
 		/// <summary>
 		/// 
 		/// </summary>
-		/// <param name="Message"></param>
-		/// <param name="DeviceId"></param>
-		/// <param name="JobId"></param>
-		/// <param name="StepId"></param>
-		/// <param name="UserId"></param>
-		public async Task NotifyDeviceServiceAsync(string Message, DeviceId? DeviceId = null, string? JobId = null, string? StepId = null, UserId? UserId = null)
+		/// <param name="message"></param>
+		/// <param name="deviceId"></param>
+		/// <param name="jobId"></param>
+		/// <param name="stepId"></param>
+		/// <param name="userId"></param>
+		public async Task NotifyDeviceServiceAsync(string message, DeviceId? deviceId = null, string? jobId = null, string? stepId = null, UserId? userId = null)
 		{
 			try 
 			{
-				IDevice? Device = null;
-				IDevicePool? Pool = null;
-				IJob? Job = null;
-				IJobStep? Step = null;
-				INode? Node = null;
-				IStream? Stream = null;
-				IUser? User = null;
-
-				if (UserId.HasValue)
+				IDevice? device = null;
+				IDevicePool? pool = null;
+				IJob? job = null;
+				IJobStep? step = null;
+				INode? node = null;
+				IStream? stream = null;
+				IUser? user = null;
+
+				if (userId.HasValue)
 				{
-					User = await UserCollection.GetUserAsync(UserId.Value);
-					if (User == null)
+					user = await UserCollection.GetUserAsync(userId.Value);
+					if (user == null)
 					{
-						Logger.LogError("Unable to send device notification, can't find User {UserId}", UserId.Value);
+						_logger.LogError("Unable to send device notification, can't find User {UserId}", userId.Value);
 						return;
 					}
 				}
 
-				if (DeviceId.HasValue)
+				if (deviceId.HasValue)
 				{
-					Device = await GetDeviceAsync(DeviceId.Value);
-					Pool = await GetPoolAsync(Device!.PoolId);
+					device = await GetDeviceAsync(deviceId.Value);
+					pool = await GetPoolAsync(device!.PoolId);
 				}
 
-				if (JobId != null)
+				if (jobId != null)
 				{
-					Job = await JobService.GetJobAsync(new JobId(JobId));
+					job = await _jobService.GetJobAsync(new JobId(jobId));
 
-					if (Job != null)
+					if (job != null)
 					{
-						Stream = await StreamService.GetStreamAsync(Job.StreamId);
+						stream = await _streamService.GetStreamAsync(job.StreamId);
 
-						if (StepId != null)
+						if (stepId != null)
 						{
-							IGraph Graph = await JobService.GetGraphAsync(Job)!;
+							IGraph graph = await _jobService.GetGraphAsync(job)!;
 
-							SubResourceId StepIdValue = SubResourceId.Parse(StepId);
-							IJobStepBatch? Batch = Job.Batches.FirstOrDefault(B => B.Steps.FirstOrDefault(S => S.Id == StepIdValue) != null);
-							if (Batch != null)
+							SubResourceId stepIdValue = SubResourceId.Parse(stepId);
+							IJobStepBatch? batch = job.Batches.FirstOrDefault(b => b.Steps.FirstOrDefault(s => s.Id == stepIdValue) != null);
+							if (batch != null)
 							{
-								Step = Batch.Steps.FirstOrDefault(S => S.Id == StepIdValue)!;
-								INodeGroup Group = Graph.Groups[Batch.GroupIdx];
-								Node = Group.Nodes[Step.NodeIdx];
+								step = batch.Steps.FirstOrDefault(s => s.Id == stepIdValue)!;
+								INodeGroup group = graph.Groups[batch.GroupIdx];
+								node = group.Nodes[step.NodeIdx];
 							}
 						}
 					}
 				}
 
-				NotificationService.NotifyDeviceService(Message, Device, Pool, Stream, Job, Step, Node, User);
+				_notificationService.NotifyDeviceService(message, device, pool, stream, job, step, node, user);
 
 			}
-			catch (Exception Ex)
+			catch (Exception ex)
 			{
-                Logger.LogError(Ex, "Error on device notification {Message}", Ex.Message);
+                _logger.LogError(ex, "Error on device notification {Message}", ex.Message);
             }
         }
 
-
-
 		/// <summary>
 		/// Authorize device action
 		/// </summary>
-		/// <param name="Action"></param>
-		/// <param name="User"></param>
+		/// <param name="action"></param>
+		/// <param name="user"></param>
 		/// <returns></returns>
-		[SuppressMessage("Usage", "CA1801:Review unused parameters")]
-		[SuppressMessage("Performance", "CA1822: Can be static ")]
-		async public Task<bool> AuthorizeAsync(AclAction Action, ClaimsPrincipal User)
+		public async Task<bool> AuthorizeAsync(AclAction action, ClaimsPrincipal user)
 		{
 			// Tihs is deprecated and device auth should be going through GetUserPoolAuthorizationsAsync
 
-			if (Action == AclAction.DeviceRead)
+			if (action == AclAction.DeviceRead)
 			{
 				return true;
 			}
 
-			if (User.IsInRole("Internal-Employees"))
+			if (user.IsInRole("Internal-Employees"))
 			{
 				return true;
 			}
 
-			if (await AclService.AuthorizeAsync(AclAction.AdminWrite, User))
+			if (await _aclService.AuthorizeAsync(AclAction.AdminWrite, user))
 			{
 				return true;
 			}
@@ -481,167 +470,165 @@ namespace Horde.Build.Services
 		/// <summary>
 		/// Get list of pool authorizations for a user
 		/// </summary>
-		/// <param name="User"></param>
+		/// <param name="user"></param>
 		/// <returns></returns>
-		async public Task<List<DevicePoolAuthorization>> GetUserPoolAuthorizationsAsync(ClaimsPrincipal User)
+		public async Task<List<DevicePoolAuthorization>> GetUserPoolAuthorizationsAsync(ClaimsPrincipal user)
 		{
 			
-			List<DevicePoolAuthorization> AuthPools = new List<DevicePoolAuthorization>();
+			List<DevicePoolAuthorization> authPools = new List<DevicePoolAuthorization>();
 
-			List<IDevicePool> AllPools = await GetPoolsAsync();
-			List<IProject> Projects = await ProjectService.GetProjectsAsync();
+			List<IDevicePool> allPools = await GetPoolsAsync();
+			List<IProject> projects = await _projectService.GetProjectsAsync();
 
 			// Set of projects associated with device pools
-			HashSet<ProjectId> ProjectIds = new HashSet<ProjectId>(AllPools.Where(x => x.ProjectIds != null).SelectMany(x => x.ProjectIds!));
+			HashSet<ProjectId> projectIds = new HashSet<ProjectId>(allPools.Where(x => x.ProjectIds != null).SelectMany(x => x.ProjectIds!));
 
-			Dictionary<ProjectId, bool> DeviceRead = new Dictionary<ProjectId, bool>();
-			Dictionary<ProjectId, bool> DeviceWrite = new Dictionary<ProjectId, bool>();
-			ProjectPermissionsCache PermissionsCache = new ProjectPermissionsCache();
+			Dictionary<ProjectId, bool> deviceRead = new Dictionary<ProjectId, bool>();
+			Dictionary<ProjectId, bool> deviceWrite = new Dictionary<ProjectId, bool>();
+			ProjectPermissionsCache permissionsCache = new ProjectPermissionsCache();
 
-			foreach (ProjectId ProjectId in ProjectIds)
+			foreach (ProjectId projectId in projectIds)
 			{
 
-				if (Projects.Where(x => x.Id == ProjectId).FirstOrDefault() == null)
+				if (projects.Where(x => x.Id == projectId).FirstOrDefault() == null)
 				{
-					Logger.LogWarning("Device pool authorization references missing project id {ProjectId}", ProjectId);
+					_logger.LogWarning("Device pool authorization references missing project id {ProjectId}", projectId);
 					continue;
 				}
 
-				DeviceRead.Add(ProjectId,  await ProjectService.AuthorizeAsync(ProjectId, AclAction.DeviceRead, User, PermissionsCache));
-				DeviceWrite.Add(ProjectId, await ProjectService.AuthorizeAsync(ProjectId, AclAction.DeviceWrite, User, PermissionsCache));
+				deviceRead.Add(projectId,  await _projectService.AuthorizeAsync(projectId, AclAction.DeviceRead, user, permissionsCache));
+				deviceWrite.Add(projectId, await _projectService.AuthorizeAsync(projectId, AclAction.DeviceWrite, user, permissionsCache));
 			}
 
 			// for global pools which aren't associated with a project
-			bool GlobalPoolAccess = User.IsInRole("Internal-Employees");
+			bool globalPoolAccess = user.IsInRole("Internal-Employees");
 
-			if (!GlobalPoolAccess)
+			if (!globalPoolAccess)
 			{
-				if (await AclService.AuthorizeAsync(AclAction.AdminWrite, User))
+				if (await _aclService.AuthorizeAsync(AclAction.AdminWrite, user))
 				{
-					GlobalPoolAccess = true;
+					globalPoolAccess = true;
 				}
 			}
 
-			foreach (IDevicePool Pool in AllPools)
+			foreach (IDevicePool pool in allPools)
 			{				
-				if (Pool.ProjectIds == null || Pool.ProjectIds.Count == 0)
+				if (pool.ProjectIds == null || pool.ProjectIds.Count == 0)
 				{
-					if (Pool.PoolType == DevicePoolType.Shared && !GlobalPoolAccess)
+					if (pool.PoolType == DevicePoolType.Shared && !globalPoolAccess)
 					{
-						AuthPools.Add(new DevicePoolAuthorization(Pool, false, false));
+						authPools.Add(new DevicePoolAuthorization(pool, false, false));
 						continue;
 					}
 
-					AuthPools.Add(new DevicePoolAuthorization(Pool, true, true));
+					authPools.Add(new DevicePoolAuthorization(pool, true, true));
 					continue;
 				}
 
-				bool Read = false;
-				bool Write = false;
+				bool read = false;
+				bool write = false;
 
-				foreach (ProjectId ProjectId in Pool.ProjectIds)
+				foreach (ProjectId projectId in pool.ProjectIds)
 				{
-					bool Value;
+					bool value;
 
-					if (!Read && DeviceRead.TryGetValue(ProjectId, out Value))
+					if (!read && deviceRead.TryGetValue(projectId, out value))
 					{
-						Read = Value;
+						read = value;
 					}
 
-					if (!Write && DeviceWrite.TryGetValue(ProjectId, out Value))
+					if (!write && deviceWrite.TryGetValue(projectId, out value))
 					{
-						Write = Value;
+						write = value;
 					}
 				}
 
-				AuthPools.Add(new DevicePoolAuthorization(Pool, Read, Write));
+				authPools.Add(new DevicePoolAuthorization(pool, read, write));
 
 			}
 
-			return AuthPools;
+			return authPools;
 		}
 
 		/// <summary>
 		/// Get list of pool authorizations for a user
 		/// </summary>
-		/// <param name="Id"></param>
-		/// <param name="User"></param>
+		/// <param name="id"></param>
+		/// <param name="user"></param>
 		/// <returns></returns>
-		async public Task<DevicePoolAuthorization?> GetUserPoolAuthorizationAsync(DevicePoolId Id, ClaimsPrincipal User)
+		public async Task<DevicePoolAuthorization?> GetUserPoolAuthorizationAsync(DevicePoolId id, ClaimsPrincipal user)
 		{
-			List<DevicePoolAuthorization> Auth = await GetUserPoolAuthorizationsAsync(User);
-			return Auth.Where(x => x.Pool.Id == Id).FirstOrDefault();
+			List<DevicePoolAuthorization> auth = await GetUserPoolAuthorizationsAsync(user);
+			return auth.Where(x => x.Pool.Id == id).FirstOrDefault();
 		}
 
 		/// <summary>
 		/// Update a device
 		/// </summary>
-		async public Task UpdateDevicePoolAsync(DevicePoolId PoolId, List<ProjectId>? ProjectIds)
+		public async Task UpdateDevicePoolAsync(DevicePoolId poolId, List<ProjectId>? projectIds)
 		{
-			await Devices.UpdatePoolAsync(PoolId, ProjectIds);
+			await _devices.UpdatePoolAsync(poolId, projectIds);
 		}
 
-
 		/// <summary>
 		/// Get Platform mappings for V1 API
 		/// </summary>		
 		public async Task<DevicePlatformMapV1> GetPlatformMapV1()
 		{
-			return await PlatformMapSingleton.GetAsync();
+			return await _platformMapSingleton.GetAsync();
 		}
 
 		/// <summary>
 		/// Updates the platform mapping information required for the V1 api
 		/// </summary>				
-		public async Task<bool> UpdatePlatformMapAsync(UpdatePlatformMapRequest Request)
+		public async Task<bool> UpdatePlatformMapAsync(UpdatePlatformMapRequest request)
 		{
 
-			List<IDevicePlatform> Platforms = await GetPlatformsAsync();
+			List<IDevicePlatform> platforms = await GetPlatformsAsync();
 
 			// Update the platform map
 			for (int i = 0; i < 10; i++)
 			{
-				DevicePlatformMapV1 Instance = await PlatformMapSingleton.GetAsync();
+				DevicePlatformMapV1 instance = await _platformMapSingleton.GetAsync();
 
-				Instance.PlatformMap = new Dictionary<string, DevicePlatformId>();
-				Instance.PlatformReverseMap = new Dictionary<DevicePlatformId, string>();
-				Instance.PerfSpecHighMap = new Dictionary<DevicePlatformId, string>();
+				instance.PlatformMap = new Dictionary<string, DevicePlatformId>();
+				instance.PlatformReverseMap = new Dictionary<DevicePlatformId, string>();
+				instance.PerfSpecHighMap = new Dictionary<DevicePlatformId, string>();
 
-				foreach(KeyValuePair<string, string> Entry in Request.PlatformMap)
+				foreach(KeyValuePair<string, string> entry in request.PlatformMap)
 				{
-					IDevicePlatform? Platform = Platforms.FirstOrDefault(P => P.Id == new DevicePlatformId(Entry.Value));
-					if (Platform == null)
+					IDevicePlatform? platform = platforms.FirstOrDefault(p => p.Id == new DevicePlatformId(entry.Value));
+					if (platform == null)
 					{
-						throw new Exception($"Unknowm platform in map {Entry.Key} : {Entry.Value}");
+						throw new Exception($"Unknowm platform in map {entry.Key} : {entry.Value}");
 					}
 
-					Instance.PlatformMap.Add(Entry.Key, Platform.Id);
+					instance.PlatformMap.Add(entry.Key, platform.Id);
 				}
 
-				foreach (KeyValuePair<string, string> Entry in Request.PlatformReverseMap)
+				foreach (KeyValuePair<string, string> entry in request.PlatformReverseMap)
 				{
-					IDevicePlatform? Platform = Platforms.FirstOrDefault(P => P.Id == new DevicePlatformId(Entry.Key));
-					if (Platform == null)
+					IDevicePlatform? platform = platforms.FirstOrDefault(p => p.Id == new DevicePlatformId(entry.Key));
+					if (platform == null)
 					{
-						throw new Exception($"Unknowm platform in reverse map {Entry.Key} : {Entry.Value}");
+						throw new Exception($"Unknowm platform in reverse map {entry.Key} : {entry.Value}");
 					}
 
-					Instance.PlatformReverseMap.Add(Platform.Id, Entry.Value);
+					instance.PlatformReverseMap.Add(platform.Id, entry.Value);
 				}
 
-				foreach (KeyValuePair<string, string> Entry in Request.PerfSpecHighMap)
+				foreach (KeyValuePair<string, string> entry in request.PerfSpecHighMap)
 				{
-					IDevicePlatform? Platform = Platforms.FirstOrDefault(P => P.Id == new DevicePlatformId(Entry.Key));
-					if (Platform == null)
+					IDevicePlatform? platform = platforms.FirstOrDefault(p => p.Id == new DevicePlatformId(entry.Key));
+					if (platform == null)
 					{
-						throw new Exception($"Unknowm platform in spec map {Entry.Key} : {Entry.Value}");
+						throw new Exception($"Unknowm platform in spec map {entry.Key} : {entry.Value}");
 					}
 
-					Instance.PerfSpecHighMap.Add(Platform.Id, Entry.Value);
+					instance.PerfSpecHighMap.Add(platform.Id, entry.Value);
 				}
 
-
-				if (await PlatformMapSingleton.TryUpdateAsync(Instance))
+				if (await _platformMapSingleton.TryUpdateAsync(instance))
 				{
 					return true;
 				}
@@ -650,8 +637,6 @@ namespace Horde.Build.Services
 			}
 
 			return false;
-
 		}
-
 	}
 }
diff --git a/Engine/Source/Programs/Horde/Horde.Build/Fleet/Autoscale/AutoscaleService.cs b/Engine/Source/Programs/Horde/Horde.Build/Fleet/Autoscale/AutoscaleService.cs
--- a/Engine/Source/Programs/Horde/Horde.Build/Fleet/Autoscale/AutoscaleService.cs
+++ b/Engine/Source/Programs/Horde/Horde.Build/Fleet/Autoscale/AutoscaleService.cs
@@ -1,31 +1,22 @@
 // Copyright Epic Games, Inc. All Rights Reserved.
 
-using Amazon;
-using Amazon.EC2;
-using Amazon.EC2.Model;
-using Amazon.Runtime;
-using Amazon.Runtime.CredentialManagement;
-using Google.Protobuf.WellKnownTypes;
-using HordeCommon.Rpc.Tasks;
-using Horde.Build.Collections;
-using Horde.Build.Models;
-using Horde.Build.Utilities;
-using Microsoft.Extensions.Logging;
-using MongoDB.Bson;
 using System;
 using System.Collections.Generic;
 using System.Diagnostics;
 using System.Linq;
-using System.Net;
-using System.Runtime.CompilerServices;
 using System.Threading;
 using System.Threading.Tasks;
+using Google.Protobuf.WellKnownTypes;
+using Horde.Build.Collections;
+using Horde.Build.Models;
+using Horde.Build.Utilities;
 using HordeCommon;
-using Microsoft.Extensions.Options;
+using HordeCommon.Rpc.Tasks;
+using Microsoft.Extensions.Hosting;
+using Microsoft.Extensions.Logging;
 using OpenTracing;
 using OpenTracing.Util;
 using StatsdClient;
-using Microsoft.Extensions.Hosting;
 
 namespace Horde.Build.Services
 {
@@ -38,13 +29,13 @@ namespace Horde.Build.Services
 	{
 		const int NumSamples = 10;
 		const int NumSamplesForResult = 9;
-		static readonly TimeSpan SampleTime = TimeSpan.FromMinutes(6.0);
-		static readonly TimeSpan ShrinkPoolCoolDown = TimeSpan.FromMinutes(20.0);
+		static readonly TimeSpan s_sampleTime = TimeSpan.FromMinutes(6.0);
+		static readonly TimeSpan s_shrinkPoolCoolDown = TimeSpan.FromMinutes(20.0);
 
 		struct UtilizationSample
 		{
-			public double JobWork;
-			public double OtherWork;
+			public double _jobWork;
+			public double _otherWork;
 		}
 
 		class AgentData
@@ -53,20 +44,20 @@ namespace Horde.Build.Services
 			public UtilizationSample[] Samples { get; } = new UtilizationSample[NumSamples];
 			public List<ILease> Leases { get; } = new List<ILease>();
 
-			public AgentData(IAgent Agent)
+			public AgentData(IAgent agent)
 			{
-				this.Agent = Agent;
+				Agent = agent;
 			}
 
-			public void Add(double MinT, double MaxT, double JobWork, double OtherWork)
+			public void Add(double minT, double maxT, double jobWork, double otherWork)
 			{
-				int MinIdx = Math.Clamp((int)MinT, 0, NumSamples - 1);
-				int MaxIdx = Math.Clamp((int)MaxT, MinIdx, NumSamples - 1);
-				for (int Idx = MinIdx; Idx <= MaxIdx; Idx++)
+				int minIdx = Math.Clamp((int)minT, 0, NumSamples - 1);
+				int maxIdx = Math.Clamp((int)maxT, minIdx, NumSamples - 1);
+				for (int idx = minIdx; idx <= maxIdx; idx++)
 				{
-					double Fraction = Math.Clamp(MaxT - Idx, 0.0, 1.0) - Math.Clamp(MinT - Idx, 0.0, 1.0);
-					Samples[Idx].JobWork += JobWork * Fraction;
-					Samples[Idx].OtherWork += OtherWork * Fraction;
+					double fraction = Math.Clamp(maxT - idx, 0.0, 1.0) - Math.Clamp(minT - idx, 0.0, 1.0);
+					Samples[idx]._jobWork += jobWork * fraction;
+					Samples[idx]._otherWork += otherWork * fraction;
 				}
 			}
 		}
@@ -77,212 +68,212 @@ namespace Horde.Build.Services
 			public List<IAgent> Agents { get; } = new List<IAgent>();
 			public UtilizationSample[] Samples { get; } = new UtilizationSample[NumSamples];
 
-			public PoolData(IPool Pool)
+			public PoolData(IPool pool)
 			{
-				this.Pool = Pool;
+				Pool = pool;
 			}
 
-			public void Add(AgentData AgentData)
+			public void Add(AgentData agentData)
 			{
-				Agents.Add(AgentData.Agent);
-				for (int Idx = 0; Idx < Samples.Length; Idx++)
+				Agents.Add(agentData.Agent);
+				for (int idx = 0; idx < Samples.Length; idx++)
 				{
-					Samples[Idx].JobWork += AgentData.Samples[Idx].JobWork;
-					Samples[Idx].OtherWork += AgentData.Samples[Idx].OtherWork;
+					Samples[idx]._jobWork += agentData.Samples[idx]._jobWork;
+					Samples[idx]._otherWork += agentData.Samples[idx]._otherWork;
 				}
 			}
 		}
 
-		IAgentCollection AgentCollection;
-		IPoolCollection PoolCollection;
-		ILeaseCollection LeaseCollection;
-		IFleetManager FleetManager;
-		IDogStatsd DogStatsd;
-		IClock Clock;
-		ITicker Ticker;
-		ILogger<AutoscaleService> Logger;
+		readonly IAgentCollection _agentCollection;
+		readonly IPoolCollection _poolCollection;
+		readonly ILeaseCollection _leaseCollection;
+		readonly IFleetManager _fleetManager;
+		readonly IDogStatsd _dogStatsd;
+		readonly IClock _clock;
+		readonly ITicker _ticker;
+		readonly ILogger<AutoscaleService> _logger;
 
 		/// <summary>
 		/// Constructor
 		/// </summary>
-		public AutoscaleService(IAgentCollection AgentCollection, IPoolCollection PoolCollection, ILeaseCollection LeaseCollection, IFleetManager FleetManager, IDogStatsd DogStatsd, DatabaseService DatabaseService, IClock Clock, ILogger<AutoscaleService> Logger)
+		public AutoscaleService(IAgentCollection agentCollection, IPoolCollection poolCollection, ILeaseCollection leaseCollection, IFleetManager fleetManager, IDogStatsd dogStatsd, DatabaseService databaseService, IClock clock, ILogger<AutoscaleService> logger)
 		{
-			this.AgentCollection = AgentCollection;
-			this.PoolCollection = PoolCollection;
-			this.LeaseCollection = LeaseCollection;
-			this.FleetManager = FleetManager;
-			this.DogStatsd = DogStatsd;
-			this.Clock = Clock;
-			if (DatabaseService.ReadOnlyMode)
+			_agentCollection = agentCollection;
+			_poolCollection = poolCollection;
+			_leaseCollection = leaseCollection;
+			_fleetManager = fleetManager;
+			_dogStatsd = dogStatsd;
+			_clock = clock;
+			if (databaseService.ReadOnlyMode)
 			{
-				this.Ticker = new NullTicker();
+				_ticker = new NullTicker();
 			}
 			else
 			{
-				this.Ticker = Clock.AddSharedTicker<AutoscaleService>(TimeSpan.FromMinutes(5.0), TickLeaderAsync, Logger);
+				_ticker = clock.AddSharedTicker<AutoscaleService>(TimeSpan.FromMinutes(5.0), TickLeaderAsync, logger);
 			}
-			this.Logger = Logger;
+			_logger = logger;
 		}
 
 		/// <inheritdoc/>
-		public Task StartAsync(CancellationToken CancellationToken) => Ticker.StartAsync();
+		public Task StartAsync(CancellationToken cancellationToken) => _ticker.StartAsync();
 
 		/// <inheritdoc/>
-		public Task StopAsync(CancellationToken CancellationToken) => Ticker.StopAsync();
+		public Task StopAsync(CancellationToken cancellationToken) => _ticker.StopAsync();
 
 		/// <inheritdoc/>
-		public void Dispose() => Ticker.Dispose();
+		public void Dispose() => _ticker.Dispose();
 
 		/// <inheritdoc/>
-		async ValueTask TickLeaderAsync(CancellationToken StoppingToken)
+		async ValueTask TickLeaderAsync(CancellationToken stoppingToken)
 		{
-			DateTime UtcNow = Clock.UtcNow;
+			DateTime utcNow = _clock.UtcNow;
 
-			Logger.LogInformation("Autoscaling pools...");
-			Stopwatch Stopwatch = Stopwatch.StartNew();
+			_logger.LogInformation("Autoscaling pools...");
+			Stopwatch stopwatch = Stopwatch.StartNew();
 
-			Dictionary<PoolId, PoolData> PoolToData;
-			using (IScope Scope = GlobalTracer.Instance.BuildSpan("Compute utilization").StartActive())
+			Dictionary<PoolId, PoolData> poolToData;
+			using (IScope scope = GlobalTracer.Instance.BuildSpan("Compute utilization").StartActive())
 			{
 				// Find all the current agents
-				List<IAgent> Agents = await AgentCollection.FindAsync(Status: AgentStatus.Ok);
+				List<IAgent> agents = await _agentCollection.FindAsync(status: AgentStatus.Ok);
 
 				// Query leases in last interval
-				DateTime MaxTime = Clock.UtcNow;
-				DateTime MinTime = MaxTime - (SampleTime * NumSamples);
-				List<ILease> Leases = await LeaseCollection.FindLeasesAsync(MinTime, MaxTime);
+				DateTime maxTime = _clock.UtcNow;
+				DateTime minTime = maxTime - (s_sampleTime * NumSamples);
+				List<ILease> leases = await _leaseCollection.FindLeasesAsync(minTime, maxTime);
 
 				// Add all the leases to a data object for each agent
-				Dictionary<AgentId, AgentData> AgentIdToData = Agents.ToDictionary(x => x.Id, x => new AgentData(x));
-				foreach (ILease Lease in Leases)
+				Dictionary<AgentId, AgentData> agentIdToData = agents.ToDictionary(x => x.Id, x => new AgentData(x));
+				foreach (ILease lease in leases)
 				{
-					AgentData? AgentData;
-					if (AgentIdToData.TryGetValue(Lease.AgentId, out AgentData) &&
-					    AgentData.Agent.SessionId == Lease.SessionId)
+					AgentData? agentData;
+					if (agentIdToData.TryGetValue(lease.AgentId, out agentData) &&
+					    agentData.Agent.SessionId == lease.SessionId)
 					{
-						AgentData.Leases.Add(Lease);
+						agentData.Leases.Add(lease);
 					}
 				}
 
 				// Compute utilization for each agent
-				foreach (AgentData AgentData in AgentIdToData.Values)
+				foreach (AgentData agentData in agentIdToData.Values)
 				{
-					foreach (ILease Lease in AgentData.Leases.OrderBy(x => x.StartTime))
+					foreach (ILease lease in agentData.Leases.OrderBy(x => x.StartTime))
 					{
-						double MinT = (Lease.StartTime - MinTime).TotalSeconds / SampleTime.TotalSeconds;
-						double MaxT = (Lease.FinishTime == null)
+						double minT = (lease.StartTime - minTime).TotalSeconds / s_sampleTime.TotalSeconds;
+						double maxT = (lease.FinishTime == null)
 							? NumSamples
-							: ((Lease.FinishTime.Value - MinTime).TotalSeconds / SampleTime.TotalSeconds);
+							: ((lease.FinishTime.Value - minTime).TotalSeconds / s_sampleTime.TotalSeconds);
 
-						Any Payload = Any.Parser.ParseFrom(Lease.Payload.ToArray());
-						if (Payload.Is(ExecuteJobTask.Descriptor))
+						Any payload = Any.Parser.ParseFrom(lease.Payload.ToArray());
+						if (payload.Is(ExecuteJobTask.Descriptor))
 						{
-							AgentData.Add(MinT, MaxT, 1.0, 0.0);
+							agentData.Add(minT, maxT, 1.0, 0.0);
 						}
 						else
 						{
-							AgentData.Add(MinT, MaxT, 0.0, 1.0);
+							agentData.Add(minT, maxT, 0.0, 1.0);
 						}
 					}
 				}
 
 				// Get all the pools
-				List<IPool> Pools = await PoolCollection.GetAsync();
-				PoolToData = Pools.ToDictionary(x => x.Id, x => new PoolData(x));
+				List<IPool> pools = await _poolCollection.GetAsync();
+				poolToData = pools.ToDictionary(x => x.Id, x => new PoolData(x));
 
 				// Find pool utilization over the query period
-				foreach (AgentData AgentData in AgentIdToData.Values)
+				foreach (AgentData agentData in agentIdToData.Values)
 				{
-					foreach (PoolId PoolId in AgentData.Agent.GetPools())
+					foreach (PoolId poolId in agentData.Agent.GetPools())
 					{
-						PoolData? PoolData;
-						if (PoolToData.TryGetValue(PoolId, out PoolData))
+						PoolData? poolData;
+						if (poolToData.TryGetValue(poolId, out poolData))
 						{
-							PoolData.Add(AgentData);
+							poolData.Add(agentData);
 						}
 					}
 				}
 			}
 
 			// Output all the final stats
-			foreach (PoolData PoolData in PoolToData.Values.OrderByDescending(x => x.Agents.Count))
+			foreach (PoolData poolData in poolToData.Values.OrderByDescending(x => x.Agents.Count))
 			{
-				IPool Pool = PoolData.Pool;
+				IPool pool = poolData.Pool;
 
-				int MinAgents = Pool.MinAgents ?? 1;
-				int NumReserveAgents = Pool.NumReserveAgents ?? 5;
+				int minAgents = pool.MinAgents ?? 1;
+				int numReserveAgents = pool.NumReserveAgents ?? 5;
 
-				double Utilization = PoolData.Samples.Select(x => x.JobWork).OrderByDescending(x => x).Skip(NumSamples - NumSamplesForResult).First();
+				double utilization = poolData.Samples.Select(x => x._jobWork).OrderByDescending(x => x).Skip(NumSamples - NumSamplesForResult).First();
 
-				int TargetAgents = Math.Max((int)Utilization + NumReserveAgents, MinAgents);
-				int Delta = TargetAgents - PoolData.Agents.Count;
+				int targetAgents = Math.Max((int)utilization + numReserveAgents, minAgents);
+				int delta = targetAgents - poolData.Agents.Count;
 
-				Logger.LogInformation("{PoolName,-48} Jobs=[{JobUtilization}] Total=[{OtherUtilization}] Min={Min,5:0.0} Max={Max,5:0.0} Avg={Avg,5:0.0} Pct={Pct,5:0.0} Current={Current,4} Target={Target,4} Delta={Delta}",
-					Pool.Name,
-					GetDensityMap(PoolData.Samples.Select(x => x.JobWork / Math.Max(1, PoolData.Agents.Count))),
-					GetDensityMap(PoolData.Samples.Select(x => x.OtherWork / Math.Max(1, PoolData.Agents.Count))),
-					PoolData.Samples.Min(x => x.JobWork),
-					PoolData.Samples.Max(x => x.JobWork),
-					Utilization,
-					PoolData.Samples.Sum(x => x.JobWork) / NumSamples,
-					PoolData.Agents.Count,
-					TargetAgents,
-					Delta);
+				_logger.LogInformation("{PoolName,-48} Jobs=[{JobUtilization}] Total=[{OtherUtilization}] Min={Min,5:0.0} Max={Max,5:0.0} Avg={Avg,5:0.0} Pct={Pct,5:0.0} Current={Current,4} Target={Target,4} Delta={Delta}",
+					pool.Name,
+					GetDensityMap(poolData.Samples.Select(x => x._jobWork / Math.Max(1, poolData.Agents.Count))),
+					GetDensityMap(poolData.Samples.Select(x => x._otherWork / Math.Max(1, poolData.Agents.Count))),
+					poolData.Samples.Min(x => x._jobWork),
+					poolData.Samples.Max(x => x._jobWork),
+					utilization,
+					poolData.Samples.Sum(x => x._jobWork) / NumSamples,
+					poolData.Agents.Count,
+					targetAgents,
+					delta);
 
 				try
 				{
-					if (Pool.EnableAutoscaling)
+					if (pool.EnableAutoscaling)
 					{
-						using (IScope Scope = GlobalTracer.Instance.BuildSpan("Scaling pool").StartActive())
+						using (IScope scope = GlobalTracer.Instance.BuildSpan("Scaling pool").StartActive())
 						{
-							Scope.Span.SetTag("poolName", Pool.Name);
-							Scope.Span.SetTag("delta", Delta);
+							scope.Span.SetTag("poolName", pool.Name);
+							scope.Span.SetTag("delta", delta);
 
-							if (Delta > 0)
+							if (delta > 0)
 							{
-								await FleetManager.ExpandPoolAsync(Pool, PoolData.Agents, Delta);
-								await PoolCollection.TryUpdateAsync(Pool, LastScaleUpTime: DateTime.UtcNow);
+								await _fleetManager.ExpandPoolAsync(pool, poolData.Agents, delta);
+								await _poolCollection.TryUpdateAsync(pool, lastScaleUpTime: DateTime.UtcNow);
 							}
 
-							if (Delta < 0)
+							if (delta < 0)
 							{
-								bool bShrinkIsOnCoolDown = Pool.LastScaleDownTime != null && Pool.LastScaleDownTime + ShrinkPoolCoolDown > DateTime.UtcNow;
+								bool bShrinkIsOnCoolDown = pool.LastScaleDownTime != null && pool.LastScaleDownTime + s_shrinkPoolCoolDown > DateTime.UtcNow;
 								if (!bShrinkIsOnCoolDown)
 								{
-									await FleetManager.ShrinkPoolAsync(Pool, PoolData.Agents, -Delta);
-									await PoolCollection.TryUpdateAsync(Pool, LastScaleDownTime: DateTime.UtcNow);
+									await _fleetManager.ShrinkPoolAsync(pool, poolData.Agents, -delta);
+									await _poolCollection.TryUpdateAsync(pool, lastScaleDownTime: DateTime.UtcNow);
 								}
 								else
 								{
-									Logger.LogDebug("Cannot shrink {PoolName} right now, it's on cool-down until {CoolDownTimeEnds}", Pool.Name, Pool.LastScaleDownTime + ShrinkPoolCoolDown);
+									_logger.LogDebug("Cannot shrink {PoolName} right now, it's on cool-down until {CoolDownTimeEnds}", pool.Name, pool.LastScaleDownTime + s_shrinkPoolCoolDown);
 								}
 							}
 						}
 					}
 				}
-				catch (Exception Ex)
+				catch (Exception ex)
 				{
-					Logger.LogError(Ex, "Failed to scale {PoolName}:\n{Exception}", Pool.Name, Ex);
+					_logger.LogError(ex, "Failed to scale {PoolName}:\n{Exception}", pool.Name, ex);
 					continue;
 				}
 
-				DogStatsd.Gauge("agentpools.autoscale.target", TargetAgents, tags: new []{"pool:" + Pool.Name});
-				DogStatsd.Gauge("agentpools.autoscale.current", PoolData.Agents.Count, tags: new []{"pool:" + Pool.Name});
+				_dogStatsd.Gauge("agentpools.autoscale.target", targetAgents, tags: new []{"pool:" + pool.Name});
+				_dogStatsd.Gauge("agentpools.autoscale.current", poolData.Agents.Count, tags: new []{"pool:" + pool.Name});
 			}
 			
-			Stopwatch.Stop();
-			Logger.LogInformation("Autoscaling pools took {ElapsedTime} ms", Stopwatch.ElapsedMilliseconds);
+			stopwatch.Stop();
+			_logger.LogInformation("Autoscaling pools took {ElapsedTime} ms", stopwatch.ElapsedMilliseconds);
 		}
 
 		/// <summary>
 		/// Creates a string of characters indicating a sequence of 0-1 values over time
 		/// </summary>
-		/// <param name="Values">Sequence of values</param>
+		/// <param name="values">Sequence of values</param>
 		/// <returns>Density map</returns>
-		static string GetDensityMap(IEnumerable<double> Values)
+		static string GetDensityMap(IEnumerable<double> values)
 		{
 			const string Greyscale = " 123456789";
-			return new string(Values.Select(x => Greyscale[Math.Clamp((int)(x * Greyscale.Length), 0, Greyscale.Length - 1)]).ToArray());
+			return new string(values.Select(x => Greyscale[Math.Clamp((int)(x * Greyscale.Length), 0, Greyscale.Length - 1)]).ToArray());
 		}
 	}
 }
