Horde: Fix coding conventions within Horde.Build.

diff --git a/Engine/Source/Programs/Horde/Horde.Build/Fleet/Autoscale/AutoscaleServiceV2.cs b/Engine/Source/Programs/Horde/Horde.Build/Fleet/Autoscale/AutoscaleServiceV2.cs
--- a/Engine/Source/Programs/Horde/Horde.Build/Fleet/Autoscale/AutoscaleServiceV2.cs
+++ b/Engine/Source/Programs/Horde/Horde.Build/Fleet/Autoscale/AutoscaleServiceV2.cs
@@ -6,12 +6,11 @@ using System.Diagnostics;
 using System.Linq;
 using System.Threading;
 using System.Threading.Tasks;
-using HordeCommon;
-using Horde.Build;
 using Horde.Build.Collections;
 using Horde.Build.Models;
 using Horde.Build.Services;
 using Horde.Build.Utilities;
+using HordeCommon;
 using Microsoft.Extensions.Hosting;
 using Microsoft.Extensions.Logging;
 using Microsoft.Extensions.Options;
@@ -28,149 +27,152 @@ namespace Horde.Build.Fleet.Autoscale
 	/// </summary>
 	public sealed class AutoscaleServiceV2 : IHostedService, IDisposable
 	{
-		private IPoolSizeStrategy LeaseUtilizationStrategy;
-		private IPoolSizeStrategy JobQueueStrategy;
-		private IPoolSizeStrategy NoOpStrategy;
-		private readonly IAgentCollection AgentCollection;
-		private readonly IPoolCollection PoolCollection;
-		private readonly IFleetManager FleetManager;
-		private readonly IDogStatsd DogStatsd;
-		private readonly IClock Clock;
-		private readonly ITicker Ticker;
-		private readonly ServerSettings Settings;
-		private readonly TimeSpan DefaultScaleOutCooldown;
-		private readonly TimeSpan DefaultScaleInCooldown;
-		private readonly ILogger<AutoscaleServiceV2> Logger;
+		private IPoolSizeStrategy _leaseUtilizationStrategy;
+		private IPoolSizeStrategy _jobQueueStrategy;
+		private IPoolSizeStrategy _noOpStrategy;
+		private readonly IAgentCollection _agentCollection;
+		private readonly IPoolCollection _poolCollection;
+		private readonly IFleetManager _fleetManager;
+		private readonly IDogStatsd _dogStatsd;
+		private readonly IClock _clock;
+		private readonly ITicker _ticker;
+		private readonly ServerSettings _settings;
+		private readonly TimeSpan _defaultScaleOutCooldown;
+		private readonly TimeSpan _defaultScaleInCooldown;
+		private readonly ILogger<AutoscaleServiceV2> _logger;
 
 		/// <summary>
 		/// Constructor
 		/// </summary>
-		public AutoscaleServiceV2(LeaseUtilizationStrategy LeaseUtilizationStrategy, JobQueueStrategy JobQueueStrategy, NoOpPoolSizeStrategy NoOpStrategy, IAgentCollection AgentCollection, IPoolCollection PoolCollection, IFleetManager FleetManager, IDogStatsd DogStatsd, IClock Clock, IOptions<ServerSettings> Settings, ILogger<AutoscaleServiceV2> Logger)
+		public AutoscaleServiceV2(LeaseUtilizationStrategy leaseUtilizationStrategy, JobQueueStrategy jobQueueStrategy, NoOpPoolSizeStrategy noOpStrategy, IAgentCollection agentCollection, IPoolCollection poolCollection, IFleetManager fleetManager, IDogStatsd dogStatsd, IClock clock, IOptions<ServerSettings> settings, ILogger<AutoscaleServiceV2> logger)
 		{
-			this.LeaseUtilizationStrategy = LeaseUtilizationStrategy;
-			this.JobQueueStrategy = JobQueueStrategy;
-			this.NoOpStrategy = NoOpStrategy;
-			this.AgentCollection = AgentCollection;
-			this.PoolCollection = PoolCollection;
-			this.FleetManager = FleetManager;
-			this.DogStatsd = DogStatsd;
-			this.Clock = Clock;
-			this.Ticker = Clock.AddSharedTicker<AutoscaleServiceV2>(TimeSpan.FromMinutes(5.0), TickLeaderAsync, Logger);
-			this.Settings = Settings.Value;
-			this.Logger = Logger;
-			this.DefaultScaleOutCooldown = TimeSpan.FromSeconds(Settings.Value.AgentPoolScaleOutCooldownSeconds);
-			this.DefaultScaleInCooldown = TimeSpan.FromSeconds(Settings.Value.AgentPoolScaleInCooldownSeconds);
+			_leaseUtilizationStrategy = leaseUtilizationStrategy;
+			_jobQueueStrategy = jobQueueStrategy;
+			_noOpStrategy = noOpStrategy;
+			_agentCollection = agentCollection;
+			_poolCollection = poolCollection;
+			_fleetManager = fleetManager;
+			_dogStatsd = dogStatsd;
+			_clock = clock;
+			_ticker = clock.AddSharedTicker<AutoscaleServiceV2>(TimeSpan.FromMinutes(5.0), TickLeaderAsync, logger);
+			_settings = settings.Value;
+			_logger = logger;
+			_defaultScaleOutCooldown = TimeSpan.FromSeconds(settings.Value.AgentPoolScaleOutCooldownSeconds);
+			_defaultScaleInCooldown = TimeSpan.FromSeconds(settings.Value.AgentPoolScaleInCooldownSeconds);
 		}
 
 		/// <inheritdoc/>
-		public Task StartAsync(CancellationToken CancellationToken) => Ticker.StartAsync();
+		public Task StartAsync(CancellationToken cancellationToken) => _ticker.StartAsync();
 
 		/// <inheritdoc/>
-		public Task StopAsync(CancellationToken CancellationToken) => Ticker.StopAsync();
+		public Task StopAsync(CancellationToken cancellationToken) => _ticker.StopAsync();
 
 		/// <inheritdoc/>
-		public void Dispose() => Ticker.Dispose();
+		public void Dispose() => _ticker.Dispose();
 
-		internal async ValueTask TickLeaderAsync(CancellationToken StoppingToken)
+		internal async ValueTask TickLeaderAsync(CancellationToken stoppingToken)
 		{
-			Logger.LogInformation("Autoscaling pools...");
-			Stopwatch Stopwatch = Stopwatch.StartNew();
+			_logger.LogInformation("Autoscaling pools...");
+			Stopwatch stopwatch = Stopwatch.StartNew();
 			using IScope _ = GlobalTracer.Instance.BuildSpan("AutoscaleService").StartActive();
 			
 			// Group pools by strategy type to ensure they can be called once (more optimal)
-			Dictionary<PoolSizeStrategy, List<PoolSizeData>> PoolSizeDataByStrategy = await GetPoolSizeDataByStrategyType();
-			foreach ((var StrategyType, var CurrentData) in PoolSizeDataByStrategy)
+			Dictionary<PoolSizeStrategy, List<PoolSizeData>> poolSizeDataByStrategy = await GetPoolSizeDataByStrategyType();
+			foreach ((PoolSizeStrategy strategyType, List<PoolSizeData> currentData) in poolSizeDataByStrategy)
 			{
-				List<PoolSizeData> NewData = await GetPoolSizeStrategy(StrategyType).CalcDesiredPoolSizesAsync(CurrentData);
-				await ResizePools(NewData);
+				List<PoolSizeData> newData = await GetPoolSizeStrategy(strategyType).CalcDesiredPoolSizesAsync(currentData);
+				await ResizePools(newData);
 			}
 			
-			Stopwatch.Stop();
-			Logger.LogInformation("Autoscaling pools took {ElapsedTime} ms", Stopwatch.ElapsedMilliseconds);
+			stopwatch.Stop();
+			_logger.LogInformation("Autoscaling pools took {ElapsedTime} ms", stopwatch.ElapsedMilliseconds);
 		}
 
 		internal async Task<Dictionary<PoolSizeStrategy, List<PoolSizeData>>> GetPoolSizeDataByStrategyType()
 		{
-			List<IAgent> Agents = await AgentCollection.FindAsync(Status: AgentStatus.Ok);
-			List<IAgent> GetAgentsInPool(PoolId PoolId) => Agents.FindAll(a => a.GetPools().Any(p => p == PoolId));
-			List<IPool> Pools = await PoolCollection.GetAsync();
+			List<IAgent> agents = await _agentCollection.FindAsync(status: AgentStatus.Ok);
+			List<IAgent> GetAgentsInPool(PoolId poolId) => agents.FindAll(a => a.GetPools().Any(p => p == poolId));
+			List<IPool> pools = await _poolCollection.GetAsync();
 
-			Dictionary<PoolSizeStrategy, List<PoolSizeData>> Result = new();
-			foreach (PoolSizeStrategy StrategyType in Enum.GetValues<PoolSizeStrategy>())
+			Dictionary<PoolSizeStrategy, List<PoolSizeData>> result = new();
+			foreach (PoolSizeStrategy strategyType in Enum.GetValues<PoolSizeStrategy>())
 			{
-				Result[StrategyType] = Pools
-					.Where(x => (x.SizeStrategy ?? Settings.DefaultAgentPoolSizeStrategy) == StrategyType)
+				result[strategyType] = pools
+					.Where(x => (x.SizeStrategy ?? _settings.DefaultAgentPoolSizeStrategy) == strategyType)
 					.Select(x => new PoolSizeData(x, GetAgentsInPool(x.Id), null))
 					.ToList();
 			}
 
-			return Result;
+			return result;
 		}
 
-		internal async Task ResizePools(List<PoolSizeData> PoolSizeDatas)
+		internal async Task ResizePools(List<PoolSizeData> poolSizeDatas)
 		{
-			foreach (PoolSizeData PoolSizeData in PoolSizeDatas.OrderByDescending(x => x.Agents.Count))
+			foreach (PoolSizeData poolSizeData in poolSizeDatas.OrderByDescending(x => x.Agents.Count))
 			{
-				IPool Pool = PoolSizeData.Pool;
+				IPool pool = poolSizeData.Pool;
 
-				if (!Pool.EnableAutoscaling || PoolSizeData.DesiredAgentCount == null) continue;
+				if (!pool.EnableAutoscaling || poolSizeData.DesiredAgentCount == null)
+				{
+					continue;
+				}
 
-				int CurrentAgentCount = PoolSizeData.Agents.Count;
-				int DesiredAgentCount = PoolSizeData.DesiredAgentCount.Value;
-				int DeltaAgentCount = DesiredAgentCount - CurrentAgentCount;
+				int currentAgentCount = poolSizeData.Agents.Count;
+				int desiredAgentCount = poolSizeData.DesiredAgentCount.Value;
+				int deltaAgentCount = desiredAgentCount - currentAgentCount;
 
-				Logger.LogInformation("{PoolName,-48} Current={Current,4} Target={Target,4} Delta={Delta,4} Status={Status}", Pool.Name, CurrentAgentCount, DesiredAgentCount, DeltaAgentCount, PoolSizeData.StatusMessage);
+				_logger.LogInformation("{PoolName,-48} Current={Current,4} Target={Target,4} Delta={Delta,4} Status={Status}", pool.Name, currentAgentCount, desiredAgentCount, deltaAgentCount, poolSizeData.StatusMessage);
 				
 				try
 				{
-					using IScope Scope = GlobalTracer.Instance.BuildSpan("ScalingPool").StartActive();
-					Scope.Span.SetTag("poolName", Pool.Name);
-					Scope.Span.SetTag("current", CurrentAgentCount);
-					Scope.Span.SetTag("desired", DesiredAgentCount);
-					Scope.Span.SetTag("delta", DeltaAgentCount);
+					using IScope scope = GlobalTracer.Instance.BuildSpan("ScalingPool").StartActive();
+					scope.Span.SetTag("poolName", pool.Name);
+					scope.Span.SetTag("current", currentAgentCount);
+					scope.Span.SetTag("desired", desiredAgentCount);
+					scope.Span.SetTag("delta", deltaAgentCount);
 
-					if (DeltaAgentCount > 0)
+					if (deltaAgentCount > 0)
 					{
-						TimeSpan ScaleOutCooldown = Pool.ScaleOutCooldown ?? DefaultScaleOutCooldown;
-						bool IsCoolingDown = Pool.LastScaleUpTime != null && Pool.LastScaleUpTime + ScaleOutCooldown > Clock.UtcNow;
-						Scope.Span.SetTag("isCoolingDown", IsCoolingDown);
-						if (!IsCoolingDown)
+						TimeSpan scaleOutCooldown = pool.ScaleOutCooldown ?? _defaultScaleOutCooldown;
+						bool isCoolingDown = pool.LastScaleUpTime != null && pool.LastScaleUpTime + scaleOutCooldown > _clock.UtcNow;
+						scope.Span.SetTag("isCoolingDown", isCoolingDown);
+						if (!isCoolingDown)
 						{
-							await FleetManager.ExpandPoolAsync(Pool, PoolSizeData.Agents, DeltaAgentCount);
-							await PoolCollection.TryUpdateAsync(Pool, LastScaleUpTime: Clock.UtcNow);
+							await _fleetManager.ExpandPoolAsync(pool, poolSizeData.Agents, deltaAgentCount);
+							await _poolCollection.TryUpdateAsync(pool, lastScaleUpTime: _clock.UtcNow);
 						}
 						else
 						{
-							TimeSpan? CooldownTimeLeft = Pool.LastScaleUpTime + DefaultScaleOutCooldown - Clock.UtcNow;
-							Logger.LogDebug("Cannot scale out {PoolName}, it's cooling down for another {TimeLeft} secs", Pool.Name, CooldownTimeLeft?.TotalSeconds);
+							TimeSpan? cooldownTimeLeft = pool.LastScaleUpTime + _defaultScaleOutCooldown - _clock.UtcNow;
+							_logger.LogDebug("Cannot scale out {PoolName}, it's cooling down for another {TimeLeft} secs", pool.Name, cooldownTimeLeft?.TotalSeconds);
 						}
 					}
 
-					if (DeltaAgentCount < 0)
+					if (deltaAgentCount < 0)
 					{
-						TimeSpan ScaleInCooldown = Pool.ScaleInCooldown ?? DefaultScaleInCooldown;
-						bool IsCoolingDown = Pool.LastScaleDownTime != null && Pool.LastScaleDownTime + ScaleInCooldown > Clock.UtcNow;
-						Scope.Span.SetTag("isCoolingDown", IsCoolingDown);
-						if (!IsCoolingDown)
+						TimeSpan scaleInCooldown = pool.ScaleInCooldown ?? _defaultScaleInCooldown;
+						bool isCoolingDown = pool.LastScaleDownTime != null && pool.LastScaleDownTime + scaleInCooldown > _clock.UtcNow;
+						scope.Span.SetTag("isCoolingDown", isCoolingDown);
+						if (!isCoolingDown)
 						{
-							await FleetManager.ShrinkPoolAsync(Pool, PoolSizeData.Agents, -DeltaAgentCount);
-							await PoolCollection.TryUpdateAsync(Pool, LastScaleDownTime: Clock.UtcNow);
+							await _fleetManager.ShrinkPoolAsync(pool, poolSizeData.Agents, -deltaAgentCount);
+							await _poolCollection.TryUpdateAsync(pool, lastScaleDownTime: _clock.UtcNow);
 						}
 						else
 						{
-							TimeSpan? CooldownTimeLeft = Pool.LastScaleDownTime + DefaultScaleInCooldown - Clock.UtcNow;
-							Logger.LogDebug("Cannot scale in {PoolName}, it's cooling down for another {TimeLeft} secs", Pool.Name, CooldownTimeLeft?.TotalSeconds);
+							TimeSpan? cooldownTimeLeft = pool.LastScaleDownTime + _defaultScaleInCooldown - _clock.UtcNow;
+							_logger.LogDebug("Cannot scale in {PoolName}, it's cooling down for another {TimeLeft} secs", pool.Name, cooldownTimeLeft?.TotalSeconds);
 						}
 					}
 				}
-				catch (Exception Ex)
+				catch (Exception ex)
 				{
-					Logger.LogError(Ex, "Failed to scale {PoolName}:\n{Exception}", Pool.Name, Ex);
+					_logger.LogError(ex, "Failed to scale {PoolName}:\n{Exception}", pool.Name, ex);
 					continue;
 				}
 
-				DogStatsd.Gauge("agentpools.autoscale.target", DesiredAgentCount, tags: new []{"pool:" + Pool.Name});
-				DogStatsd.Gauge("agentpools.autoscale.current", CurrentAgentCount, tags: new []{"pool:" + Pool.Name});
+				_dogStatsd.Gauge("agentpools.autoscale.target", desiredAgentCount, tags: new []{"pool:" + pool.Name});
+				_dogStatsd.Gauge("agentpools.autoscale.current", currentAgentCount, tags: new []{"pool:" + pool.Name});
 			}
 		}
 
@@ -178,24 +180,24 @@ namespace Horde.Build.Fleet.Autoscale
 		/// Backdoor for tests to override strategies with test doubles
 		/// These cannot supplied in the constructor since the DI requires concrete implementations.
 		/// </summary>
-		/// <param name="LeaseUtilizationStrategy"></param>
-		/// <param name="JobQueueStrategy"></param>
-		/// <param name="NoOpStrategy"></param>
-		internal void OverridePoolSizeStrategiesDuringTesting(IPoolSizeStrategy LeaseUtilizationStrategy, IPoolSizeStrategy JobQueueStrategy, IPoolSizeStrategy NoOpStrategy)
+		/// <param name="leaseUtilizationStrategy"></param>
+		/// <param name="jobQueueStrategy"></param>
+		/// <param name="noOpStrategy"></param>
+		internal void OverridePoolSizeStrategiesDuringTesting(IPoolSizeStrategy leaseUtilizationStrategy, IPoolSizeStrategy jobQueueStrategy, IPoolSizeStrategy noOpStrategy)
 		{
-			this.LeaseUtilizationStrategy = LeaseUtilizationStrategy;
-			this.JobQueueStrategy = JobQueueStrategy;
-			this.NoOpStrategy = NoOpStrategy;
+			_leaseUtilizationStrategy = leaseUtilizationStrategy;
+			_jobQueueStrategy = jobQueueStrategy;
+			_noOpStrategy = noOpStrategy;
 		}
 		
-		private IPoolSizeStrategy GetPoolSizeStrategy(PoolSizeStrategy Type)
+		private IPoolSizeStrategy GetPoolSizeStrategy(PoolSizeStrategy type)
 		{
-			return Type switch
+			return type switch
 			{
-				PoolSizeStrategy.LeaseUtilization => LeaseUtilizationStrategy,
-				PoolSizeStrategy.JobQueue => JobQueueStrategy,
-				PoolSizeStrategy.NoOp => NoOpStrategy,
-				_ => NoOpStrategy,
+				PoolSizeStrategy.LeaseUtilization => _leaseUtilizationStrategy,
+				PoolSizeStrategy.JobQueue => _jobQueueStrategy,
+				PoolSizeStrategy.NoOp => _noOpStrategy,
+				_ => _noOpStrategy,
 			};
 		}
 	}
diff --git a/Engine/Source/Programs/Horde/Horde.Build/Fleet/Autoscale/IPoolSizeStrategy.cs b/Engine/Source/Programs/Horde/Horde.Build/Fleet/Autoscale/IPoolSizeStrategy.cs
--- a/Engine/Source/Programs/Horde/Horde.Build/Fleet/Autoscale/IPoolSizeStrategy.cs
+++ b/Engine/Source/Programs/Horde/Horde.Build/Fleet/Autoscale/IPoolSizeStrategy.cs
@@ -4,12 +4,9 @@ using System.Collections.Generic;
 using System.Linq;
 using System.Threading.Tasks;
 using Horde.Build.Models;
-using Horde.Build.Utilities;
 
 namespace Horde.Build.Fleet.Autoscale
 {
-	using PoolId = StringId<IPool>;
-
 	/// <summary>
 	/// Available pool sizing strategies
 	/// </summary>
@@ -62,30 +59,30 @@ namespace Horde.Build.Fleet.Autoscale
 		/// <summary>
 		/// Constructor
 		/// </summary>
-		/// <param name="Pool"></param>
-		/// <param name="Agents"></param>
-		/// <param name="DesiredAgentCount"></param>
-		/// <param name="StatusMessage"></param>
-		public PoolSizeData(IPool Pool, List<IAgent> Agents, int? DesiredAgentCount, string StatusMessage = "N/A")
+		/// <param name="pool"></param>
+		/// <param name="agents"></param>
+		/// <param name="desiredAgentCount"></param>
+		/// <param name="statusMessage"></param>
+		public PoolSizeData(IPool pool, List<IAgent> agents, int? desiredAgentCount, string statusMessage = "N/A")
 		{
-			this.Pool = Pool;
-			this.Agents = Agents;
-			this.DesiredAgentCount = DesiredAgentCount;
-			this.StatusMessage = StatusMessage;
+			Pool = pool;
+			Agents = agents;
+			DesiredAgentCount = desiredAgentCount;
+			StatusMessage = statusMessage;
 		}
 
 		/// <summary>
 		/// Copy the object, inheriting any unspecified values from current instance
 		/// Needed because the class is immutable 
 		/// </summary>
-		/// <param name="Pool"></param>
-		/// <param name="Agents"></param>
-		/// <param name="DesiredAgentCount"></param>
-		/// <param name="StatusMessage"></param>
+		/// <param name="pool"></param>
+		/// <param name="agents"></param>
+		/// <param name="desiredAgentCount"></param>
+		/// <param name="statusMessage"></param>
 		/// <returns>A new copy</returns>
-		public PoolSizeData Copy(IPool? Pool = null, List<IAgent>? Agents = null, int? DesiredAgentCount = null, string? StatusMessage = null)
+		public PoolSizeData Copy(IPool? pool = null, List<IAgent>? agents = null, int? desiredAgentCount = null, string? statusMessage = null)
 		{
-			return new PoolSizeData(Pool ?? this.Pool, Agents ?? this.Agents, DesiredAgentCount ?? this.DesiredAgentCount, StatusMessage ?? this.StatusMessage);
+			return new PoolSizeData(pool ?? Pool, agents ?? Agents, desiredAgentCount ?? DesiredAgentCount, statusMessage ?? StatusMessage);
 		}
 	}
 	
@@ -97,9 +94,9 @@ namespace Horde.Build.Fleet.Autoscale
 		/// <summary>
 		/// Calculate the adequate number of agents to be online for given pools
 		/// </summary>
-		/// <param name="Pools">Pools including attached agents</param>
+		/// <param name="pools">Pools including attached agents</param>
 		/// <returns></returns>
-		Task<List<PoolSizeData>> CalcDesiredPoolSizesAsync(List<PoolSizeData> Pools);
+		Task<List<PoolSizeData>> CalcDesiredPoolSizesAsync(List<PoolSizeData> pools);
 		
 		/// <summary>
 		/// Name of the strategy
@@ -114,10 +111,10 @@ namespace Horde.Build.Fleet.Autoscale
 	public class NoOpPoolSizeStrategy : IPoolSizeStrategy
 	{
 		/// <inheritdoc/>
-		public Task<List<PoolSizeData>> CalcDesiredPoolSizesAsync(List<PoolSizeData> Pools)
+		public Task<List<PoolSizeData>> CalcDesiredPoolSizesAsync(List<PoolSizeData> pools)
 		{
-			List<PoolSizeData> Result = Pools.Select(x => new PoolSizeData(x.Pool, x.Agents, x.Agents.Count, "(no-op)")).ToList();
-			return Task.FromResult(Result);
+			List<PoolSizeData> result = pools.Select(x => new PoolSizeData(x.Pool, x.Agents, x.Agents.Count, "(no-op)")).ToList();
+			return Task.FromResult(result);
 		}
 
 		/// <inheritdoc/>
diff --git a/Engine/Source/Programs/Horde/Horde.Build/Fleet/Autoscale/JobQueueStrategy.cs b/Engine/Source/Programs/Horde/Horde.Build/Fleet/Autoscale/JobQueueStrategy.cs
--- a/Engine/Source/Programs/Horde/Horde.Build/Fleet/Autoscale/JobQueueStrategy.cs
+++ b/Engine/Source/Programs/Horde/Horde.Build/Fleet/Autoscale/JobQueueStrategy.cs
@@ -4,18 +4,18 @@ using System;
 using System.Collections.Generic;
 using System.Linq;
 using System.Threading.Tasks;
-using HordeCommon;
 using Horde.Build.Api;
 using Horde.Build.Collections;
 using Horde.Build.Models;
 using Horde.Build.Services;
 using Horde.Build.Utilities;
+using HordeCommon;
 
 namespace Horde.Build.Fleet.Autoscale
 {
 	using PoolId = StringId<IPool>;
 	using StreamId = StringId<IStream>;
-	
+
 	/// <summary>
 	/// Job queue sizing settings for a pool
 	/// </summary>
@@ -38,12 +38,12 @@ namespace Horde.Build.Fleet.Autoscale
 		/// <summary>
 		/// Constructor
 		/// </summary>
-		/// <param name="ScaleOutFactor"></param>
-		/// <param name="ScaleInFactor"></param>
-		public JobQueueSettings(double? ScaleOutFactor = null, double? ScaleInFactor = null)
+		/// <param name="scaleOutFactor"></param>
+		/// <param name="scaleInFactor"></param>
+		public JobQueueSettings(double? scaleOutFactor = null, double? scaleInFactor = null)
 		{
-			this.ScaleOutFactor = ScaleOutFactor.GetValueOrDefault(this.ScaleOutFactor);
-			this.ScaleInFactor = ScaleInFactor.GetValueOrDefault(this.ScaleInFactor);
+			ScaleOutFactor = scaleOutFactor.GetValueOrDefault(ScaleOutFactor);
+			ScaleInFactor = scaleInFactor.GetValueOrDefault(ScaleInFactor);
 		}
 	}
 	
@@ -55,15 +55,15 @@ namespace Horde.Build.Fleet.Autoscale
 	/// </summary>
 	public class JobQueueStrategy : IPoolSizeStrategy
 	{
-		private readonly IJobCollection Jobs;
-		private readonly IGraphCollection Graphs;
-		private readonly StreamService StreamService;
-		private readonly IClock Clock;
+		private readonly IJobCollection _jobs;
+		private readonly IGraphCollection _graphs;
+		private readonly StreamService _streamService;
+		private readonly IClock _clock;
 		
 		/// <summary>
 		/// How far back in time to look for job batches (that potentially are in the queue)
 		/// </summary>
-		private readonly TimeSpan SamplePeriod = TimeSpan.FromDays(5);
+		private readonly TimeSpan _samplePeriod = TimeSpan.FromDays(5);
 		
 		/// <summary>
 		/// Time spent in ready state before considered truly waiting for an agent
@@ -73,23 +73,23 @@ namespace Horde.Build.Fleet.Autoscale
 		/// </summary>
 		internal readonly TimeSpan ReadyTimeThreshold = TimeSpan.FromSeconds(45.0);
 
-		private readonly JobQueueSettings DefaultPoolSettings = new();
+		private readonly JobQueueSettings _defaultPoolSettings = new();
 
 		/// <summary>
 		/// Constructor
 		/// </summary>
-		/// <param name="Jobs"></param>
-		/// <param name="Graphs"></param>
-		/// <param name="StreamService"></param>
-		/// <param name="Clock"></param>
-		/// <param name="SamplePeriod">Time period for each sample</param>
-		public JobQueueStrategy(IJobCollection Jobs, IGraphCollection Graphs, StreamService StreamService, IClock Clock, TimeSpan? SamplePeriod = null)
+		/// <param name="jobs"></param>
+		/// <param name="graphs"></param>
+		/// <param name="streamService"></param>
+		/// <param name="clock"></param>
+		/// <param name="samplePeriod">Time period for each sample</param>
+		public JobQueueStrategy(IJobCollection jobs, IGraphCollection graphs, StreamService streamService, IClock clock, TimeSpan? samplePeriod = null)
 		{
-			this.Jobs = Jobs;
-			this.Graphs = Graphs;
-			this.StreamService = StreamService;
-			this.Clock = Clock;
-			this.SamplePeriod = SamplePeriod ?? this.SamplePeriod;
+			_jobs = jobs;
+			_graphs = graphs;
+			_streamService = streamService;
+			_clock = clock;
+			_samplePeriod = samplePeriod ?? _samplePeriod;
 		}
 
 		/// <inheritdoc/>
@@ -98,75 +98,84 @@ namespace Horde.Build.Fleet.Autoscale
 		/// <summary>
 		/// Extract all job step batches from a job, with their associated pool 
 		/// </summary>
-		/// <param name="Job">Job to extract from</param>
-		/// <param name="Streams">Cached lookup table of streams</param>
+		/// <param name="job">Job to extract from</param>
+		/// <param name="streams">Cached lookup table of streams</param>
 		/// <returns></returns>
-		private async Task<List<(IJob Job, IJobStepBatch Batch, PoolId PoolId)>> GetJobBatchesWithPools(IJob Job, Dictionary<StreamId, IStream> Streams)
+		private async Task<List<(IJob Job, IJobStepBatch Batch, PoolId PoolId)>> GetJobBatchesWithPools(IJob job, Dictionary<StreamId, IStream> streams)
 		{
-			IGraph Graph = await Graphs.GetAsync(Job.GraphHash);
+			IGraph graph = await _graphs.GetAsync(job.GraphHash);
 
-			List<(IJob Job, IJobStepBatch Batch, PoolId PoolId)> JobBatches = new();
-			foreach (IJobStepBatch Batch in Job.Batches)
+			List<(IJob Job, IJobStepBatch Batch, PoolId PoolId)> jobBatches = new();
+			foreach (IJobStepBatch batch in job.Batches)
 			{
-				if (Batch.State != JobStepBatchState.Ready) continue;
+				if (batch.State != JobStepBatchState.Ready)
+				{
+					continue;
+				}
 				
-				TimeSpan? WaitTime = Clock.UtcNow - Batch.ReadyTimeUtc;
-				if (WaitTime == null) continue;
-				if (WaitTime.Value < ReadyTimeThreshold) continue;
+				TimeSpan? waitTime = _clock.UtcNow - batch.ReadyTimeUtc;
+				if (waitTime == null)
+				{
+					continue;
+				}
+				if (waitTime.Value < ReadyTimeThreshold)
+				{
+					continue;
+				}
 
-				if (!Streams.TryGetValue(Job.StreamId, out IStream? Stream))
+				if (!streams.TryGetValue(job.StreamId, out IStream? stream))
 				{
 					continue;
 				}
 
-				string BatchAgentType = Graph.Groups[Batch.GroupIdx].AgentType;
-				if (!Stream.AgentTypes.TryGetValue(BatchAgentType, out AgentType? AgentType))
+				string batchAgentType = graph.Groups[batch.GroupIdx].AgentType;
+				if (!stream.AgentTypes.TryGetValue(batchAgentType, out AgentType? agentType))
 				{
 					continue;
 				}
 				
-				JobBatches.Add((Job, Batch, AgentType.Pool));
+				jobBatches.Add((job, batch, agentType.Pool));
 			}
 
-			return JobBatches;
+			return jobBatches;
 		}
 
-		internal async Task<Dictionary<PoolId, int>> GetPoolQueueSizesAsync(DateTimeOffset JobsCreatedAfter)
+		internal async Task<Dictionary<PoolId, int>> GetPoolQueueSizesAsync(DateTimeOffset jobsCreatedAfter)
 		{
-			List<IStream> StreamsList = await StreamService.GetStreamsAsync();
-			Dictionary<StreamId, IStream> Streams = StreamsList.ToDictionary(x => x.Id, x => x);
-			List<IJob> RecentJobs = await Jobs.FindAsync(MinCreateTime: JobsCreatedAfter);
+			List<IStream> streamsList = await _streamService.GetStreamsAsync();
+			Dictionary<StreamId, IStream> streams = streamsList.ToDictionary(x => x.Id, x => x);
+			List<IJob> recentJobs = await _jobs.FindAsync(minCreateTime: jobsCreatedAfter);
 
-			List<(IJob Job, IJobStepBatch Batch, PoolId PoolId)> JobBatches = new();
-			foreach (IJob Job in RecentJobs)
+			List<(IJob Job, IJobStepBatch Batch, PoolId PoolId)> jobBatches = new();
+			foreach (IJob job in recentJobs)
 			{
-				JobBatches.AddRange(await GetJobBatchesWithPools(Job, Streams));
+				jobBatches.AddRange(await GetJobBatchesWithPools(job, streams));
 			}
 
-			List<(PoolId PoolId, int QueueSize)> PoolsWithQueueSize = JobBatches.GroupBy(t => t.PoolId).Select(t => (t.Key, t.Count())).ToList();
-			return PoolsWithQueueSize.ToDictionary(x => x.PoolId, x => x.QueueSize);
+			List<(PoolId PoolId, int QueueSize)> poolsWithQueueSize = jobBatches.GroupBy(t => t.PoolId).Select(t => (t.Key, t.Count())).ToList();
+			return poolsWithQueueSize.ToDictionary(x => x.PoolId, x => x.QueueSize);
 		}
 
 		/// <inheritdoc/>
-		public async Task<List<PoolSizeData>> CalcDesiredPoolSizesAsync(List<PoolSizeData> Pools)
+		public async Task<List<PoolSizeData>> CalcDesiredPoolSizesAsync(List<PoolSizeData> pools)
 		{
-			DateTimeOffset MinCreateTime = Clock.UtcNow - SamplePeriod;
-			Dictionary<PoolId, int> PoolQueueSizes = await GetPoolQueueSizesAsync(MinCreateTime);
+			DateTimeOffset minCreateTime = _clock.UtcNow - _samplePeriod;
+			Dictionary<PoolId, int> poolQueueSizes = await GetPoolQueueSizesAsync(minCreateTime);
 
-			return Pools.Select(Current =>
+			return pools.Select(current =>
 			{
-				JobQueueSettings Settings = Current.Pool.JobQueueSettings ?? DefaultPoolSettings;
-				PoolQueueSizes.TryGetValue(Current.Pool.Id, out var QueueSize);
-				if (QueueSize > 0)
+				JobQueueSettings settings = current.Pool.JobQueueSettings ?? _defaultPoolSettings;
+				poolQueueSizes.TryGetValue(current.Pool.Id, out int queueSize);
+				if (queueSize > 0)
 				{
-					int AdditionalAgentCount = (int)Math.Ceiling(QueueSize * Settings.ScaleOutFactor);
-					int DesiredAgentCount = Current.Agents.Count + AdditionalAgentCount;
-					return new PoolSizeData(Current.Pool, Current.Agents, DesiredAgentCount, $"QueueSize={QueueSize}");
+					int additionalAgentCount = (int)Math.Ceiling(queueSize * settings.ScaleOutFactor);
+					int desiredAgentCount = current.Agents.Count + additionalAgentCount;
+					return new PoolSizeData(current.Pool, current.Agents, desiredAgentCount, $"QueueSize={queueSize}");
 				}
 				else
 				{
-					int DesiredAgentCount = (int)(Current.Agents.Count * Settings.ScaleInFactor);
-					return new PoolSizeData(Current.Pool, Current.Agents, DesiredAgentCount, "Empty job queue");
+					int desiredAgentCount = (int)(current.Agents.Count * settings.ScaleInFactor);
+					return new PoolSizeData(current.Pool, current.Agents, desiredAgentCount, "Empty job queue");
 				}
 			}).ToList();
 		}
diff --git a/Engine/Source/Programs/Horde/Horde.Build/Fleet/Autoscale/LeaseUtilizationStrategy.cs b/Engine/Source/Programs/Horde/Horde.Build/Fleet/Autoscale/LeaseUtilizationStrategy.cs
--- a/Engine/Source/Programs/Horde/Horde.Build/Fleet/Autoscale/LeaseUtilizationStrategy.cs
+++ b/Engine/Source/Programs/Horde/Horde.Build/Fleet/Autoscale/LeaseUtilizationStrategy.cs
@@ -6,11 +6,11 @@ using System.Linq;
 using System.Text;
 using System.Threading.Tasks;
 using Google.Protobuf.WellKnownTypes;
-using HordeCommon;
-using HordeCommon.Rpc.Tasks;
 using Horde.Build.Collections;
 using Horde.Build.Models;
 using Horde.Build.Utilities;
+using HordeCommon;
+using HordeCommon.Rpc.Tasks;
 using OpenTracing;
 using OpenTracing.Util;
 
@@ -32,8 +32,8 @@ namespace Horde.Build.Fleet.Autoscale
 	{
 		struct UtilizationSample
 		{
-			public double JobWork;
-			public double OtherWork;
+			public double _jobWork;
+			public double _otherWork;
 		}
 
 		class AgentData
@@ -41,24 +41,24 @@ namespace Horde.Build.Fleet.Autoscale
 			public IAgent Agent { get; }
 			public UtilizationSample[] Samples { get; }
 			public List<ILease> Leases { get; } = new ();
-			private readonly int NumSamples;
+			private readonly int _numSamples;
 
-			public AgentData(IAgent Agent, int NumSamples)
+			public AgentData(IAgent agent, int numSamples)
 			{
-				this.Agent = Agent;
-				this.NumSamples = NumSamples;
-				Samples = new UtilizationSample[NumSamples];
+				Agent = agent;
+				_numSamples = numSamples;
+				Samples = new UtilizationSample[numSamples];
 			}
 
-			public void Add(double MinT, double MaxT, double JobWork, double OtherWork)
+			public void Add(double minT, double maxT, double jobWork, double otherWork)
 			{
-				int MinIdx = Math.Clamp((int)MinT, 0, NumSamples - 1);
-				int MaxIdx = Math.Clamp((int)MaxT, MinIdx, NumSamples - 1);
-				for (int Idx = MinIdx; Idx <= MaxIdx; Idx++)
+				int minIdx = Math.Clamp((int)minT, 0, _numSamples - 1);
+				int maxIdx = Math.Clamp((int)maxT, minIdx, _numSamples - 1);
+				for (int idx = minIdx; idx <= maxIdx; idx++)
 				{
-					double Fraction = Math.Clamp(MaxT - Idx, 0.0, 1.0) - Math.Clamp(MinT - Idx, 0.0, 1.0);
-					Samples[Idx].JobWork += JobWork * Fraction;
-					Samples[Idx].OtherWork += OtherWork * Fraction;
+					double fraction = Math.Clamp(maxT - idx, 0.0, 1.0) - Math.Clamp(minT - idx, 0.0, 1.0);
+					Samples[idx]._jobWork += jobWork * fraction;
+					Samples[idx]._otherWork += otherWork * fraction;
 				}
 			}
 		}
@@ -69,50 +69,50 @@ namespace Horde.Build.Fleet.Autoscale
 			public List<IAgent> Agents { get; } = new ();
 			public UtilizationSample[] Samples { get; }
 
-			public PoolData(IPool Pool, int NumSamples)
+			public PoolData(IPool pool, int numSamples)
 			{
-				this.Pool = Pool;
-				Samples = new UtilizationSample[NumSamples];
+				Pool = pool;
+				Samples = new UtilizationSample[numSamples];
 			}
 
-			public void Add(AgentData AgentData)
+			public void Add(AgentData agentData)
 			{
-				Agents.Add(AgentData.Agent);
-				for (int Idx = 0; Idx < Samples.Length; Idx++)
+				Agents.Add(agentData.Agent);
+				for (int idx = 0; idx < Samples.Length; idx++)
 				{
-					Samples[Idx].JobWork += AgentData.Samples[Idx].JobWork;
-					Samples[Idx].OtherWork += AgentData.Samples[Idx].OtherWork;
+					Samples[idx]._jobWork += agentData.Samples[idx]._jobWork;
+					Samples[idx]._otherWork += agentData.Samples[idx]._otherWork;
 				}
 			}
 		}
 		
-		private readonly IAgentCollection AgentCollection;
-		private readonly IPoolCollection PoolCollection;
-		private readonly ILeaseCollection LeaseCollection;
-		private readonly IClock Clock;
-		private readonly TimeSpan SampleTime = TimeSpan.FromMinutes(6.0);
-		private readonly int NumSamples = 10;
-		private readonly int NumSamplesForResult = 9;
+		private readonly IAgentCollection _agentCollection;
+		private readonly IPoolCollection _poolCollection;
+		private readonly ILeaseCollection _leaseCollection;
+		private readonly IClock _clock;
+		private readonly TimeSpan _sampleTime = TimeSpan.FromMinutes(6.0);
+		private readonly int _numSamples = 10;
+		private readonly int _numSamplesForResult = 9;
 
 		/// <summary>
 		/// Constructor
 		/// </summary>
-		/// <param name="AgentCollection"></param>
-		/// <param name="PoolCollection"></param>
-		/// <param name="LeaseCollection"></param>
-		/// <param name="Clock"></param>
-		/// <param name="NumSamples">Number of samples to collect for calculating lease utilization</param>
-		/// <param name="NumSamplesForResult">Min number of samples for a valid result</param>
-		/// <param name="SampleTime">Time period for each sample</param>
-		public LeaseUtilizationStrategy(IAgentCollection AgentCollection, IPoolCollection PoolCollection, ILeaseCollection LeaseCollection, IClock Clock, int NumSamples = 10, int NumSamplesForResult = 9, TimeSpan? SampleTime = null)
+		/// <param name="agentCollection"></param>
+		/// <param name="poolCollection"></param>
+		/// <param name="leaseCollection"></param>
+		/// <param name="clock"></param>
+		/// <param name="numSamples">Number of samples to collect for calculating lease utilization</param>
+		/// <param name="numSamplesForResult">Min number of samples for a valid result</param>
+		/// <param name="sampleTime">Time period for each sample</param>
+		public LeaseUtilizationStrategy(IAgentCollection agentCollection, IPoolCollection poolCollection, ILeaseCollection leaseCollection, IClock clock, int numSamples = 10, int numSamplesForResult = 9, TimeSpan? sampleTime = null)
 		{
-			this.AgentCollection = AgentCollection;
-			this.PoolCollection = PoolCollection;
-			this.LeaseCollection = LeaseCollection;
-			this.Clock = Clock;
-			this.NumSamples = NumSamples;
-			this.NumSamplesForResult = NumSamplesForResult;
-			this.SampleTime = SampleTime ?? this.SampleTime;
+			_agentCollection = agentCollection;
+			_poolCollection = poolCollection;
+			_leaseCollection = leaseCollection;
+			_clock = clock;
+			_numSamples = numSamples;
+			_numSamplesForResult = numSamplesForResult;
+			_sampleTime = sampleTime ?? _sampleTime;
 		}
 
 		private async Task<Dictionary<AgentId, AgentData>> GetAgentDataAsync()
@@ -120,125 +120,125 @@ namespace Horde.Build.Fleet.Autoscale
 			using IScope _ = GlobalTracer.Instance.BuildSpan("GetAgentDataAsync").StartActive();
 			
 			// Find all the current agents
-			List<IAgent> Agents = await AgentCollection.FindAsync(Status: AgentStatus.Ok);
+			List<IAgent> agents = await _agentCollection.FindAsync(status: AgentStatus.Ok);
 
 			// Query leases in last interval
-			DateTime MaxTime = Clock.UtcNow;
-			DateTime MinTime = MaxTime - (SampleTime * NumSamples);
-			List<ILease> Leases = await LeaseCollection.FindLeasesAsync(MinTime, MaxTime);
+			DateTime maxTime = _clock.UtcNow;
+			DateTime minTime = maxTime - (_sampleTime * _numSamples);
+			List<ILease> leases = await _leaseCollection.FindLeasesAsync(minTime, maxTime);
 
 			// Add all the leases to a data object for each agent
-			Dictionary<AgentId, AgentData> AgentIdToData = Agents.ToDictionary(x => x.Id, x => new AgentData(x, NumSamples));
-			foreach (ILease Lease in Leases)
+			Dictionary<AgentId, AgentData> agentIdToData = agents.ToDictionary(x => x.Id, x => new AgentData(x, _numSamples));
+			foreach (ILease lease in leases)
 			{
-				AgentData? AgentData;
-				if (AgentIdToData.TryGetValue(Lease.AgentId, out AgentData) && AgentData.Agent.SessionId == Lease.SessionId)
+				AgentData? agentData;
+				if (agentIdToData.TryGetValue(lease.AgentId, out agentData) && agentData.Agent.SessionId == lease.SessionId)
 				{
-					AgentData.Leases.Add(Lease);
+					agentData.Leases.Add(lease);
 				}
 			}
 			
 			// Compute utilization for each agent
-			foreach (AgentData AgentData in AgentIdToData.Values)
+			foreach (AgentData agentData in agentIdToData.Values)
 			{
-				foreach (ILease Lease in AgentData.Leases.OrderBy(x => x.StartTime))
+				foreach (ILease lease in agentData.Leases.OrderBy(x => x.StartTime))
 				{
-					double MinT = (Lease.StartTime - MinTime).TotalSeconds / SampleTime.TotalSeconds;
-					double MaxT = (Lease.FinishTime == null)
-						? NumSamples
-						: ((Lease.FinishTime.Value - MinTime).TotalSeconds / SampleTime.TotalSeconds);
+					double minT = (lease.StartTime - minTime).TotalSeconds / _sampleTime.TotalSeconds;
+					double maxT = (lease.FinishTime == null)
+						? _numSamples
+						: ((lease.FinishTime.Value - minTime).TotalSeconds / _sampleTime.TotalSeconds);
 
-					Any Payload = Any.Parser.ParseFrom(Lease.Payload.ToArray());
-					if (Payload.Is(ExecuteJobTask.Descriptor))
+					Any payload = Any.Parser.ParseFrom(lease.Payload.ToArray());
+					if (payload.Is(ExecuteJobTask.Descriptor))
 					{
-						AgentData.Add(MinT, MaxT, 1.0, 0.0);
+						agentData.Add(minT, maxT, 1.0, 0.0);
 					}
 					else
 					{
-						AgentData.Add(MinT, MaxT, 0.0, 1.0);
+						agentData.Add(minT, maxT, 0.0, 1.0);
 					}
 				}
 			}
 
-			return AgentIdToData;
+			return agentIdToData;
 		}
 		
 		private async Task<Dictionary<PoolId, PoolData>> GetPoolDataAsync()
 		{
 			using IScope _ = GlobalTracer.Instance.BuildSpan("GetPoolDataAsync").StartActive();
 
-			Dictionary<AgentId, AgentData> AgentIdToData = await GetAgentDataAsync();
+			Dictionary<AgentId, AgentData> agentIdToData = await GetAgentDataAsync();
 			
 			// Get all the pools
-			List<IPool> Pools = await PoolCollection.GetAsync();
-			Dictionary<PoolId, PoolData> PoolToData = Pools.ToDictionary(x => x.Id, x => new PoolData(x, NumSamples));
+			List<IPool> pools = await _poolCollection.GetAsync();
+			Dictionary<PoolId, PoolData> poolToData = pools.ToDictionary(x => x.Id, x => new PoolData(x, _numSamples));
 
 			// Find pool utilization over the query period
-			foreach (AgentData AgentData in AgentIdToData.Values)
+			foreach (AgentData agentData in agentIdToData.Values)
 			{
-				foreach (PoolId PoolId in AgentData.Agent.GetPools())
+				foreach (PoolId poolId in agentData.Agent.GetPools())
 				{
-					PoolData? PoolData;
-					if (PoolToData.TryGetValue(PoolId, out PoolData))
+					PoolData? poolData;
+					if (poolToData.TryGetValue(poolId, out poolData))
 					{
-						PoolData.Add(AgentData);
+						poolData.Add(agentData);
 					}
 				}
 			}
 
-			return PoolToData;
+			return poolToData;
 		}
 
 		/// <inheritdoc/>
 		public string Name { get; } = "LeaseUtilization";
 
 		/// <inheritdoc/>
-		public async Task<List<PoolSizeData>> CalcDesiredPoolSizesAsync(List<PoolSizeData> Pools)
+		public async Task<List<PoolSizeData>> CalcDesiredPoolSizesAsync(List<PoolSizeData> pools)
 		{
-			Dictionary<PoolId, PoolData> PoolToData = await GetPoolDataAsync();
-			List<PoolSizeData> Result = new();
+			Dictionary<PoolId, PoolData> poolToData = await GetPoolDataAsync();
+			List<PoolSizeData> result = new();
 
-			foreach (PoolData PoolData in PoolToData.Values.OrderByDescending(x => x.Agents.Count))
+			foreach (PoolData poolData in poolToData.Values.OrderByDescending(x => x.Agents.Count))
 			{
-				IPool Pool = PoolData.Pool;
+				IPool pool = poolData.Pool;
 
-				PoolSizeData? PoolSize = Pools.Find(x => x.Pool.Id == Pool.Id);
-				if (PoolSize != null)
+				PoolSizeData? poolSize = pools.Find(x => x.Pool.Id == pool.Id);
+				if (poolSize != null)
 				{
-					int MinAgents = Pool.MinAgents ?? 1;
-					int NumReserveAgents = Pool.NumReserveAgents ?? 5;
-					double Utilization = PoolData.Samples.Select(x => x.JobWork).OrderByDescending(x => x).Skip(NumSamples - NumSamplesForResult).First();
+					int minAgents = pool.MinAgents ?? 1;
+					int numReserveAgents = pool.NumReserveAgents ?? 5;
+					double utilization = poolData.Samples.Select(x => x._jobWork).OrderByDescending(x => x).Skip(_numSamples - _numSamplesForResult).First();
 				
 					// Number of agents in use over the sampling period. Can never be greater than number of agents available in pool.
-					int NumAgentsUtilized = (int)Utilization;
+					int numAgentsUtilized = (int)utilization;
 					
 					// Include reserve agent count to ensure pool always can grow
-					int DesiredAgentCount = Math.Max(NumAgentsUtilized + NumReserveAgents, MinAgents);;
+					int desiredAgentCount = Math.Max(numAgentsUtilized + numReserveAgents, minAgents);
 					
-					StringBuilder Sb = new();
-					Sb.AppendFormat("Jobs=[{0}] ", GetDensityMap(PoolData.Samples.Select(x => x.JobWork / Math.Max(1, PoolData.Agents.Count))));
-					Sb.AppendFormat("Total=[{0}] ", GetDensityMap(PoolData.Samples.Select(x => x.OtherWork / Math.Max(1, PoolData.Agents.Count))));
-					Sb.AppendFormat("Min=[{0,5:0.0}] ", PoolData.Samples.Min(x => x.JobWork));
-					Sb.AppendFormat("Max=[{0,5:0.0}] ", PoolData.Samples.Max(x => x.JobWork));
-					Sb.AppendFormat("Avg=[{0,5:0.0}] ", Utilization);
-					Sb.AppendFormat("Pct=[{0,5:0.0}] ", PoolData.Samples.Sum(x => x.JobWork) / NumSamples);
-
-					Result.Add(new(Pool, PoolSize.Agents, DesiredAgentCount, Sb.ToString()));
+					StringBuilder sb = new();
+					sb.AppendFormat("Jobs=[{0}] ", GetDensityMap(poolData.Samples.Select(x => x._jobWork / Math.Max(1, poolData.Agents.Count))));
+					sb.AppendFormat("Total=[{0}] ", GetDensityMap(poolData.Samples.Select(x => x._otherWork / Math.Max(1, poolData.Agents.Count))));
+					sb.AppendFormat("Min=[{0,5:0.0}] ", poolData.Samples.Min(x => x._jobWork));
+					sb.AppendFormat("Max=[{0,5:0.0}] ", poolData.Samples.Max(x => x._jobWork));
+					sb.AppendFormat("Avg=[{0,5:0.0}] ", utilization);
+					sb.AppendFormat("Pct=[{0,5:0.0}] ", poolData.Samples.Sum(x => x._jobWork) / _numSamples);
+
+					result.Add(new(pool, poolSize.Agents, desiredAgentCount, sb.ToString()));
 				}
 			}
 
-			return Result;
+			return result;
 		}
 		
 		/// <summary>
 		/// Creates a string of characters indicating a sequence of 0-1 values over time
 		/// </summary>
-		/// <param name="Values">Sequence of values</param>
+		/// <param name="values">Sequence of values</param>
 		/// <returns>Density map</returns>
-		private static string GetDensityMap(IEnumerable<double> Values)
+		private static string GetDensityMap(IEnumerable<double> values)
 		{
 			const string Greyscale = " 123456789";
-			return new string(Values.Select(x => Greyscale[Math.Clamp((int)(x * Greyscale.Length), 0, Greyscale.Length - 1)]).ToArray());
+			return new string(values.Select(x => Greyscale[Math.Clamp((int)(x * Greyscale.Length), 0, Greyscale.Length - 1)]).ToArray());
 		}
 	}
 }
diff --git a/Engine/Source/Programs/Horde/Horde.Build/Fleet/IFleetManager.cs b/Engine/Source/Programs/Horde/Horde.Build/Fleet/IFleetManager.cs
--- a/Engine/Source/Programs/Horde/Horde.Build/Fleet/IFleetManager.cs
+++ b/Engine/Source/Programs/Horde/Horde.Build/Fleet/IFleetManager.cs
@@ -1,10 +1,8 @@
 // Copyright Epic Games, Inc. All Rights Reserved.
 
-using Horde.Build.Models;
-using System;
 using System.Collections.Generic;
-using System.Linq;
 using System.Threading.Tasks;
+using Horde.Build.Models;
 
 namespace Horde.Build.Services
 {
@@ -16,27 +14,27 @@ namespace Horde.Build.Services
 		/// <summary>
 		/// Expand the given pool
 		/// </summary>
-		/// <param name="Pool">Pool to resize</param>
-		/// <param name="Agents">Current list of agents in the pool</param>
-		/// <param name="Count">Number of agents to add</param>
+		/// <param name="pool">Pool to resize</param>
+		/// <param name="agents">Current list of agents in the pool</param>
+		/// <param name="count">Number of agents to add</param>
 		/// <returns>Async task</returns>
-		Task ExpandPoolAsync(IPool Pool, IReadOnlyList<IAgent> Agents, int Count);
+		Task ExpandPoolAsync(IPool pool, IReadOnlyList<IAgent> agents, int count);
 
 		/// <summary>
 		/// Shrink the given pool
 		/// </summary>
-		/// <param name="Pool">Pool to resize</param>
-		/// <param name="Agents">Current list of agents in the pool</param>
-		/// <param name="Count">Number of agents to remove</param>
+		/// <param name="pool">Pool to resize</param>
+		/// <param name="agents">Current list of agents in the pool</param>
+		/// <param name="count">Number of agents to remove</param>
 		/// <returns>Async task</returns>
-		Task ShrinkPoolAsync(IPool Pool, IReadOnlyList<IAgent> Agents, int Count);
+		Task ShrinkPoolAsync(IPool pool, IReadOnlyList<IAgent> agents, int count);
 
 		/// <summary>
 		/// Returns the number of stopped instances in the given pool
 		/// </summary>
-		/// <param name="Pool">Pool to resize</param>
+		/// <param name="pool">Pool to resize</param>
 		/// <returns>Async task</returns>
-		Task<int> GetNumStoppedInstancesAsync(IPool Pool);
+		Task<int> GetNumStoppedInstancesAsync(IPool pool);
 	}
 }
 
diff --git a/Engine/Source/Programs/Horde/Horde.Build/Fleet/IPoolCollection.cs b/Engine/Source/Programs/Horde/Horde.Build/Fleet/IPoolCollection.cs
--- a/Engine/Source/Programs/Horde/Horde.Build/Fleet/IPoolCollection.cs
+++ b/Engine/Source/Programs/Horde/Horde.Build/Fleet/IPoolCollection.cs
@@ -1,19 +1,12 @@
 // Copyright Epic Games, Inc. All Rights Reserved.
 
-using EpicGames.Horde.Common;
-using Horde.Build.Models;
-using Horde.Build.Services;
-using Horde.Build.Utilities;
-using MongoDB.Bson;
-using MongoDB.Bson.Serialization;
-using MongoDB.Bson.Serialization.Attributes;
-using MongoDB.Driver;
 using System;
 using System.Collections.Generic;
-using System.Linq;
-using System.Linq.Expressions;
 using System.Threading.Tasks;
+using EpicGames.Horde.Common;
 using Horde.Build.Fleet.Autoscale;
+using Horde.Build.Models;
+using Horde.Build.Utilities;
 
 namespace Horde.Build.Collections
 {
@@ -27,32 +20,32 @@ namespace Horde.Build.Collections
 		/// <summary>
 		/// Creates a new pool
 		/// </summary>
-		/// <param name="Id">Unique id for the new pool</param>
-		/// <param name="Name">Name of the new pool</param>
-		/// <param name="Condition">Condition for agents to be included in this pool</param>
-		/// <param name="EnableAutoscaling">Whether to enable autoscaling for this pool</param>
-		/// <param name="MinAgents">Minimum number of agents in the pool</param>
-		/// <param name="NumReserveAgents">Minimum number of idle agents to maintain</param>
-		/// <param name="ScaleOutCooldown">Cooldown time between scale-out events</param>
-		/// <param name="ScaleInCooldown">Cooldown time between scale-in events</param>
-		/// <param name="SizeStrategy">Pool sizing strategy</param>
-		/// <param name="LeaseUtilizationSettings">Settings for lease utilization strategy</param>
-		/// <param name="JobQueueSettings">Settings for job queue strategy</param>
-		/// <param name="Properties">Properties for the pool</param>
+		/// <param name="id">Unique id for the new pool</param>
+		/// <param name="name">Name of the new pool</param>
+		/// <param name="condition">Condition for agents to be included in this pool</param>
+		/// <param name="enableAutoscaling">Whether to enable autoscaling for this pool</param>
+		/// <param name="minAgents">Minimum number of agents in the pool</param>
+		/// <param name="numReserveAgents">Minimum number of idle agents to maintain</param>
+		/// <param name="scaleOutCooldown">Cooldown time between scale-out events</param>
+		/// <param name="scaleInCooldown">Cooldown time between scale-in events</param>
+		/// <param name="sizeStrategy">Pool sizing strategy</param>
+		/// <param name="leaseUtilizationSettings">Settings for lease utilization strategy</param>
+		/// <param name="jobQueueSettings">Settings for job queue strategy</param>
+		/// <param name="properties">Properties for the pool</param>
 		/// <returns>The new pool document</returns>
 		Task<IPool> AddAsync(
-			PoolId Id,
-			string Name,
-			Condition? Condition = null,
-			bool? EnableAutoscaling = null,
-			int? MinAgents = null,
-			int? NumReserveAgents = null,
-			TimeSpan? ScaleOutCooldown = null,
-			TimeSpan? ScaleInCooldown = null,
-			PoolSizeStrategy? SizeStrategy = null,
-			LeaseUtilizationSettings? LeaseUtilizationSettings = null,
-			JobQueueSettings? JobQueueSettings = null,
-			IEnumerable<KeyValuePair<string, string>>? Properties = null);
+			PoolId id,
+			string name,
+			Condition? condition = null,
+			bool? enableAutoscaling = null,
+			int? minAgents = null,
+			int? numReserveAgents = null,
+			TimeSpan? scaleOutCooldown = null,
+			TimeSpan? scaleInCooldown = null,
+			PoolSizeStrategy? sizeStrategy = null,
+			LeaseUtilizationSettings? leaseUtilizationSettings = null,
+			JobQueueSettings? jobQueueSettings = null,
+			IEnumerable<KeyValuePair<string, string>>? properties = null);
 
 		/// <summary>
 		/// Enumerates all the pools
@@ -63,9 +56,9 @@ namespace Horde.Build.Collections
 		/// <summary>
 		/// Gets a pool by ID
 		/// </summary>
-		/// <param name="Id">Unique id of the pool</param>
+		/// <param name="id">Unique id of the pool</param>
 		/// <returns>The pool document</returns>
-		Task<IPool?> GetAsync(PoolId Id);
+		Task<IPool?> GetAsync(PoolId id);
 
 		/// <summary>
 		/// Gets a list of all valid pool ids
@@ -76,46 +69,46 @@ namespace Horde.Build.Collections
 		/// <summary>
 		/// Gets a pool by ID
 		/// </summary>
-		/// <param name="Id">Unique id of the pool</param>
+		/// <param name="id">Unique id of the pool</param>
 		/// <returns>The pool document</returns>
-		Task<bool> DeleteAsync(PoolId Id);
+		Task<bool> DeleteAsync(PoolId id);
 
 		/// <summary>
 		/// Updates a pool
 		/// </summary>
-		/// <param name="Pool">The pool to update</param>
-		/// <param name="NewName">New name for the pool</param>
-		/// <param name="NewCondition">New condition for the pool</param>
-		/// <param name="NewEnableAutoscaling">New setting for whether to enable autoscaling</param>
-		/// <param name="NewMinAgents">Minimum number of agents in the pool</param>
-		/// <param name="NewNumReserveAgents">Minimum number of idle agents to maintain</param>
-		/// <param name="NewWorkspaces">New workspaces for the pool</param>
-		/// <param name="NewUseAutoSdk">New setting for whether to use autosdk</param>
-		/// <param name="NewProperties">New properties for the pool</param>
-		/// <param name="LastScaleUpTime">New time for last (auto) scale up</param>
-		/// <param name="LastScaleDownTime">New time for last (auto) scale down</param>
-		/// <param name="ScaleOutCooldown">Cooldown time between scale-out events</param>
-		/// <param name="ScaleInCooldown">Cooldown time between scale-in events</param>
-		/// <param name="SizeStrategy">Pool sizing strategy</param>
-		/// <param name="LeaseUtilizationSettings">Settings for lease utilization strategy</param>
-		/// <param name="JobQueueSettings">Settings for job queue strategy</param>
+		/// <param name="pool">The pool to update</param>
+		/// <param name="newName">New name for the pool</param>
+		/// <param name="newCondition">New condition for the pool</param>
+		/// <param name="newEnableAutoscaling">New setting for whether to enable autoscaling</param>
+		/// <param name="newMinAgents">Minimum number of agents in the pool</param>
+		/// <param name="newNumReserveAgents">Minimum number of idle agents to maintain</param>
+		/// <param name="newWorkspaces">New workspaces for the pool</param>
+		/// <param name="newUseAutoSdk">New setting for whether to use autosdk</param>
+		/// <param name="newProperties">New properties for the pool</param>
+		/// <param name="lastScaleUpTime">New time for last (auto) scale up</param>
+		/// <param name="lastScaleDownTime">New time for last (auto) scale down</param>
+		/// <param name="scaleOutCooldown">Cooldown time between scale-out events</param>
+		/// <param name="scaleInCooldown">Cooldown time between scale-in events</param>
+		/// <param name="sizeStrategy">Pool sizing strategy</param>
+		/// <param name="leaseUtilizationSettings">Settings for lease utilization strategy</param>
+		/// <param name="jobQueueSettings">Settings for job queue strategy</param>
 		/// <returns>Async task</returns>
 		Task<IPool?> TryUpdateAsync(
-			IPool Pool,
-			string? NewName = null,
-			Condition? NewCondition = null,
-			bool? NewEnableAutoscaling = null,
-			int? NewMinAgents = null,
-			int? NewNumReserveAgents = null,
-			List<AgentWorkspace>? NewWorkspaces = null,
-			bool? NewUseAutoSdk = null,
-			Dictionary<string, string?>? NewProperties = null,
-			DateTime? LastScaleUpTime = null,
-			DateTime? LastScaleDownTime = null,
-			TimeSpan? ScaleOutCooldown = null,
-			TimeSpan? ScaleInCooldown = null,
-			PoolSizeStrategy? SizeStrategy = null,
-			LeaseUtilizationSettings? LeaseUtilizationSettings = null,
-			JobQueueSettings? JobQueueSettings = null);
+			IPool pool,
+			string? newName = null,
+			Condition? newCondition = null,
+			bool? newEnableAutoscaling = null,
+			int? newMinAgents = null,
+			int? newNumReserveAgents = null,
+			List<AgentWorkspace>? newWorkspaces = null,
+			bool? newUseAutoSdk = null,
+			Dictionary<string, string?>? newProperties = null,
+			DateTime? lastScaleUpTime = null,
+			DateTime? lastScaleDownTime = null,
+			TimeSpan? scaleOutCooldown = null,
+			TimeSpan? scaleInCooldown = null,
+			PoolSizeStrategy? sizeStrategy = null,
+			LeaseUtilizationSettings? leaseUtilizationSettings = null,
+			JobQueueSettings? jobQueueSettings = null);
 	}
 }
diff --git a/Engine/Source/Programs/Horde/Horde.Build/Fleet/ITelemetryCollection.cs b/Engine/Source/Programs/Horde/Horde.Build/Fleet/ITelemetryCollection.cs
--- a/Engine/Source/Programs/Horde/Horde.Build/Fleet/ITelemetryCollection.cs
+++ b/Engine/Source/Programs/Horde/Horde.Build/Fleet/ITelemetryCollection.cs
@@ -1,17 +1,12 @@
 // Copyright Epic Games, Inc. All Rights Reserved.
 
-using Horde.Build.Models;
-using Horde.Build.Utilities;
 using System;
 using System.Collections.Generic;
-using System.Linq;
 using System.Threading.Tasks;
+using Horde.Build.Models;
 
 namespace Horde.Build.Collections
 {
-	using PoolId = StringId<IPool>;
-	using StreamId = StringId<IStream>;
-
 	/// <summary>
 	/// Collection of utilization collection
 	/// </summary>
@@ -20,17 +15,17 @@ namespace Horde.Build.Collections
 		/// <summary>
 		/// Adds entries for the given utilization
 		/// </summary>
-		/// <param name="Telemetry">Telemetry data to add</param>
+		/// <param name="telemetry">Telemetry data to add</param>
 		/// <returns>Async task</returns>
-		Task AddUtilizationTelemetryAsync(IUtilizationTelemetry Telemetry);
+		Task AddUtilizationTelemetryAsync(IUtilizationTelemetry telemetry);
 
 		/// <summary>
 		/// Finds utilization data matching the given criteria
 		/// </summary>
-		/// <param name="StartTimeUtc">Start time to query utilization for</param>
-		/// <param name="FinishTimeUtc">Finish time to query utilization for</param>
+		/// <param name="startTimeUtc">Start time to query utilization for</param>
+		/// <param name="finishTimeUtc">Finish time to query utilization for</param>
 		/// <returns>The utilization data</returns>
-		Task<List<IUtilizationTelemetry>> GetUtilizationTelemetryAsync(DateTime StartTimeUtc, DateTime FinishTimeUtc);
+		Task<List<IUtilizationTelemetry>> GetUtilizationTelemetryAsync(DateTime startTimeUtc, DateTime finishTimeUtc);
 
 		/// <summary>
 		/// Finds the latest utilization data
diff --git a/Engine/Source/Programs/Horde/Horde.Build/Fleet/Impl/Managers/Fleet/AwsFleetManager.cs b/Engine/Source/Programs/Horde/Horde.Build/Fleet/Impl/Managers/Fleet/AwsFleetManager.cs
--- a/Engine/Source/Programs/Horde/Horde.Build/Fleet/Impl/Managers/Fleet/AwsFleetManager.cs
+++ b/Engine/Source/Programs/Horde/Horde.Build/Fleet/Impl/Managers/Fleet/AwsFleetManager.cs
@@ -1,16 +1,16 @@
 // Copyright Epic Games, Inc. All Rights Reserved.
 
+using System;
+using System.Collections.Generic;
+using System.Linq;
+using System.Threading.Tasks;
 using Amazon;
 using Amazon.EC2;
 using Amazon.EC2.Model;
 using Horde.Build.Collections;
 using Horde.Build.Models;
-using Microsoft.Extensions.Logging;
-using System;
-using System.Collections.Generic;
-using System.Linq;
-using System.Threading.Tasks;
 using Horde.Build.Utilities;
+using Microsoft.Extensions.Logging;
 using OpenTracing;
 using OpenTracing.Util;
 
@@ -23,127 +23,126 @@ namespace Horde.Build.Services.Impl
 	{
 		const string AwsTagPropertyName = "aws-tag";
 		const string PoolTagName = "Horde_Autoscale_Pool";
-
-		AmazonEC2Client Client;
-		IAgentCollection AgentCollection;
-		ILogger Logger;
+		readonly AmazonEC2Client _client;
+		readonly IAgentCollection _agentCollection;
+		readonly ILogger _logger;
 
 		/// <summary>
 		/// Constructor
 		/// </summary>
-		public AwsFleetManager(IAgentCollection AgentCollection, ILogger<AwsFleetManager> Logger)
+		public AwsFleetManager(IAgentCollection agentCollection, ILogger<AwsFleetManager> logger)
 		{
-			this.AgentCollection = AgentCollection;
-			this.Logger = Logger;
+			_agentCollection = agentCollection;
+			_logger = logger;
 
-			AmazonEC2Config Config = new AmazonEC2Config();
-			Config.RegionEndpoint = RegionEndpoint.USEast1;
+			AmazonEC2Config config = new AmazonEC2Config();
+			config.RegionEndpoint = RegionEndpoint.USEast1;
 
-			Logger.LogInformation("Initializing AWS fleet manager for region {Region}", Config.RegionEndpoint);
+			logger.LogInformation("Initializing AWS fleet manager for region {Region}", config.RegionEndpoint);
 
-			Client = new AmazonEC2Client(Config);
+			_client = new AmazonEC2Client(config);
 		}
 
 		/// <inheritdoc/>
 		public void Dispose()
 		{
-			Client.Dispose();
+			_client.Dispose();
 		}
 
 		/// <inheritdoc/>
-		public async Task ExpandPoolAsync(IPool Pool, IReadOnlyList<IAgent> Agents, int Count)
+		public async Task ExpandPoolAsync(IPool pool, IReadOnlyList<IAgent> agents, int count)
 		{
-			using IScope Scope = GlobalTracer.Instance.BuildSpan("ExpandPool").StartActive();
-			Scope.Span.SetTag("poolName", Pool.Name);
-			Scope.Span.SetTag("numAgents", Agents.Count);
-			Scope.Span.SetTag("count", Count);
+			using IScope scope = GlobalTracer.Instance.BuildSpan("ExpandPool").StartActive();
+			scope.Span.SetTag("poolName", pool.Name);
+			scope.Span.SetTag("numAgents", agents.Count);
+			scope.Span.SetTag("count", count);
 
-			DescribeInstancesResponse DescribeResponse;
-			using (IScope DescribeScope = GlobalTracer.Instance.BuildSpan("DescribeInstances").StartActive())
+			DescribeInstancesResponse describeResponse;
+			using (IScope describeScope = GlobalTracer.Instance.BuildSpan("DescribeInstances").StartActive())
 			{
 				// Find stopped instances in the correct pool
-				DescribeInstancesRequest DescribeRequest = new DescribeInstancesRequest();
-				DescribeRequest.Filters = new List<Filter>();
-				DescribeRequest.Filters.Add(new Filter("instance-state-name", new List<string> { InstanceStateName.Stopped.Value }));
-				DescribeRequest.Filters.Add(new Filter("tag:" + PoolTagName, new List<string> { Pool.Name }));
-				DescribeResponse = await Client.DescribeInstancesAsync(DescribeRequest);
-				DescribeScope.Span.SetTag("res.statusCode", (int)DescribeResponse.HttpStatusCode);
-				DescribeScope.Span.SetTag("res.numReservations", DescribeResponse.Reservations.Count);
+				DescribeInstancesRequest describeRequest = new DescribeInstancesRequest();
+				describeRequest.Filters = new List<Filter>();
+				describeRequest.Filters.Add(new Filter("instance-state-name", new List<string> { InstanceStateName.Stopped.Value }));
+				describeRequest.Filters.Add(new Filter("tag:" + PoolTagName, new List<string> { pool.Name }));
+				describeResponse = await _client.DescribeInstancesAsync(describeRequest);
+				describeScope.Span.SetTag("res.statusCode", (int)describeResponse.HttpStatusCode);
+				describeScope.Span.SetTag("res.numReservations", describeResponse.Reservations.Count);
 			}
 
-			using (IScope StartScope = GlobalTracer.Instance.BuildSpan("StartInstances").StartActive())
+			using (IScope startScope = GlobalTracer.Instance.BuildSpan("StartInstances").StartActive())
 			{
 				// Try to start the given instances
-				StartInstancesRequest StartRequest = new StartInstancesRequest();
-				StartRequest.InstanceIds.AddRange(DescribeResponse.Reservations.SelectMany(x => x.Instances).Select(x => x.InstanceId).Take(Count));
+				StartInstancesRequest startRequest = new StartInstancesRequest();
+				startRequest.InstanceIds.AddRange(describeResponse.Reservations.SelectMany(x => x.Instances).Select(x => x.InstanceId).Take(count));
 				
-				StartScope.Span.SetTag("req.instanceIds", string.Join(",", StartRequest.InstanceIds));
-				if (StartRequest.InstanceIds.Count > 0)
+				startScope.Span.SetTag("req.instanceIds", String.Join(",", startRequest.InstanceIds));
+				if (startRequest.InstanceIds.Count > 0)
 				{
-					StartInstancesResponse StartResponse = await Client.StartInstancesAsync(StartRequest);
-					StartScope.Span.SetTag("res.statusCode", (int)StartResponse.HttpStatusCode);
-					StartScope.Span.SetTag("res.numInstances", StartResponse.StartingInstances.Count);
-					if ((int)StartResponse.HttpStatusCode >= 200 && (int)StartResponse.HttpStatusCode <= 299)
+					StartInstancesResponse startResponse = await _client.StartInstancesAsync(startRequest);
+					startScope.Span.SetTag("res.statusCode", (int)startResponse.HttpStatusCode);
+					startScope.Span.SetTag("res.numInstances", startResponse.StartingInstances.Count);
+					if ((int)startResponse.HttpStatusCode >= 200 && (int)startResponse.HttpStatusCode <= 299)
 					{
-						foreach (InstanceStateChange InstanceChange in StartResponse.StartingInstances)
+						foreach (InstanceStateChange instanceChange in startResponse.StartingInstances)
 						{
-							Logger.LogInformation("Starting instance {InstanceId} for pool {PoolId} (prev state {PrevState}, current state {CurrentState}", InstanceChange.InstanceId, Pool.Id, InstanceChange.PreviousState, InstanceChange.CurrentState);
+							_logger.LogInformation("Starting instance {InstanceId} for pool {PoolId} (prev state {PrevState}, current state {CurrentState}", instanceChange.InstanceId, pool.Id, instanceChange.PreviousState, instanceChange.CurrentState);
 						}
 					}
 				}
 
-				if (StartRequest.InstanceIds.Count < Count)
+				if (startRequest.InstanceIds.Count < count)
 				{
-					Logger.LogInformation("Unable to expand pool with the requested number of instances. " +
-					                      "Num requested instances to add {RequestedCount}. Actual instances started {InstancesStarted}", Count, StartRequest.InstanceIds.Count);
+					_logger.LogInformation("Unable to expand pool with the requested number of instances. " +
+					                      "Num requested instances to add {RequestedCount}. Actual instances started {InstancesStarted}", count, startRequest.InstanceIds.Count);
 				}
 			}
 		}
 
 		/// <inheritdoc/>
-		public async Task ShrinkPoolAsync(IPool Pool, IReadOnlyList<IAgent> Agents, int Count)
+		public async Task ShrinkPoolAsync(IPool pool, IReadOnlyList<IAgent> agents, int count)
 		{
-			using IScope Scope = GlobalTracer.Instance.BuildSpan("ShrinkPool").StartActive();
-			Scope.Span.SetTag("poolName", Pool.Name);
-			Scope.Span.SetTag("count", Count);
+			using IScope scope = GlobalTracer.Instance.BuildSpan("ShrinkPool").StartActive();
+			scope.Span.SetTag("poolName", pool.Name);
+			scope.Span.SetTag("count", count);
 			
-			string AwsTagProperty = $"{AwsTagPropertyName}={PoolTagName}:{Pool.Name}";
+			string awsTagProperty = $"{AwsTagPropertyName}={PoolTagName}:{pool.Name}";
 			
 			// Sort the agents by number of active leases. It's better to shutdown agents currently doing nothing.
-			List<IAgent> FilteredAgents = Agents.OrderBy(x => x.Leases.Count).ToList();
-			List<IAgent> AgentsWithAwsTags = FilteredAgents.Where(x => x.HasProperty(AwsTagProperty)).ToList(); 
-			List<IAgent> AgentsLimitedByCount = AgentsWithAwsTags.Take(Count).ToList();
+			List<IAgent> filteredAgents = agents.OrderBy(x => x.Leases.Count).ToList();
+			List<IAgent> agentsWithAwsTags = filteredAgents.Where(x => x.HasProperty(awsTagProperty)).ToList(); 
+			List<IAgent> agentsLimitedByCount = agentsWithAwsTags.Take(count).ToList();
 			
-			Scope.Span.SetTag("agents.num", Agents.Count);
-			Scope.Span.SetTag("agents.filtered.num", FilteredAgents.Count);
-			Scope.Span.SetTag("agents.withAwsTags.num", AgentsWithAwsTags.Count);
-			Scope.Span.SetTag("agents.limitedByCount.num", AgentsLimitedByCount.Count);
+			scope.Span.SetTag("agents.num", agents.Count);
+			scope.Span.SetTag("agents.filtered.num", filteredAgents.Count);
+			scope.Span.SetTag("agents.withAwsTags.num", agentsWithAwsTags.Count);
+			scope.Span.SetTag("agents.limitedByCount.num", agentsLimitedByCount.Count);
 
-			foreach (IAgent Agent in AgentsLimitedByCount)
+			foreach (IAgent agent in agentsLimitedByCount)
 			{
-				IAuditLogChannel<AgentId> AgentLogger = AgentCollection.GetLogger(Agent.Id);
-				if (await AgentCollection.TryUpdateSettingsAsync(Agent, bRequestShutdown: true, ShutdownReason: "Autoscaler") != null)
+				IAuditLogChannel<AgentId> agentLogger = _agentCollection.GetLogger(agent.Id);
+				if (await _agentCollection.TryUpdateSettingsAsync(agent, bRequestShutdown: true, shutdownReason: "Autoscaler") != null)
 				{
-					AgentLogger.LogInformation("Marked for shutdown due to autoscaling (currently {NumLeases} leases outstanding)", Agent.Leases.Count);
+					agentLogger.LogInformation("Marked for shutdown due to autoscaling (currently {NumLeases} leases outstanding)", agent.Leases.Count);
 				}
 				else
 				{
-					AgentLogger.LogError("Unable to mark agent for shutdown due to autoscaling");
+					agentLogger.LogError("Unable to mark agent for shutdown due to autoscaling");
 				}
 			}
 		}
 
 		/// <inheritdoc/>
-		public async Task<int> GetNumStoppedInstancesAsync(IPool Pool)
+		public async Task<int> GetNumStoppedInstancesAsync(IPool pool)
 		{
 			// Find all instances in the pool
-			DescribeInstancesRequest DescribeRequest = new DescribeInstancesRequest();
-			DescribeRequest.Filters = new List<Filter>();
-			DescribeRequest.Filters.Add(new Filter("instance-state-name", new List<string> { InstanceStateName.Stopped.Value }));
-			DescribeRequest.Filters.Add(new Filter("tag:" + PoolTagName, new List<string> { Pool.Name }));
+			DescribeInstancesRequest describeRequest = new DescribeInstancesRequest();
+			describeRequest.Filters = new List<Filter>();
+			describeRequest.Filters.Add(new Filter("instance-state-name", new List<string> { InstanceStateName.Stopped.Value }));
+			describeRequest.Filters.Add(new Filter("tag:" + PoolTagName, new List<string> { pool.Name }));
 
-			DescribeInstancesResponse DescribeResponse = await Client.DescribeInstancesAsync(DescribeRequest);
-			return DescribeResponse.Reservations.SelectMany(x => x.Instances).Select(x => x.InstanceId).Distinct().Count();
+			DescribeInstancesResponse describeResponse = await _client.DescribeInstancesAsync(describeRequest);
+			return describeResponse.Reservations.SelectMany(x => x.Instances).Select(x => x.InstanceId).Distinct().Count();
 		}
 	}
 }
diff --git a/Engine/Source/Programs/Horde/Horde.Build/Fleet/Impl/Managers/Fleet/DefaultFleetManager.cs b/Engine/Source/Programs/Horde/Horde.Build/Fleet/Impl/Managers/Fleet/DefaultFleetManager.cs
--- a/Engine/Source/Programs/Horde/Horde.Build/Fleet/Impl/Managers/Fleet/DefaultFleetManager.cs
+++ b/Engine/Source/Programs/Horde/Horde.Build/Fleet/Impl/Managers/Fleet/DefaultFleetManager.cs
@@ -1,11 +1,9 @@
 // Copyright Epic Games, Inc. All Rights Reserved.
 
-using Horde.Build.Models;
-using Microsoft.Extensions.Logging;
-using System;
 using System.Collections.Generic;
-using System.Linq;
 using System.Threading.Tasks;
+using Horde.Build.Models;
+using Microsoft.Extensions.Logging;
 
 namespace Horde.Build.Services.Impl
 {
@@ -14,32 +12,32 @@ namespace Horde.Build.Services.Impl
 	/// </summary>
 	public class DefaultFleetManager : IFleetManager
 	{
-		ILogger Logger;
+		readonly ILogger _logger;
 
 		/// <summary>
 		/// Constructor
 		/// </summary>
-		/// <param name="Logger">Logging device</param>
-		public DefaultFleetManager(ILogger<DefaultFleetManager> Logger)
+		/// <param name="logger">Logging device</param>
+		public DefaultFleetManager(ILogger<DefaultFleetManager> logger)
 		{
-			this.Logger = Logger;
+			_logger = logger;
 		}
 
 		/// <inheritdoc/>
-		public Task ExpandPoolAsync(IPool Pool, IReadOnlyList<IAgent> Agents, int Count)
+		public Task ExpandPoolAsync(IPool pool, IReadOnlyList<IAgent> agents, int count)
 		{
-			Logger.LogInformation("Expand pool {PoolId} by {Count} agents", Pool.Id, Count);
+			_logger.LogInformation("Expand pool {PoolId} by {Count} agents", pool.Id, count);
 			return Task.CompletedTask;
 		}
 
 		/// <inheritdoc/>
-		public Task ShrinkPoolAsync(IPool Pool, IReadOnlyList<IAgent> Agents, int Count)
+		public Task ShrinkPoolAsync(IPool pool, IReadOnlyList<IAgent> agents, int count)
 		{
-			Logger.LogInformation("Shrink pool {PoolId} by {Count} agents", Pool.Id, Count);
+			_logger.LogInformation("Shrink pool {PoolId} by {Count} agents", pool.Id, count);
 			return Task.CompletedTask;
 		}
 
 		/// <inheritdoc/>
-		public Task<int> GetNumStoppedInstancesAsync(IPool Pool) => Task.FromResult(0);
+		public Task<int> GetNumStoppedInstancesAsync(IPool pool) => Task.FromResult(0);
 	}
 }
diff --git a/Engine/Source/Programs/Horde/Horde.Build/Fleet/Impl/PoolCollection.cs b/Engine/Source/Programs/Horde/Horde.Build/Fleet/Impl/PoolCollection.cs
--- a/Engine/Source/Programs/Horde/Horde.Build/Fleet/Impl/PoolCollection.cs
+++ b/Engine/Source/Programs/Horde/Horde.Build/Fleet/Impl/PoolCollection.cs
@@ -1,19 +1,17 @@
 // Copyright Epic Games, Inc. All Rights Reserved.
 
+using System;
+using System.Collections.Generic;
+using System.Linq;
+using System.Threading.Tasks;
 using EpicGames.Horde.Common;
+using Horde.Build.Fleet.Autoscale;
 using Horde.Build.Models;
 using Horde.Build.Services;
 using Horde.Build.Utilities;
 using MongoDB.Bson;
-using MongoDB.Bson.Serialization;
 using MongoDB.Bson.Serialization.Attributes;
 using MongoDB.Driver;
-using System;
-using System.Collections.Generic;
-using System.Linq;
-using System.Linq.Expressions;
-using System.Threading.Tasks;
-using Horde.Build.Fleet.Autoscale;
 
 namespace Horde.Build.Collections.Impl
 {
@@ -56,259 +54,259 @@ namespace Horde.Build.Collections.Impl
 			{
 			}
 
-			public PoolDocument(IPool Other)
+			public PoolDocument(IPool other)
 			{
-				Id = Other.Id;
-				Name = Other.Name;
-				Condition = Other.Condition;
-				Workspaces.AddRange(Other.Workspaces);
-				UseAutoSdk = Other.UseAutoSdk;
-				Properties = new Dictionary<string, string>(Other.Properties);
-				EnableAutoscaling = Other.EnableAutoscaling;
-				MinAgents = Other.MinAgents;
-				NumReserveAgents = Other.NumReserveAgents;
-				LastScaleUpTime = Other.LastScaleUpTime;
-				LastScaleDownTime = Other.LastScaleDownTime;
-				ScaleOutCooldown = Other.ScaleOutCooldown;
-				ScaleInCooldown = Other.ScaleInCooldown;
-				SizeStrategy = Other.SizeStrategy;
-				LeaseUtilizationSettings = Other.LeaseUtilizationSettings;
-				JobQueueSettings = Other.JobQueueSettings;
-				UpdateIndex = Other.UpdateIndex;
+				Id = other.Id;
+				Name = other.Name;
+				Condition = other.Condition;
+				Workspaces.AddRange(other.Workspaces);
+				UseAutoSdk = other.UseAutoSdk;
+				Properties = new Dictionary<string, string>(other.Properties);
+				EnableAutoscaling = other.EnableAutoscaling;
+				MinAgents = other.MinAgents;
+				NumReserveAgents = other.NumReserveAgents;
+				LastScaleUpTime = other.LastScaleUpTime;
+				LastScaleDownTime = other.LastScaleDownTime;
+				ScaleOutCooldown = other.ScaleOutCooldown;
+				ScaleInCooldown = other.ScaleInCooldown;
+				SizeStrategy = other.SizeStrategy;
+				LeaseUtilizationSettings = other.LeaseUtilizationSettings;
+				JobQueueSettings = other.JobQueueSettings;
+				UpdateIndex = other.UpdateIndex;
 			}
 		}
 
 		/// <summary>
 		/// Collection of pool documents
 		/// </summary>
-		IMongoCollection<PoolDocument> Pools;
+		readonly IMongoCollection<PoolDocument> _pools;
 
 		/// <summary>
 		/// Constructor
 		/// </summary>
-		/// <param name="DatabaseService">The database service instance</param>
-		public PoolCollection(DatabaseService DatabaseService)
+		/// <param name="databaseService">The database service instance</param>
+		public PoolCollection(DatabaseService databaseService)
 		{
-			Pools = DatabaseService.GetCollection<PoolDocument>("Pools");
+			_pools = databaseService.GetCollection<PoolDocument>("Pools");
 		}
 
 		/// <inheritdoc/>
 		public async Task<IPool> AddAsync(
-			PoolId Id,
-			string Name,
-			Condition? Condition,
-			bool? EnableAutoscaling,
-			int? MinAgents,
-			int? NumReserveAgents,
-			TimeSpan? ScaleOutCooldown,
-			TimeSpan? ScaleInCooldown,
-			PoolSizeStrategy? SizeStrategy,
-			LeaseUtilizationSettings? LeaseUtilizationSettings,
-			JobQueueSettings? JobQueueSettings,
-			IEnumerable<KeyValuePair<string, string>>? Properties)
+			PoolId id,
+			string name,
+			Condition? condition,
+			bool? enableAutoscaling,
+			int? minAgents,
+			int? numReserveAgents,
+			TimeSpan? scaleOutCooldown,
+			TimeSpan? scaleInCooldown,
+			PoolSizeStrategy? sizeStrategy,
+			LeaseUtilizationSettings? leaseUtilizationSettings,
+			JobQueueSettings? jobQueueSettings,
+			IEnumerable<KeyValuePair<string, string>>? properties)
 		{
-			PoolDocument Pool = new PoolDocument();
-			Pool.Id = Id;
-			Pool.Name = Name;
-			Pool.Condition = Condition;
-			if (EnableAutoscaling != null)
+			PoolDocument pool = new PoolDocument();
+			pool.Id = id;
+			pool.Name = name;
+			pool.Condition = condition;
+			if (enableAutoscaling != null)
 			{
-				Pool.EnableAutoscaling = EnableAutoscaling.Value;
+				pool.EnableAutoscaling = enableAutoscaling.Value;
 			}
-			Pool.MinAgents = MinAgents;
-			Pool.NumReserveAgents = NumReserveAgents;
-			if (Properties != null)
+			pool.MinAgents = minAgents;
+			pool.NumReserveAgents = numReserveAgents;
+			if (properties != null)
 			{
-				Pool.Properties = new Dictionary<string, string>(Properties);
+				pool.Properties = new Dictionary<string, string>(properties);
 			}
 
-			Pool.ScaleOutCooldown = ScaleOutCooldown;
-			Pool.ScaleInCooldown = ScaleInCooldown;
-			Pool.SizeStrategy = SizeStrategy;
-			Pool.LeaseUtilizationSettings = LeaseUtilizationSettings;
-			Pool.JobQueueSettings = JobQueueSettings;
-			await Pools.InsertOneAsync(Pool);
-			return Pool;
+			pool.ScaleOutCooldown = scaleOutCooldown;
+			pool.ScaleInCooldown = scaleInCooldown;
+			pool.SizeStrategy = sizeStrategy;
+			pool.LeaseUtilizationSettings = leaseUtilizationSettings;
+			pool.JobQueueSettings = jobQueueSettings;
+			await _pools.InsertOneAsync(pool);
+			return pool;
 		}
 
 		/// <inheritdoc/>
 		public async Task<List<IPool>> GetAsync()
 		{
-			List<PoolDocument> Results = await Pools.Find(FilterDefinition<PoolDocument>.Empty).ToListAsync();
-			return Results.ConvertAll<IPool>(x => x);
+			List<PoolDocument> results = await _pools.Find(FilterDefinition<PoolDocument>.Empty).ToListAsync();
+			return results.ConvertAll<IPool>(x => x);
 		}
 
 		/// <inheritdoc/>
-		public async Task<IPool?> GetAsync(PoolId Id)
+		public async Task<IPool?> GetAsync(PoolId id)
 		{
-			return await Pools.Find(x => x.Id == Id).FirstOrDefaultAsync();
+			return await _pools.Find(x => x.Id == id).FirstOrDefaultAsync();
 		}
 
 		/// <inheritdoc/>
 		public async Task<List<PoolId>> GetPoolIdsAsync()
 		{
-			ProjectionDefinition<PoolDocument, BsonDocument> Projection = Builders<PoolDocument>.Projection.Include(x => x.Id);
-			List<BsonDocument> Results = await Pools.Find(FilterDefinition<PoolDocument>.Empty).Project(Projection).ToListAsync();
-			return Results.ConvertAll(x => new PoolId(x.GetElement("_id").Value.AsString));
+			ProjectionDefinition<PoolDocument, BsonDocument> projection = Builders<PoolDocument>.Projection.Include(x => x.Id);
+			List<BsonDocument> results = await _pools.Find(FilterDefinition<PoolDocument>.Empty).Project(projection).ToListAsync();
+			return results.ConvertAll(x => new PoolId(x.GetElement("_id").Value.AsString));
 		}
 
 		/// <inheritdoc/>
-		public async Task<bool> DeleteAsync(PoolId Id)
+		public async Task<bool> DeleteAsync(PoolId id)
 		{
-			FilterDefinition<PoolDocument> Filter = Builders<PoolDocument>.Filter.Eq(x => x.Id, Id);
-			DeleteResult Result = await Pools.DeleteOneAsync(Filter);
-			return Result.DeletedCount > 0;
+			FilterDefinition<PoolDocument> filter = Builders<PoolDocument>.Filter.Eq(x => x.Id, id);
+			DeleteResult result = await _pools.DeleteOneAsync(filter);
+			return result.DeletedCount > 0;
 		}
 
 		/// <summary>
 		/// Attempts to update a pool, upgrading it to the latest document schema if necessary
 		/// </summary>
-		/// <param name="PoolToUpdate">Interface for the document to update</param>
-		/// <param name="Transaction">The transaction to execute</param>
+		/// <param name="poolToUpdate">Interface for the document to update</param>
+		/// <param name="transaction">The transaction to execute</param>
 		/// <returns>New pool definition, or null on failure.</returns>
-		async Task<IPool?> TryUpdateAsync(IPool PoolToUpdate, TransactionBuilder<PoolDocument> Transaction)
+		async Task<IPool?> TryUpdateAsync(IPool poolToUpdate, TransactionBuilder<PoolDocument> transaction)
 		{
-			if (Transaction.IsEmpty)
+			if (transaction.IsEmpty)
 			{
-				return PoolToUpdate;
+				return poolToUpdate;
 			}
 
-			Transaction.Set(x => x.UpdateIndex, PoolToUpdate.UpdateIndex + 1);
+			transaction.Set(x => x.UpdateIndex, poolToUpdate.UpdateIndex + 1);
 
-			PoolDocument? MutablePool = PoolToUpdate as PoolDocument;
-			if (MutablePool != null)
+			PoolDocument? mutablePool = poolToUpdate as PoolDocument;
+			if (mutablePool != null)
 			{
-				UpdateResult Result = await Pools.UpdateOneAsync(x => x.Id == PoolToUpdate.Id && x.UpdateIndex == PoolToUpdate.UpdateIndex, Transaction.ToUpdateDefinition());
-				if (Result.ModifiedCount == 0)
+				UpdateResult result = await _pools.UpdateOneAsync(x => x.Id == poolToUpdate.Id && x.UpdateIndex == poolToUpdate.UpdateIndex, transaction.ToUpdateDefinition());
+				if (result.ModifiedCount == 0)
 				{
 					return null;
 				}
 
-				Transaction.ApplyTo(MutablePool);
+				transaction.ApplyTo(mutablePool);
 			}
 			else
 			{
-				MutablePool = new PoolDocument(PoolToUpdate);
-				Transaction.ApplyTo(MutablePool);
+				mutablePool = new PoolDocument(poolToUpdate);
+				transaction.ApplyTo(mutablePool);
 
-				ReplaceOneResult Result = await Pools.ReplaceOneAsync<PoolDocument>(x => x.Id == PoolToUpdate.Id && x.UpdateIndex == PoolToUpdate.UpdateIndex, MutablePool);
-				if (Result.ModifiedCount == 0)
+				ReplaceOneResult result = await _pools.ReplaceOneAsync<PoolDocument>(x => x.Id == poolToUpdate.Id && x.UpdateIndex == poolToUpdate.UpdateIndex, mutablePool);
+				if (result.ModifiedCount == 0)
 				{
 					return null;
 				}
 			}
-			return MutablePool;
+			return mutablePool;
 		}
 
 		/// <inheritdoc/>
 		public Task<IPool?> TryUpdateAsync(
-			IPool Pool,
-			string? NewName,
-			Condition? NewCondition,
-			bool? NewEnableAutoscaling,
-			int? NewMinAgents,
-			int? NewNumReserveAgents,
-			List<AgentWorkspace>? NewWorkspaces,
-			bool? NewUseAutoSdk,
-			Dictionary<string, string?>? NewProperties,
-			DateTime? LastScaleUpTime,
-			DateTime? LastScaleDownTime,
-			TimeSpan? ScaleOutCooldown,
-			TimeSpan? ScaleInCooldown, 
-			PoolSizeStrategy? SizeStrategy,
-			LeaseUtilizationSettings? LeaseUtilizationSettings,
-			JobQueueSettings? JobQueueSettings)
+			IPool pool,
+			string? newName,
+			Condition? newCondition,
+			bool? newEnableAutoscaling,
+			int? newMinAgents,
+			int? newNumReserveAgents,
+			List<AgentWorkspace>? newWorkspaces,
+			bool? newUseAutoSdk,
+			Dictionary<string, string?>? newProperties,
+			DateTime? lastScaleUpTime,
+			DateTime? lastScaleDownTime,
+			TimeSpan? scaleOutCooldown,
+			TimeSpan? scaleInCooldown, 
+			PoolSizeStrategy? sizeStrategy,
+			LeaseUtilizationSettings? leaseUtilizationSettings,
+			JobQueueSettings? jobQueueSettings)
 		{
-			TransactionBuilder<PoolDocument> Transaction = new TransactionBuilder<PoolDocument>();
-			if (NewName != null)
+			TransactionBuilder<PoolDocument> transaction = new TransactionBuilder<PoolDocument>();
+			if (newName != null)
 			{
-				Transaction.Set(x => x.Name, NewName);
+				transaction.Set(x => x.Name, newName);
 			}
-			if (NewCondition != null)
+			if (newCondition != null)
 			{
-				if (NewCondition.IsEmpty())
+				if (newCondition.IsEmpty())
 				{
-					Transaction.Unset(x => x.Condition!);
+					transaction.Unset(x => x.Condition!);
 				}
 				else
 				{
-					Transaction.Set(x => x.Condition, NewCondition);
+					transaction.Set(x => x.Condition, newCondition);
 				}
 			}
-			if (NewEnableAutoscaling != null)
+			if (newEnableAutoscaling != null)
 			{
-				if (NewEnableAutoscaling.Value)
+				if (newEnableAutoscaling.Value)
 				{
-					Transaction.Unset(x => x.EnableAutoscaling);
+					transaction.Unset(x => x.EnableAutoscaling);
 				}
 				else
 				{
-					Transaction.Set(x => x.EnableAutoscaling, NewEnableAutoscaling.Value);
+					transaction.Set(x => x.EnableAutoscaling, newEnableAutoscaling.Value);
 				}
 			}
-			if (NewMinAgents != null)
+			if (newMinAgents != null)
 			{
-				if (NewMinAgents.Value < 0)
+				if (newMinAgents.Value < 0)
 				{
-					Transaction.Unset(x => x.MinAgents!);
+					transaction.Unset(x => x.MinAgents!);
 				}
 				else
 				{
-					Transaction.Set(x => x.MinAgents, NewMinAgents.Value);
+					transaction.Set(x => x.MinAgents, newMinAgents.Value);
 				}
 			}
-			if (NewNumReserveAgents != null)
+			if (newNumReserveAgents != null)
 			{
-				if (NewNumReserveAgents.Value < 0)
+				if (newNumReserveAgents.Value < 0)
 				{
-					Transaction.Unset(x => x.NumReserveAgents!);
+					transaction.Unset(x => x.NumReserveAgents!);
 				}
 				else
 				{
-					Transaction.Set(x => x.NumReserveAgents, NewNumReserveAgents.Value);
+					transaction.Set(x => x.NumReserveAgents, newNumReserveAgents.Value);
 				}
 			}
-			if (NewWorkspaces != null)
+			if (newWorkspaces != null)
 			{
-				Transaction.Set(x => x.Workspaces, NewWorkspaces);
+				transaction.Set(x => x.Workspaces, newWorkspaces);
 			}
-			if (NewUseAutoSdk != null)
+			if (newUseAutoSdk != null)
 			{
-				Transaction.Set(x => x.UseAutoSdk, NewUseAutoSdk.Value);
+				transaction.Set(x => x.UseAutoSdk, newUseAutoSdk.Value);
 			}
-			if (NewProperties != null)
+			if (newProperties != null)
 			{
-				Transaction.UpdateDictionary(x => x.Properties, NewProperties);
+				transaction.UpdateDictionary(x => x.Properties, newProperties);
 			}
-			if (LastScaleUpTime != null)
+			if (lastScaleUpTime != null)
 			{
-				Transaction.Set(x => x.LastScaleUpTime, LastScaleUpTime);
+				transaction.Set(x => x.LastScaleUpTime, lastScaleUpTime);
 			}
-			if (LastScaleDownTime != null)
+			if (lastScaleDownTime != null)
 			{
-				Transaction.Set(x => x.LastScaleDownTime, LastScaleDownTime);
+				transaction.Set(x => x.LastScaleDownTime, lastScaleDownTime);
 			}
-			if (ScaleOutCooldown != null)
+			if (scaleOutCooldown != null)
 			{
-				Transaction.Set(x => x.ScaleOutCooldown, ScaleOutCooldown);
+				transaction.Set(x => x.ScaleOutCooldown, scaleOutCooldown);
 			}
-			if (ScaleInCooldown != null)
+			if (scaleInCooldown != null)
 			{
-				Transaction.Set(x => x.ScaleInCooldown, ScaleInCooldown);
+				transaction.Set(x => x.ScaleInCooldown, scaleInCooldown);
 			}
-			if (SizeStrategy != null)
+			if (sizeStrategy != null)
 			{
-				Transaction.Set(x => x.SizeStrategy, SizeStrategy);
+				transaction.Set(x => x.SizeStrategy, sizeStrategy);
 			}
-			if (LeaseUtilizationSettings != null)
+			if (leaseUtilizationSettings != null)
 			{
-				Transaction.Set(x => x.LeaseUtilizationSettings, LeaseUtilizationSettings);
+				transaction.Set(x => x.LeaseUtilizationSettings, leaseUtilizationSettings);
 			}
-			if (JobQueueSettings != null)
+			if (jobQueueSettings != null)
 			{
-				Transaction.Set(x => x.JobQueueSettings, JobQueueSettings);
+				transaction.Set(x => x.JobQueueSettings, jobQueueSettings);
 			}
-			return TryUpdateAsync(Pool, Transaction);
+			return TryUpdateAsync(pool, transaction);
 		}
 	}
 }
diff --git a/Engine/Source/Programs/Horde/Horde.Build/Fleet/Impl/TelemetryCollection.cs b/Engine/Source/Programs/Horde/Horde.Build/Fleet/Impl/TelemetryCollection.cs
--- a/Engine/Source/Programs/Horde/Horde.Build/Fleet/Impl/TelemetryCollection.cs
+++ b/Engine/Source/Programs/Horde/Horde.Build/Fleet/Impl/TelemetryCollection.cs
@@ -1,18 +1,15 @@
 // Copyright Epic Games, Inc. All Rights Reserved.
 
+using System;
+using System.Collections.Generic;
+using System.Linq;
+using System.Threading.Tasks;
 using EpicGames.Core;
 using Horde.Build.Models;
 using Horde.Build.Services;
 using Horde.Build.Utilities;
-using Microsoft.Extensions.Hosting;
-using MongoDB.Bson;
 using MongoDB.Bson.Serialization.Attributes;
 using MongoDB.Driver;
-using System;
-using System.Collections.Generic;
-using System.Linq;
-using System.Threading;
-using System.Threading.Tasks;
 
 namespace Horde.Build.Collections.Impl
 {
@@ -43,14 +40,14 @@ namespace Horde.Build.Collections.Impl
 				Pools = new List<PoolUtilizationDocument>();
 			}
 
-			public UtilizationDocument(IUtilizationTelemetry Other)
+			public UtilizationDocument(IUtilizationTelemetry other)
 			{
-				this.StartTime = Other.StartTime;
-				this.FinishTime = Other.FinishTime;
-				this.NumAgents = Other.NumAgents;
-				this.Pools = Other.Pools.ConvertAll(x => new PoolUtilizationDocument(x));
-				this.HibernatingTime = Other.HibernatingTime;
-				this.AdminTime = Other.AdminTime;
+				StartTime = other.StartTime;
+				FinishTime = other.FinishTime;
+				NumAgents = other.NumAgents;
+				Pools = other.Pools.ConvertAll(x => new PoolUtilizationDocument(x));
+				HibernatingTime = other.HibernatingTime;
+				AdminTime = other.AdminTime;
 			}
 		}
 
@@ -72,14 +69,14 @@ namespace Horde.Build.Collections.Impl
 				Streams = new List<StreamUtilizationDocument>();
 			}
 
-			public PoolUtilizationDocument(IPoolUtilizationTelemetry Other)
+			public PoolUtilizationDocument(IPoolUtilizationTelemetry other)
 			{
-				this.PoolId = Other.PoolId;
-				this.NumAgents = Other.NumAgents;
-				this.Streams = Other.Streams.ConvertAll(x => new StreamUtilizationDocument(x));
-				this.AdminTime = Other.AdminTime;
-				this.HibernatingTime = Other.HibernatingTime;
-				this.OtherTime = Other.OtherTime;
+				PoolId = other.PoolId;
+				NumAgents = other.NumAgents;
+				Streams = other.Streams.ConvertAll(x => new StreamUtilizationDocument(x));
+				AdminTime = other.AdminTime;
+				HibernatingTime = other.HibernatingTime;
+				OtherTime = other.OtherTime;
 			}
 		}
 
@@ -93,48 +90,48 @@ namespace Horde.Build.Collections.Impl
 			{
 			}
 
-			public StreamUtilizationDocument(IStreamUtilizationTelemetry Other)
+			public StreamUtilizationDocument(IStreamUtilizationTelemetry other)
 			{
-				this.StreamId = Other.StreamId;
-				this.Time = Other.Time;
+				StreamId = other.StreamId;
+				Time = other.Time;
 			}
 		}
 
-		IMongoCollection<UtilizationDocument> Utilization;
+		readonly IMongoCollection<UtilizationDocument> _utilization;
 
 		/// <summary>
 		/// Constructor
 		/// </summary>
-		/// <param name="Database"></param>
-		public TelemetryCollection(DatabaseService Database)
+		/// <param name="database"></param>
+		public TelemetryCollection(DatabaseService database)
 		{
-			Utilization = Database.GetCollection<UtilizationDocument>("Utilization");
+			_utilization = database.GetCollection<UtilizationDocument>("Utilization");
 
-			if (!Database.ReadOnlyMode)
+			if (!database.ReadOnlyMode)
 			{
-				Utilization.Indexes.CreateOne(new CreateIndexModel<UtilizationDocument>(Builders<UtilizationDocument>.IndexKeys.Ascending(x => x.FinishTime).Ascending(x => x.StartTime)));
+				_utilization.Indexes.CreateOne(new CreateIndexModel<UtilizationDocument>(Builders<UtilizationDocument>.IndexKeys.Ascending(x => x.FinishTime).Ascending(x => x.StartTime)));
 			}
 		}
 
 		/// <inheritdoc/>
-		public async Task AddUtilizationTelemetryAsync(IUtilizationTelemetry NewTelemetry)
+		public async Task AddUtilizationTelemetryAsync(IUtilizationTelemetry newTelemetry)
 		{
-			await Utilization.InsertOneAsync(new UtilizationDocument(NewTelemetry));
+			await _utilization.InsertOneAsync(new UtilizationDocument(newTelemetry));
 		}
 
 		/// <inheritdoc/>
-		public async Task<List<IUtilizationTelemetry>> GetUtilizationTelemetryAsync(DateTime StartTimeUtc, DateTime FinishTimeUtc)
+		public async Task<List<IUtilizationTelemetry>> GetUtilizationTelemetryAsync(DateTime startTimeUtc, DateTime finishTimeUtc)
 		{
-			FilterDefinition<UtilizationDocument> Filter = Builders<UtilizationDocument>.Filter.Gte(x => x.FinishTime, StartTimeUtc) & Builders<UtilizationDocument>.Filter.Lte(x => x.StartTime, FinishTimeUtc);
+			FilterDefinition<UtilizationDocument> filter = Builders<UtilizationDocument>.Filter.Gte(x => x.FinishTime, startTimeUtc) & Builders<UtilizationDocument>.Filter.Lte(x => x.StartTime, finishTimeUtc);
 
-			List<UtilizationDocument> Documents = await Utilization.Find(Filter).ToListAsync();
-			return Documents.ConvertAll<IUtilizationTelemetry>(x => x);
+			List<UtilizationDocument> documents = await _utilization.Find(filter).ToListAsync();
+			return documents.ConvertAll<IUtilizationTelemetry>(x => x);
 		}
 
 		/// <inheritdoc/>
 		public async Task<IUtilizationTelemetry?> GetLatestUtilizationTelemetryAsync()
 		{
-			return await Utilization.Find(FilterDefinition<UtilizationDocument>.Empty).SortByDescending(x => x.FinishTime).FirstOrDefaultAsync();
+			return await _utilization.Find(FilterDefinition<UtilizationDocument>.Empty).SortByDescending(x => x.FinishTime).FirstOrDefaultAsync();
 		}
 	}
 }
