Horde: Fix coding conventions within Horde.Build.

diff --git a/Engine/Source/Programs/Horde/Horde.Build/Fleet/PoolService.cs b/Engine/Source/Programs/Horde/Horde.Build/Fleet/PoolService.cs
--- a/Engine/Source/Programs/Horde/Horde.Build/Fleet/PoolService.cs
+++ b/Engine/Source/Programs/Horde/Horde.Build/Fleet/PoolService.cs
@@ -1,16 +1,16 @@
 // Copyright Epic Games, Inc. All Rights Reserved.
 
-using EpicGames.Horde.Common;
-using HordeCommon;
-using Horde.Build.Collections;
-using Horde.Build.Models;
-using Horde.Build.Utilities;
 using System;
 using System.Collections.Generic;
 using System.Linq;
 using System.Threading;
 using System.Threading.Tasks;
+using EpicGames.Horde.Common;
+using Horde.Build.Collections;
 using Horde.Build.Fleet.Autoscale;
+using Horde.Build.Models;
+using Horde.Build.Utilities;
+using HordeCommon;
 
 namespace Horde.Build.Services
 {
@@ -24,135 +24,135 @@ namespace Horde.Build.Services
 		/// <summary>
 		/// The database service instance
 		/// </summary>
-		DatabaseService DatabaseService;
+		readonly DatabaseService _databaseService;
 
 		/// <summary>
 		/// Collection of pool documents
 		/// </summary>
-		IPoolCollection Pools;
+		readonly IPoolCollection _pools;
 
 		/// <summary>
 		/// Returns the current time
 		/// </summary>
-		IClock Clock;
+		readonly IClock _clock;
 
 		/// <summary>
 		/// Cached set of pools, along with the timestamp that it was obtained
 		/// </summary>
-		Tuple<DateTime, Dictionary<PoolId, IPool>>? CachedPoolLookup;
+		Tuple<DateTime, Dictionary<PoolId, IPool>>? _cachedPoolLookup;
 
 		/// <summary>
 		/// Constructor
 		/// </summary>
-		/// <param name="DatabaseService"></param>
-		/// <param name="Pools">Collection of pool documents</param>
-		/// <param name="Clock"></param>
-		public PoolService(DatabaseService DatabaseService, IPoolCollection Pools, IClock Clock)
+		/// <param name="databaseService"></param>
+		/// <param name="pools">Collection of pool documents</param>
+		/// <param name="clock"></param>
+		public PoolService(DatabaseService databaseService, IPoolCollection pools, IClock clock)
 		{
-			this.DatabaseService = DatabaseService;
-			this.Pools = Pools;
-			this.Clock = Clock;
+			_databaseService = databaseService;
+			_pools = pools;
+			_clock = clock;
 		}
 
 		/// <summary>
 		/// Creates a new pool
 		/// </summary>
-		/// <param name="Name">Name of the new pool</param>
-		/// <param name="Condition">Condition for agents to be automatically included in this pool</param>
-		/// <param name="EnableAutoscaling">Whether to enable autoscaling for this pool</param>
-		/// <param name="MinAgents">Minimum number of agents in the pool</param>
-		/// <param name="NumReserveAgents">Minimum number of idle agents to maintain</param>
-		/// <param name="ScaleOutCooldown">Cooldown time between scale-out events</param>
-		/// <param name="ScaleInCooldown">Cooldown time between scale-in events</param>
-		/// <param name="SizeStrategy">Pool sizing strategy</param>
-		/// <param name="LeaseUtilizationSettings">Settings for lease utilization strategy</param>
-		/// <param name="JobQueueSettings">Settings for job queue strategy</param>
-		/// <param name="Properties">Properties for the new pool</param>
+		/// <param name="name">Name of the new pool</param>
+		/// <param name="condition">Condition for agents to be automatically included in this pool</param>
+		/// <param name="enableAutoscaling">Whether to enable autoscaling for this pool</param>
+		/// <param name="minAgents">Minimum number of agents in the pool</param>
+		/// <param name="numReserveAgents">Minimum number of idle agents to maintain</param>
+		/// <param name="scaleOutCooldown">Cooldown time between scale-out events</param>
+		/// <param name="scaleInCooldown">Cooldown time between scale-in events</param>
+		/// <param name="sizeStrategy">Pool sizing strategy</param>
+		/// <param name="leaseUtilizationSettings">Settings for lease utilization strategy</param>
+		/// <param name="jobQueueSettings">Settings for job queue strategy</param>
+		/// <param name="properties">Properties for the new pool</param>
 		/// <returns>The new pool document</returns>
 		public Task<IPool> CreatePoolAsync(
-			string Name,
-			Condition? Condition = null,
-			bool? EnableAutoscaling = null,
-			int? MinAgents = null,
-			int? NumReserveAgents = null,
-			TimeSpan? ScaleOutCooldown = null,
-			TimeSpan? ScaleInCooldown = null,
-			PoolSizeStrategy? SizeStrategy = null,
-			LeaseUtilizationSettings? LeaseUtilizationSettings = null,
-			JobQueueSettings? JobQueueSettings = null,
-			Dictionary<string, string>? Properties = null)
+			string name,
+			Condition? condition = null,
+			bool? enableAutoscaling = null,
+			int? minAgents = null,
+			int? numReserveAgents = null,
+			TimeSpan? scaleOutCooldown = null,
+			TimeSpan? scaleInCooldown = null,
+			PoolSizeStrategy? sizeStrategy = null,
+			LeaseUtilizationSettings? leaseUtilizationSettings = null,
+			JobQueueSettings? jobQueueSettings = null,
+			Dictionary<string, string>? properties = null)
 		{
-			return Pools.AddAsync(
-				PoolId.Sanitize(Name),
-				Name,
-				Condition,
-				EnableAutoscaling,
-				MinAgents,
-				NumReserveAgents,
-				ScaleOutCooldown,
-				ScaleInCooldown,
-				SizeStrategy,
-				LeaseUtilizationSettings,
-				JobQueueSettings,
-				Properties);
+			return _pools.AddAsync(
+				PoolId.Sanitize(name),
+				name,
+				condition,
+				enableAutoscaling,
+				minAgents,
+				numReserveAgents,
+				scaleOutCooldown,
+				scaleInCooldown,
+				sizeStrategy,
+				leaseUtilizationSettings,
+				jobQueueSettings,
+				properties);
 		}
 
 		/// <summary>
 		/// Deletes a pool
 		/// </summary>
-		/// <param name="PoolId">Unique id of the pool</param>
+		/// <param name="poolId">Unique id of the pool</param>
 		/// <returns>Async task object</returns>
-		public Task<bool> DeletePoolAsync(PoolId PoolId)
+		public Task<bool> DeletePoolAsync(PoolId poolId)
 		{
-			return Pools.DeleteAsync(PoolId);
+			return _pools.DeleteAsync(poolId);
 		}
 
 		/// <summary>
 		/// Updates an existing pool
 		/// </summary>
-		/// <param name="Pool">The pool to update</param>
-		/// <param name="NewName">The new name for the pool</param>
-		/// <param name="NewCondition">New requirements for the pool</param>
-		/// <param name="NewEnableAutoscaling">Whether to enable autoscaling</param>
-		/// <param name="NewMinAgents">Minimum number of agents in the pool</param>
-		/// <param name="NewNumReserveAgents">Minimum number of idle agents to maintain</param>
-		/// <param name="NewProperties">Properties on the pool to update. Any properties with a value of null will be removed.</param>
-		/// <param name="ScaleOutCooldown">Cooldown time between scale-out events</param>
-		/// <param name="ScaleInCooldown">Cooldown time between scale-in events</param>
-		/// <param name="SizeStrategy">New pool sizing strategy for the pool</param>
+		/// <param name="pool">The pool to update</param>
+		/// <param name="newName">The new name for the pool</param>
+		/// <param name="newCondition">New requirements for the pool</param>
+		/// <param name="newEnableAutoscaling">Whether to enable autoscaling</param>
+		/// <param name="newMinAgents">Minimum number of agents in the pool</param>
+		/// <param name="newNumReserveAgents">Minimum number of idle agents to maintain</param>
+		/// <param name="newProperties">Properties on the pool to update. Any properties with a value of null will be removed.</param>
+		/// <param name="scaleOutCooldown">Cooldown time between scale-out events</param>
+		/// <param name="scaleInCooldown">Cooldown time between scale-in events</param>
+		/// <param name="sizeStrategy">New pool sizing strategy for the pool</param>
 		/// <returns>Async task object</returns>
 		public async Task<IPool?> UpdatePoolAsync(
-			IPool? Pool,
-			string? NewName = null,
-			Condition? NewCondition = null,
-			bool? NewEnableAutoscaling = null,
-			int? NewMinAgents = null,
-			int? NewNumReserveAgents = null,
-			Dictionary<string, string?>? NewProperties = null,
-			TimeSpan? ScaleOutCooldown = null,
-			TimeSpan? ScaleInCooldown = null,
-			PoolSizeStrategy? SizeStrategy = null)
+			IPool? pool,
+			string? newName = null,
+			Condition? newCondition = null,
+			bool? newEnableAutoscaling = null,
+			int? newMinAgents = null,
+			int? newNumReserveAgents = null,
+			Dictionary<string, string?>? newProperties = null,
+			TimeSpan? scaleOutCooldown = null,
+			TimeSpan? scaleInCooldown = null,
+			PoolSizeStrategy? sizeStrategy = null)
 		{
-			for (; Pool != null; Pool = await Pools.GetAsync(Pool.Id))
+			for (; pool != null; pool = await _pools.GetAsync(pool.Id))
 			{
-				IPool? NewPool = await Pools.TryUpdateAsync(
-					Pool,
-					NewName,
-					NewCondition,
-					NewEnableAutoscaling,
-					NewMinAgents,
-					NewNumReserveAgents,
-					NewProperties: NewProperties,
-					ScaleOutCooldown: ScaleOutCooldown,
-					ScaleInCooldown: ScaleInCooldown,
-					SizeStrategy: SizeStrategy);
+				IPool? newPool = await _pools.TryUpdateAsync(
+					pool,
+					newName,
+					newCondition,
+					newEnableAutoscaling,
+					newMinAgents,
+					newNumReserveAgents,
+					newProperties: newProperties,
+					scaleOutCooldown: scaleOutCooldown,
+					scaleInCooldown: scaleInCooldown,
+					sizeStrategy: sizeStrategy);
 				
-				if (NewPool != null)
+				if (newPool != null)
 				{
-					return NewPool;
+					return newPool;
 				}
 			}
-			return Pool;
+			return pool;
 		}
 
 		/// <summary>
@@ -161,78 +161,78 @@ namespace Horde.Build.Services
 		/// <returns>List of pool documents</returns>
 		public Task<List<IPool>> GetPoolsAsync()
 		{
-			return Pools.GetAsync();
+			return _pools.GetAsync();
 		}
 
 		/// <summary>
 		/// Gets a pool by ID
 		/// </summary>
-		/// <param name="PoolId">Unique id of the pool</param>
+		/// <param name="poolId">Unique id of the pool</param>
 		/// <returns>The pool document</returns>
-		public Task<IPool?> GetPoolAsync(PoolId PoolId)
+		public Task<IPool?> GetPoolAsync(PoolId poolId)
 		{
-			return Pools.GetAsync(PoolId);
+			return _pools.GetAsync(poolId);
 		}
 
 		/// <summary>
 		/// Get a list of workspaces for the given agent
 		/// </summary>
-		/// <param name="Agent">The agent to return workspaces for</param>
-		/// <param name="ValidAtTime">Absolute time at which we expect the results to be valid. Values may be cached as long as they are after this time.</param>
+		/// <param name="agent">The agent to return workspaces for</param>
+		/// <param name="validAtTime">Absolute time at which we expect the results to be valid. Values may be cached as long as they are after this time.</param>
 		/// <returns>List of workspaces</returns>
-		public async Task<HashSet<AgentWorkspace>> GetWorkspacesAsync(IAgent Agent, DateTime ValidAtTime)
+		public async Task<HashSet<AgentWorkspace>> GetWorkspacesAsync(IAgent agent, DateTime validAtTime)
 		{
 			bool bAddAutoSdkWorkspace = false;
-			HashSet<AgentWorkspace> Workspaces = new HashSet<AgentWorkspace>();
+			HashSet<AgentWorkspace> workspaces = new HashSet<AgentWorkspace>();
 
-			Dictionary<PoolId, IPool> PoolMapping = await GetPoolLookupAsync(ValidAtTime);
-			foreach (PoolId PoolId in Agent.GetPools())
+			Dictionary<PoolId, IPool> poolMapping = await GetPoolLookupAsync(validAtTime);
+			foreach (PoolId poolId in agent.GetPools())
 			{
-				IPool? Pool;
-				if (PoolMapping.TryGetValue(PoolId, out Pool))
+				IPool? pool;
+				if (poolMapping.TryGetValue(poolId, out pool))
 				{
-					Workspaces.UnionWith(Pool.Workspaces);
-					bAddAutoSdkWorkspace |= Pool.UseAutoSdk;
+					workspaces.UnionWith(pool.Workspaces);
+					bAddAutoSdkWorkspace |= pool.UseAutoSdk;
 				}
 			}
 
 			if (bAddAutoSdkWorkspace)
 			{
-				Globals Globals = await DatabaseService.GetGlobalsAsync();
-				Workspaces.UnionWith(Agent.GetAutoSdkWorkspaces(Globals, Workspaces.ToList()));
+				Globals globals = await _databaseService.GetGlobalsAsync();
+				workspaces.UnionWith(agent.GetAutoSdkWorkspaces(globals, workspaces.ToList()));
 			}
 
-			return Workspaces;
+			return workspaces;
 		}
 
 		/// <summary>
 		/// Gets a mapping from pool identifiers to definitions
 		/// </summary>
-		/// <param name="ValidAtTime">Absolute time at which we expect the results to be valid. Values may be cached as long as they are after this time.</param>
+		/// <param name="validAtTime">Absolute time at which we expect the results to be valid. Values may be cached as long as they are after this time.</param>
 		/// <returns>Map of pool ids to pool documents</returns>
-		private async Task<Dictionary<PoolId, IPool>> GetPoolLookupAsync(DateTime ValidAtTime)
+		private async Task<Dictionary<PoolId, IPool>> GetPoolLookupAsync(DateTime validAtTime)
 		{
-			Tuple<DateTime, Dictionary<PoolId, IPool>>? CachedPoolLookupCopy = CachedPoolLookup;
-			if (CachedPoolLookupCopy == null || CachedPoolLookupCopy.Item1 < ValidAtTime)
+			Tuple<DateTime, Dictionary<PoolId, IPool>>? cachedPoolLookupCopy = _cachedPoolLookup;
+			if (cachedPoolLookupCopy == null || cachedPoolLookupCopy.Item1 < validAtTime)
 			{
 				// Get a new list of cached pools
-				DateTime NewCacheTime = Clock.UtcNow;
-				List<IPool> NewPools = await Pools.GetAsync();
-				Tuple<DateTime, Dictionary<PoolId, IPool>> NewCachedPoolLookup = Tuple.Create(NewCacheTime, NewPools.ToDictionary(x => x.Id, x => x));
+				DateTime newCacheTime = _clock.UtcNow;
+				List<IPool> newPools = await _pools.GetAsync();
+				Tuple<DateTime, Dictionary<PoolId, IPool>> newCachedPoolLookup = Tuple.Create(newCacheTime, newPools.ToDictionary(x => x.Id, x => x));
 
 				// Try to swap it with the current version
-				while (CachedPoolLookupCopy == null || CachedPoolLookupCopy.Item1 < NewCacheTime)
+				while (cachedPoolLookupCopy == null || cachedPoolLookupCopy.Item1 < newCacheTime)
 				{
-					Tuple<DateTime, Dictionary<PoolId, IPool>>? OriginalValue = Interlocked.CompareExchange(ref CachedPoolLookup, NewCachedPoolLookup, CachedPoolLookupCopy);
-					if (OriginalValue == CachedPoolLookupCopy)
+					Tuple<DateTime, Dictionary<PoolId, IPool>>? originalValue = Interlocked.CompareExchange(ref _cachedPoolLookup, newCachedPoolLookup, cachedPoolLookupCopy);
+					if (originalValue == cachedPoolLookupCopy)
 					{
-						CachedPoolLookupCopy = NewCachedPoolLookup;
+						cachedPoolLookupCopy = newCachedPoolLookup;
 						break;
 					}
-					CachedPoolLookupCopy = OriginalValue;
+					cachedPoolLookupCopy = originalValue;
 				}
 			}
-			return CachedPoolLookupCopy.Item2;
+			return cachedPoolLookupCopy.Item2;
 		}
 	}
 }
diff --git a/Engine/Source/Programs/Horde/Horde.Build/Fleet/PoolUpdateService.cs b/Engine/Source/Programs/Horde/Horde.Build/Fleet/PoolUpdateService.cs
--- a/Engine/Source/Programs/Horde/Horde.Build/Fleet/PoolUpdateService.cs
+++ b/Engine/Source/Programs/Horde/Horde.Build/Fleet/PoolUpdateService.cs
@@ -1,16 +1,16 @@
 // Copyright Epic Games, Inc. All Rights Reserved.
 
-using HordeCommon;
-using Horde.Build.Collections;
-using Horde.Build.Models;
-using Horde.Build.Utilities;
-using Microsoft.Extensions.Hosting;
-using Microsoft.Extensions.Logging;
 using System;
 using System.Collections.Generic;
 using System.Linq;
 using System.Threading;
 using System.Threading.Tasks;
+using Horde.Build.Collections;
+using Horde.Build.Models;
+using Horde.Build.Utilities;
+using HordeCommon;
+using Microsoft.Extensions.Hosting;
+using Microsoft.Extensions.Logging;
 
 namespace Horde.Build.Services
 {
@@ -21,111 +21,111 @@ namespace Horde.Build.Services
 	/// </summary>
 	public sealed class PoolUpdateService : IHostedService, IDisposable
 	{
-		IPoolCollection Pools;
-		IStreamCollection Streams;
-		ILogger<PoolService> Logger;
-		ITicker Ticker;
+		readonly IPoolCollection _pools;
+		readonly IStreamCollection _streams;
+		readonly ILogger<PoolService> _logger;
+		readonly ITicker _ticker;
 
 		/// <summary>
 		/// Constructor
 		/// </summary>
-		public PoolUpdateService(IPoolCollection Pools, IStreamCollection Streams, IClock Clock, ILogger<PoolService> Logger)
+		public PoolUpdateService(IPoolCollection pools, IStreamCollection streams, IClock clock, ILogger<PoolService> logger)
 		{
-			this.Pools = Pools;
-			this.Streams = Streams;
-			this.Logger = Logger;
-			this.Ticker = Clock.AddTicker(TimeSpan.FromSeconds(30.0), UpdatePoolsAsync, Logger);
+			_pools = pools;
+			_streams = streams;
+			_logger = logger;
+			_ticker = clock.AddTicker($"{nameof(PoolUpdateService)}.{nameof(UpdatePoolsAsync)}", TimeSpan.FromSeconds(30.0), UpdatePoolsAsync, logger);
 		}
 
 		/// <inheritdoc/>
-		public Task StartAsync(CancellationToken CancellationToken) => Ticker.StartAsync();
+		public Task StartAsync(CancellationToken cancellationToken) => _ticker.StartAsync();
 
 		/// <inheritdoc/>
-		public Task StopAsync(CancellationToken CancellationToken) => Ticker.StopAsync();
+		public Task StopAsync(CancellationToken cancellationToken) => _ticker.StopAsync();
 
 		/// <inheritdoc/>
-		public void Dispose() => Ticker.Dispose();
+		public void Dispose() => _ticker.Dispose();
 
 		/// <summary>
 		/// Execute the background task
 		/// </summary>
-		/// <param name="StoppingToken">Cancellation token for the async task</param>
+		/// <param name="stoppingToken">Cancellation token for the async task</param>
 		/// <returns>Async task</returns>
-		async ValueTask UpdatePoolsAsync(CancellationToken StoppingToken)
+		async ValueTask UpdatePoolsAsync(CancellationToken stoppingToken)
 		{
 			// Capture the start time for this operation. We use this to attempt to sequence updates to agents, and prevent overriding another server's updates.
-			DateTime StartTime = DateTime.UtcNow;
+			DateTime startTime = DateTime.UtcNow;
 
 			// Update the list
 			bool bRetryUpdate = true;
-			while (bRetryUpdate && !StoppingToken.IsCancellationRequested)
+			while (bRetryUpdate && !stoppingToken.IsCancellationRequested)
 			{
-				Logger.LogDebug("Updating pool->workspace map");
+				_logger.LogDebug("Updating pool->workspace map");
 
 				// Assume this will be the last iteration
 				bRetryUpdate = false;
 
 				// Capture the list of pools at the start of this update
-				List<IPool> CurrentPools = await Pools.GetAsync();
+				List<IPool> currentPools = await _pools.GetAsync();
 
 				// Lookup table of pool id to workspaces
-				HashSet<PoolId> PoolsWithAutoSdk = new HashSet<PoolId>();
-				Dictionary<PoolId, List<AgentWorkspace>> PoolToAgentWorkspaces = new Dictionary<PoolId, List<AgentWorkspace>>();
+				HashSet<PoolId> poolsWithAutoSdk = new HashSet<PoolId>();
+				Dictionary<PoolId, List<AgentWorkspace>> poolToAgentWorkspaces = new Dictionary<PoolId, List<AgentWorkspace>>();
 
 				// Populate the workspace list from the current stream
-				List<IStream> ActiveStreams = await Streams.FindAllAsync();
-				foreach (IStream ActiveStream in ActiveStreams)
+				List<IStream> activeStreams = await _streams.FindAllAsync();
+				foreach (IStream activeStream in activeStreams)
 				{
-					foreach (KeyValuePair<string, AgentType> AgentTypePair in ActiveStream.AgentTypes)
+					foreach (KeyValuePair<string, AgentType> agentTypePair in activeStream.AgentTypes)
 					{
 						// Create the new agent workspace
-						(AgentWorkspace, bool)? Result;
-						if (ActiveStream.TryGetAgentWorkspace(AgentTypePair.Value, out Result))
+						(AgentWorkspace, bool)? result;
+						if (activeStream.TryGetAgentWorkspace(agentTypePair.Value, out result))
 						{
-							(AgentWorkspace AgentWorkspace, bool UseAutoSdk) = Result.Value;
-							AgentType AgentType = AgentTypePair.Value;
+							(AgentWorkspace agentWorkspace, bool useAutoSdk) = result.Value;
+							AgentType agentType = agentTypePair.Value;
 
 							// Find or add a list of workspaces for this pool
-							List<AgentWorkspace>? AgentWorkspaces;
-							if (!PoolToAgentWorkspaces.TryGetValue(AgentType.Pool, out AgentWorkspaces))
+							List<AgentWorkspace>? agentWorkspaces;
+							if (!poolToAgentWorkspaces.TryGetValue(agentType.Pool, out agentWorkspaces))
 							{
-								AgentWorkspaces = new List<AgentWorkspace>();
-								PoolToAgentWorkspaces.Add(AgentType.Pool, AgentWorkspaces);
+								agentWorkspaces = new List<AgentWorkspace>();
+								poolToAgentWorkspaces.Add(agentType.Pool, agentWorkspaces);
 							}
 
 							// Add it to the list
-							if (!AgentWorkspaces.Contains(AgentWorkspace))
+							if (!agentWorkspaces.Contains(agentWorkspace))
 							{
-								AgentWorkspaces.Add(AgentWorkspace);
+								agentWorkspaces.Add(agentWorkspace);
 							}
-							if (UseAutoSdk)
+							if (useAutoSdk)
 							{
-								PoolsWithAutoSdk.Add(AgentType.Pool);
+								poolsWithAutoSdk.Add(agentType.Pool);
 							}
 						}
 					}
 				}
 
 				// Update the list of workspaces for each pool
-				foreach (IPool CurrentPool in CurrentPools)
+				foreach (IPool currentPool in currentPools)
 				{
 					// Get the new list of workspaces for this pool
-					List<AgentWorkspace>? NewWorkspaces;
-					if (!PoolToAgentWorkspaces.TryGetValue(CurrentPool.Id, out NewWorkspaces))
+					List<AgentWorkspace>? newWorkspaces;
+					if (!poolToAgentWorkspaces.TryGetValue(currentPool.Id, out newWorkspaces))
 					{
-						NewWorkspaces = new List<AgentWorkspace>();
+						newWorkspaces = new List<AgentWorkspace>();
 					}
 
 					// Update the pools document
-					bool UseAutoSdk = PoolsWithAutoSdk.Contains(CurrentPool.Id);
-					if (!AgentWorkspace.SetEquals(CurrentPool.Workspaces, NewWorkspaces) || CurrentPool.Workspaces.Count != NewWorkspaces.Count || CurrentPool.UseAutoSdk != UseAutoSdk)
+					bool useAutoSdk = poolsWithAutoSdk.Contains(currentPool.Id);
+					if (!AgentWorkspace.SetEquals(currentPool.Workspaces, newWorkspaces) || currentPool.Workspaces.Count != newWorkspaces.Count || currentPool.UseAutoSdk != useAutoSdk)
 					{
-						Logger.LogInformation("New workspaces for pool {Pool}:{Workspaces}", CurrentPool.Id, String.Join("", NewWorkspaces.Select(x => $"\n  Identifier=\"{x.Identifier}\", Stream={x.Stream}")));
+						_logger.LogInformation("New workspaces for pool {Pool}:{Workspaces}", currentPool.Id, String.Join("", newWorkspaces.Select(x => $"\n  Identifier=\"{x.Identifier}\", Stream={x.Stream}")));
 
-						IPool? Result = await Pools.TryUpdateAsync(CurrentPool, NewWorkspaces: NewWorkspaces, NewUseAutoSdk: UseAutoSdk);
-						if (Result == null)
+						IPool? result = await _pools.TryUpdateAsync(currentPool, newWorkspaces: newWorkspaces, newUseAutoSdk: useAutoSdk);
+						if (result == null)
 						{
-							Logger.LogInformation("Pool modified; will retry");
+							_logger.LogInformation("Pool modified; will retry");
 							bRetryUpdate = true;
 						}
 					}
diff --git a/Engine/Source/Programs/Horde/Horde.Build/Fleet/Telemetry.cs b/Engine/Source/Programs/Horde/Horde.Build/Fleet/Telemetry.cs
--- a/Engine/Source/Programs/Horde/Horde.Build/Fleet/Telemetry.cs
+++ b/Engine/Source/Programs/Horde/Horde.Build/Fleet/Telemetry.cs
@@ -1,10 +1,8 @@
 // Copyright Epic Games, Inc. All Rights Reserved.
 
-using Horde.Build.Utilities;
 using System;
 using System.Collections.Generic;
-using System.Linq;
-using System.Threading.Tasks;
+using Horde.Build.Utilities;
 
 namespace Horde.Build.Models
 {
@@ -76,29 +74,29 @@ namespace Horde.Build.Models
 		/// <summary>
 		/// Constructor
 		/// </summary>
-		/// <param name="StartTime"></param>
-		/// <param name="FinishTime"></param>
-		public NewUtilizationTelemetry(DateTime StartTime, DateTime FinishTime)
+		/// <param name="startTime"></param>
+		/// <param name="finishTime"></param>
+		public NewUtilizationTelemetry(DateTime startTime, DateTime finishTime)
 		{
-			this.StartTime = StartTime;
-			this.FinishTime = FinishTime;
+			StartTime = startTime;
+			FinishTime = finishTime;
 		}
 
 		/// <summary>
 		/// Adds a new pool
 		/// </summary>
-		/// <param name="PoolId">The pool id</param>
+		/// <param name="poolId">The pool id</param>
 		/// <returns>Telemetry for the given pool</returns>
-		public NewPoolUtilizationTelemetry FindOrAddPool(PoolId PoolId)
+		public NewPoolUtilizationTelemetry FindOrAddPool(PoolId poolId)
 		{
-			NewPoolUtilizationTelemetry? Pool;
-			if (!PoolsLookup.TryGetValue(PoolId, out Pool))
+			NewPoolUtilizationTelemetry? pool;
+			if (!PoolsLookup.TryGetValue(poolId, out pool))
 			{
-				Pool = new NewPoolUtilizationTelemetry(PoolId);
-				Pools.Add(Pool);
-				PoolsLookup.Add(PoolId, Pool);
+				pool = new NewPoolUtilizationTelemetry(poolId);
+				Pools.Add(pool);
+				PoolsLookup.Add(poolId, pool);
 			}
-			return Pool;
+			return pool;
 		}
 	}
 
@@ -167,26 +165,26 @@ namespace Horde.Build.Models
 		/// <summary>
 		/// Constructor
 		/// </summary>
-		/// <param name="PoolId"></param>
-		public NewPoolUtilizationTelemetry(PoolId PoolId)
+		/// <param name="poolId"></param>
+		public NewPoolUtilizationTelemetry(PoolId poolId)
 		{
-			this.PoolId = PoolId;
+			PoolId = poolId;
 		}
 
 		/// <summary>
 		/// Adds a stream to the object
 		/// </summary>
-		/// <param name="StreamId"></param>
-		public NewPoolUtilizationTelemetryStream FindOrAddStream(StreamId StreamId)
+		/// <param name="streamId"></param>
+		public NewPoolUtilizationTelemetryStream FindOrAddStream(StreamId streamId)
 		{
-			NewPoolUtilizationTelemetryStream? Stream;
-			if (!StreamLookup.TryGetValue(StreamId, out Stream))
+			NewPoolUtilizationTelemetryStream? stream;
+			if (!StreamLookup.TryGetValue(streamId, out stream))
 			{
-				Stream = new NewPoolUtilizationTelemetryStream(StreamId);
-				Streams.Add(Stream);
-				StreamLookup.Add(StreamId, Stream);
+				stream = new NewPoolUtilizationTelemetryStream(streamId);
+				Streams.Add(stream);
+				StreamLookup.Add(streamId, stream);
 			}
-			return Stream;
+			return stream;
 		}
 	}
 
@@ -220,10 +218,10 @@ namespace Horde.Build.Models
 		/// <summary>
 		/// Constructor
 		/// </summary>
-		/// <param name="StreamId"></param>
-		public NewPoolUtilizationTelemetryStream(StreamId StreamId)
+		/// <param name="streamId"></param>
+		public NewPoolUtilizationTelemetryStream(StreamId streamId)
 		{
-			this.StreamId = StreamId;
+			StreamId = streamId;
 		}
 	}
 }
diff --git a/Engine/Source/Programs/Horde/Horde.Build/Fleet/TelemetryService.cs b/Engine/Source/Programs/Horde/Horde.Build/Fleet/TelemetryService.cs
--- a/Engine/Source/Programs/Horde/Horde.Build/Fleet/TelemetryService.cs
+++ b/Engine/Source/Programs/Horde/Horde.Build/Fleet/TelemetryService.cs
@@ -1,19 +1,16 @@
 // Copyright Epic Games, Inc. All Rights Reserved.
 
-using Google.Protobuf.WellKnownTypes;
-using HordeCommon;
-using HordeCommon.Rpc.Tasks;
-using Horde.Build.Collections;
-using Horde.Build.Models;
-using Horde.Build.Utilities;
-using Microsoft.Extensions.Hosting;
-using Microsoft.Extensions.Logging;
-using MongoDB.Bson;
 using System;
 using System.Collections.Generic;
 using System.Linq;
 using System.Threading;
 using System.Threading.Tasks;
+using Horde.Build.Collections;
+using Horde.Build.Models;
+using Horde.Build.Utilities;
+using HordeCommon;
+using Microsoft.Extensions.Hosting;
+using Microsoft.Extensions.Logging;
 
 namespace Horde.Build.Services
 {
@@ -24,121 +21,121 @@ namespace Horde.Build.Services
 	/// </summary>
 	public sealed class TelemetryService : IHostedService, IDisposable
 	{
-		ITelemetryCollection TelemetryCollection;
-		IAgentCollection AgentCollection;
-		ILeaseCollection LeaseCollection;
-		IPoolCollection PoolCollection;
-		IFleetManager FleetManager;
-		IClock Clock;
-		ITicker Tick;
-		ILogger Logger;
+		readonly ITelemetryCollection _telemetryCollection;
+		readonly IAgentCollection _agentCollection;
+		readonly ILeaseCollection _leaseCollection;
+		readonly IPoolCollection _poolCollection;
+		readonly IFleetManager _fleetManager;
+		readonly IClock _clock;
+		readonly ITicker _tick;
+		readonly ILogger _logger;
 
 		/// <summary>
 		/// Constructor
 		/// </summary>
-		public TelemetryService(ITelemetryCollection TelemetryCollection, IAgentCollection AgentCollection, ILeaseCollection LeaseCollection, IPoolCollection PoolCollection, IFleetManager FleetManager, IClock Clock, ILogger<TelemetryService> Logger)
+		public TelemetryService(ITelemetryCollection telemetryCollection, IAgentCollection agentCollection, ILeaseCollection leaseCollection, IPoolCollection poolCollection, IFleetManager fleetManager, IClock clock, ILogger<TelemetryService> logger)
 		{
-			this.TelemetryCollection = TelemetryCollection;
-			this.AgentCollection = AgentCollection;
-			this.LeaseCollection = LeaseCollection;
-			this.PoolCollection = PoolCollection;
-			this.FleetManager = FleetManager;
-			this.Clock = Clock;
-			this.Tick = Clock.AddSharedTicker<TelemetryService>(TimeSpan.FromMinutes(10.0), TickLeaderAsync, Logger);
-			this.Logger = Logger;
+			_telemetryCollection = telemetryCollection;
+			_agentCollection = agentCollection;
+			_leaseCollection = leaseCollection;
+			_poolCollection = poolCollection;
+			_fleetManager = fleetManager;
+			_clock = clock;
+			_tick = clock.AddSharedTicker<TelemetryService>(TimeSpan.FromMinutes(10.0), TickLeaderAsync, logger);
+			_logger = logger;
 		}
 
 		/// <inheritdoc/>
-		public Task StartAsync(CancellationToken CancellationToken) => Tick.StartAsync();
+		public Task StartAsync(CancellationToken cancellationToken) => _tick.StartAsync();
 
 		/// <inheritdoc/>
-		public Task StopAsync(CancellationToken CancellationToken) => Tick.StopAsync();
+		public Task StopAsync(CancellationToken cancellationToken) => _tick.StopAsync();
 
 		/// <inheritdoc/>
-		public void Dispose() => Tick.Dispose();
+		public void Dispose() => _tick.Dispose();
 
 		/// <inheritdoc/>
-		async ValueTask TickLeaderAsync(CancellationToken StoppingToken)
+		async ValueTask TickLeaderAsync(CancellationToken stoppingToken)
 		{
-			DateTime CurrentTime = Clock.UtcNow;
+			DateTime currentTime = _clock.UtcNow;
 
 			// Find the last time that we need telemetry for
-			DateTime MaxTime = CurrentTime.Date + TimeSpan.FromHours(CurrentTime.Hour);
+			DateTime maxTime = currentTime.Date + TimeSpan.FromHours(currentTime.Hour);
 
 			// Get the latest telemetry data
-			IUtilizationTelemetry? Latest = await TelemetryCollection.GetLatestUtilizationTelemetryAsync();
-			TimeSpan Interval = TimeSpan.FromHours(1.0);
-			int Count = (Latest == null) ? (7 * 24) : (int)Math.Round((MaxTime - Latest.FinishTime) / Interval);
-			DateTime MinTime = MaxTime - Count * Interval;
+			IUtilizationTelemetry? latest = await _telemetryCollection.GetLatestUtilizationTelemetryAsync();
+			TimeSpan interval = TimeSpan.FromHours(1.0);
+			int count = (latest == null) ? (7 * 24) : (int)Math.Round((maxTime - latest.FinishTime) / interval);
+			DateTime minTime = maxTime - count * interval;
 
 			// Query all the current data
-			List<IAgent> Agents = await AgentCollection.FindAsync();
-			List<IPool> Pools = await PoolCollection.GetAsync();
-			List<ILease> Leases = await LeaseCollection.FindLeasesAsync(MinTime: MinTime);
+			List<IAgent> agents = await _agentCollection.FindAsync();
+			List<IPool> pools = await _poolCollection.GetAsync();
+			List<ILease> leases = await _leaseCollection.FindLeasesAsync(minTime: minTime);
 
 			// Remove any agents which are offline
-			Agents.RemoveAll(x => !x.Enabled || !x.IsSessionValid(CurrentTime));
+			agents.RemoveAll(x => !x.Enabled || !x.IsSessionValid(currentTime));
 
 			// Find all the agents
-			Dictionary<AgentId, List<PoolId>> AgentToPoolIds = Agents.ToDictionary(x => x.Id, x => x.GetPools().ToList());
+			Dictionary<AgentId, List<PoolId>> agentToPoolIds = agents.ToDictionary(x => x.Id, x => x.GetPools().ToList());
 
 			// Generate all the telemetry data
-			DateTime BucketMinTime = MinTime;
-			for (int Idx = 0; Idx < Count; Idx++)
+			DateTime bucketMinTime = minTime;
+			for (int idx = 0; idx < count; idx++)
 			{
-				DateTime BucketMaxTime = BucketMinTime + Interval;
-				Logger.LogInformation("Creating telemetry for {MinTime} to {MaxTime}", BucketMinTime, BucketMaxTime);
+				DateTime bucketMaxTime = bucketMinTime + interval;
+				_logger.LogInformation("Creating telemetry for {MinTime} to {MaxTime}", bucketMinTime, bucketMaxTime);
 
-				NewUtilizationTelemetry Telemetry = new NewUtilizationTelemetry(BucketMinTime, BucketMaxTime);
-				Telemetry.NumAgents = Agents.Count;
-				foreach (IPool Pool in Pools)
+				NewUtilizationTelemetry telemetry = new NewUtilizationTelemetry(bucketMinTime, bucketMaxTime);
+				telemetry.NumAgents = agents.Count;
+				foreach (IPool pool in pools)
 				{
-					if (Pool.EnableAutoscaling)
+					if (pool.EnableAutoscaling)
 					{
-						int NumStoppedInstances = await FleetManager.GetNumStoppedInstancesAsync(Pool);
-						Telemetry.NumAgents += NumStoppedInstances;
+						int numStoppedInstances = await _fleetManager.GetNumStoppedInstancesAsync(pool);
+						telemetry.NumAgents += numStoppedInstances;
 
-						NewPoolUtilizationTelemetry PoolTelemetry = Telemetry.FindOrAddPool(Pool.Id);
-						PoolTelemetry.NumAgents += NumStoppedInstances;
-						PoolTelemetry.HibernatingTime += Interval.TotalHours * NumStoppedInstances;
+						NewPoolUtilizationTelemetry poolTelemetry = telemetry.FindOrAddPool(pool.Id);
+						poolTelemetry.NumAgents += numStoppedInstances;
+						poolTelemetry.HibernatingTime += interval.TotalHours * numStoppedInstances;
 					}
 				}
-				foreach (PoolId PoolId in AgentToPoolIds.Values.SelectMany(x => x))
+				foreach (PoolId poolId in agentToPoolIds.Values.SelectMany(x => x))
 				{
-					Telemetry.FindOrAddPool(PoolId).NumAgents++;
+					telemetry.FindOrAddPool(poolId).NumAgents++;
 				}
-				foreach (ILease Lease in Leases)
+				foreach (ILease lease in leases)
 				{
-					if (Lease.StartTime < BucketMaxTime && (!Lease.FinishTime.HasValue || Lease.FinishTime >= BucketMinTime))
+					if (lease.StartTime < bucketMaxTime && (!lease.FinishTime.HasValue || lease.FinishTime >= bucketMinTime))
 					{
-						List<PoolId>? LeasePools;
-						if (AgentToPoolIds.TryGetValue(Lease.AgentId, out LeasePools))
+						List<PoolId>? leasePools;
+						if (agentToPoolIds.TryGetValue(lease.AgentId, out leasePools))
 						{
-							DateTime FinishTime = Lease.FinishTime ?? BucketMaxTime;
-							double Time = new TimeSpan(Math.Min(FinishTime.Ticks, BucketMaxTime.Ticks) - Math.Max(Lease.StartTime.Ticks, BucketMinTime.Ticks)).TotalHours;
+							DateTime finishTime = lease.FinishTime ?? bucketMaxTime;
+							double time = new TimeSpan(Math.Min(finishTime.Ticks, bucketMaxTime.Ticks) - Math.Max(lease.StartTime.Ticks, bucketMinTime.Ticks)).TotalHours;
 
-							foreach (PoolId PoolId in LeasePools)
+							foreach (PoolId poolId in leasePools)
 							{
-								NewPoolUtilizationTelemetry PoolTelemetry = Telemetry.FindOrAddPool(PoolId);
-								if (Lease.PoolId == null || Lease.StreamId == null)
+								NewPoolUtilizationTelemetry poolTelemetry = telemetry.FindOrAddPool(poolId);
+								if (lease.PoolId == null || lease.StreamId == null)
 								{
-									PoolTelemetry.AdminTime += Time;
+									poolTelemetry.AdminTime += time;
 								}
-								else if (PoolId == Lease.PoolId)
+								else if (poolId == lease.PoolId)
 								{
-									PoolTelemetry.FindOrAddStream(Lease.StreamId.Value).Time += Time;
+									poolTelemetry.FindOrAddStream(lease.StreamId.Value).Time += time;
 								}
 								else
 								{
-									PoolTelemetry.OtherTime += Time;
+									poolTelemetry.OtherTime += time;
 								}
 							}
 						}
 					}
 				}
-				await TelemetryCollection.AddUtilizationTelemetryAsync(Telemetry);
+				await _telemetryCollection.AddUtilizationTelemetryAsync(telemetry);
 
-				BucketMinTime = BucketMaxTime;
+				bucketMinTime = bucketMaxTime;
 			}
 		}
 	}
diff --git a/Engine/Source/Programs/Horde/Horde.Build/GlobalSuppressions.cs b/Engine/Source/Programs/Horde/Horde.Build/GlobalSuppressions.cs
--- a/Engine/Source/Programs/Horde/Horde.Build/GlobalSuppressions.cs
+++ b/Engine/Source/Programs/Horde/Horde.Build/GlobalSuppressions.cs
@@ -7,37 +7,12 @@
 
 using System.Diagnostics.CodeAnalysis;
 
-[assembly: System.Diagnostics.CodeAnalysis.SuppressMessage("Reliability", "CA2007:Consider calling ConfigureAwait on the awaited task", Justification = "Not necessary in NET Core; no SynchronizationContext.", Scope = "module")]
-[assembly: System.Diagnostics.CodeAnalysis.SuppressMessage("Usage", "CA2227:Collection properties should be read only", Justification = "<Pending>", Scope = "module")]
-[assembly: System.Diagnostics.CodeAnalysis.SuppressMessage("Design", "CA1062:Validate arguments of public methods", Justification = "<Pending>", Scope = "module")]
-[assembly: System.Diagnostics.CodeAnalysis.SuppressMessage("Globalization", "CA1303:Do not pass literals as localized parameters", Justification = "<Pending>", Scope = "module")]
-[assembly: System.Diagnostics.CodeAnalysis.SuppressMessage("Performance", "CA1819:Properties should not return arrays", Justification = "<Pending>", Scope = "module")]
-[assembly: System.Diagnostics.CodeAnalysis.SuppressMessage("Naming", "CA1716:Identifiers should not match keywords", Justification = "<Pending>", Scope = "type", Target = "~T:HordeServer.Models.Event")]
-[assembly: System.Diagnostics.CodeAnalysis.SuppressMessage("Design", "CA1032:Implement standard exception constructors", Justification = "<Pending>", Scope = "type", Target = "~T:HordeServer.Services.RetryNotAllowedException")]
-[assembly: System.Diagnostics.CodeAnalysis.SuppressMessage("Naming", "CA1714:Flags enums should have plural names", Justification = "<Pending>", Scope = "type", Target = "~T:HordeServer.Api.ScheduleDaysOfWeek")]
-[assembly: System.Diagnostics.CodeAnalysis.SuppressMessage("Naming", "CA1720:Identifier contains type name", Justification = "<Pending>", Scope = "member", Target = "~F:HordeServer.Api.TemplateParameterType.String")]
-
-[assembly: System.Diagnostics.CodeAnalysis.SuppressMessage("Build", "CA1822:Member Login does not access instance data and can be marked as static (Shared in VisualBasic)", Justification = "<Pending>", Scope = "member", Target = "~M:HordeServer.Controllers.AccountController.Login~Microsoft.AspNetCore.Mvc.IActionResult")]
-[assembly: System.Diagnostics.CodeAnalysis.SuppressMessage("Build", "CA1031:Modify 'Logout' to catch a more specific allowed exception type, or rethrow the exception.", Justification = "<Pending>", Scope = "member", Target = "~M:HordeServer.Controllers.AccountController.Logout~System.Threading.Tasks.Task{Microsoft.AspNetCore.Mvc.IActionResult}")]
-[assembly: System.Diagnostics.CodeAnalysis.SuppressMessage("Build", "CA1305:The behavior of 'string.ToString()' could vary based on the current user's locale settings. Replace this call in 'AgentLease.ToResponse()' with a call to 'string.ToString(IFormatProvider)'.", Justification = "<Pending>", Scope = "member", Target = "~M:HordeServer.Models.AgentLease.ToResponse~HordeServer.Api.GetAgentLeaseResponse")]
-[assembly: System.Diagnostics.CodeAnalysis.SuppressMessage("Build", "CA5380:Adding certificates to the operating system's trusted root certificates increases the risk of incorrectly authenticating an illegitimate certificate", Justification = "<Pending>", Scope = "member", Target = "~T:HordeServer.Services.DatabaseService")]
-[assembly: System.Diagnostics.CodeAnalysis.SuppressMessage("Build", "CA1724:The type name Template conflicts in whole or in part with the namespace name 'Microsoft.AspNetCore.Routing.Template'. Change either name to eliminate the conflict.", Justification = "<Pending>", Scope = "type", Target = "~T:HordeServer.Models.Template")]
-[assembly: System.Diagnostics.CodeAnalysis.SuppressMessage("Build", "CA1724:The type name Session conflicts in whole or in part with the namespace name 'Microsoft.AspNetCore.Session'. Change either name to eliminate the conflict.", Justification = "<Pending>", Scope = "type", Target = "~T:HordeServer.Models.Session")]
-[assembly: System.Diagnostics.CodeAnalysis.SuppressMessage("Build", "CA1812:BotsService is an internal class that is apparently never instantiated. If so, remove the code from the assembly. If this class is intended to contain only static members, make it static (Shared in Visual Basic).", Justification = "<Pending>", Scope = "type", Target = "~T:HordeServer.Services.BotsService")]
-[assembly: System.Diagnostics.CodeAnalysis.SuppressMessage("Build", "CA1812:RpcService is an internal class that is apparently never instantiated. If so, remove the code from the assembly. If this class is intended to contain only static members, make it static (Shared in Visual Basic).", Justification = "<Pending>", Scope = "type", Target = "~T:HordeServer.Services.RpcService")]
-[assembly: System.Diagnostics.CodeAnalysis.SuppressMessage("Build", "CA1812:Startup is an internal class that is apparently never instantiated. If so, remove the code from the assembly. If this class is intended to contain only static members, make it static (Shared in Visual Basic).", Justification = "<Pending>", Scope = "type", Target = "~T:HordeServer.Startup")]
-[assembly: System.Diagnostics.CodeAnalysis.SuppressMessage("Build", "CA1812:Startup.AgentJwtBearerHandler is an internal class that is apparently never instantiated. If so, remove the code from the assembly. If this class is intended to contain only static members, make it static (Shared in Visual Basic).", Justification = "<Pending>", Scope = "type", Target = "~T:HordeServer.Startup.AgentJwtBearerHandler")]
-[assembly: System.Diagnostics.CodeAnalysis.SuppressMessage("Build", "CA1812:Startup.AnonymousAuthenticationHandler is an internal class that is apparently never instantiated. If so, remove the code from the assembly. If this class is intended to contain only static members, make it static (Shared in Visual Basic).", Justification = "<Pending>", Scope = "type", Target = "~T:HordeServer.Startup.AnonymousAuthenticationHandler")]
-[assembly: System.Diagnostics.CodeAnalysis.SuppressMessage("Performance", "CA1822:Mark members as static", Justification = "<Pending>", Scope = "member", Target = "~M:HordeServer.Controllers.DebugController.Environment~Microsoft.AspNetCore.Mvc.IActionResult")]
-[assembly: System.Diagnostics.CodeAnalysis.SuppressMessage("Naming", "CA1716:Identifiers should not match keywords", Justification = "<Pending>", Scope = "module")]
-[assembly: SuppressMessage("Design", "CA1031:Do not catch general exception types", Justification = "<Pending>", Scope = "member", Target = "~M:HordeServer.Services.AgentService.ExecuteAsync(System.Threading.CancellationToken)~System.Threading.Tasks.Task")]
-[assembly: SuppressMessage("Design", "CA1031:Do not catch general exception types", Justification = "<Pending>", Scope = "member", Target = "~M:HordeServer.Utilities.PeriodicBackgroundService.ExecuteAsync(System.Threading.CancellationToken)~System.Threading.Tasks.Task")]
-[assembly: SuppressMessage("Performance", "CA1815:Override equals and operator equals on value types", Justification = "<Pending>", Scope = "type", Target = "~T:HordeServer.Models.JobStepRefId")]
-[assembly: SuppressMessage("Style", "IDE0054:Use compound assignment", Justification = "<Pending>", Scope = "member", Target = "~M:HordeServer.Models.Job.GetTimingInfo~System.Collections.Generic.Dictionary{HordeServer.Models.Node,HordeServer.Models.TimingInfo}")]
-[assembly: SuppressMessage("Design", "CA1031:Do not catch general exception types", Justification = "<Pending>", Scope = "member", Target = "~M:HordeServer.Startup.DashboardBearerHandler.HandleAuthenticateAsync~System.Threading.Tasks.Task{Microsoft.AspNetCore.Authentication.AuthenticateResult}")]
-[assembly: SuppressMessage("Design", "CA1031:Do not catch general exception types", Justification = "<Pending>", Scope = "member", Target = "~M:HordeServer.Services.ConformService.ExecuteAsync(System.Threading.CancellationToken)~System.Threading.Tasks.Task")]
-[assembly: SuppressMessage("Design", "CA1031:Do not catch general exception types", Justification = "<Pending>", Scope = "member", Target = "~M:HordeServer.Services.ScheduleService.ExecuteAsync(System.Threading.CancellationToken)~System.Threading.Tasks.Task")]
-[assembly: SuppressMessage("Design", "CA1031:Do not catch general exception types", Justification = "<Pending>", Scope = "member", Target = "~M:HordeServer.Services.ScheduleService.UpdateSchedulesAsync~System.Threading.Tasks.Task{System.DateTimeOffset}")]
+[assembly: SuppressMessage("Reliability", "CA2007:Consider calling ConfigureAwait on the awaited task", Justification = "Not necessary in NET Core; no SynchronizationContext.", Scope = "module")]
+[assembly: SuppressMessage("Usage", "CA2227:Collection properties should be read only", Justification = "<Pending>", Scope = "module")]
+[assembly: SuppressMessage("Design", "CA1062:Validate arguments of public methods", Justification = "<Pending>", Scope = "module")]
+[assembly: SuppressMessage("Globalization", "CA1303:Do not pass literals as localized parameters", Justification = "<Pending>", Scope = "module")]
+[assembly: SuppressMessage("Performance", "CA1819:Properties should not return arrays", Justification = "<Pending>", Scope = "module")]
+[assembly: SuppressMessage("Naming", "CA1716:Identifiers should not match keywords", Justification = "<Pending>", Scope = "module")]
 [assembly: SuppressMessage("Design", "CA1030:Use events where appropriate", Justification = "<Pending>", Scope = "module")]
 [assembly: SuppressMessage("Design", "CA1031:Do not catch general exception types", Justification = "<Pending>", Scope = "module")]
 [assembly: SuppressMessage("Design", "CA1032:Implement standard exception constructors", Justification = "<Pending>", Scope = "module")]
diff --git a/Engine/Source/Programs/Horde/Horde.Build/Issues/IIssue.cs b/Engine/Source/Programs/Horde/Horde.Build/Issues/IIssue.cs
--- a/Engine/Source/Programs/Horde/Horde.Build/Issues/IIssue.cs
+++ b/Engine/Source/Programs/Horde/Horde.Build/Issues/IIssue.cs
@@ -1,24 +1,14 @@
 // Copyright Epic Games, Inc. All Rights Reserved.
 
-using EpicGames.Core;
-using Horde.Build.Api;
-using Horde.Build.Collections;
-using HordeCommon;
-using Horde.Build.Services;
-using Horde.Build.Utilities;
-using MongoDB.Bson;
-using MongoDB.Bson.Serialization.Attributes;
-using MongoDB.Bson.Serialization.Serializers;
 using System;
 using System.Collections.Generic;
-using System.Diagnostics.CodeAnalysis;
 using System.Linq;
-using System.Runtime.CompilerServices;
+using Horde.Build.Utilities;
+using MongoDB.Bson;
 
 namespace Horde.Build.Models
 {
 	using StreamId = StringId<IStream>;
-	using TemplateRefId = StringId<TemplateRef>;
 	using UserId = ObjectId<IUser>;
 
 	/// <summary>
@@ -213,43 +203,43 @@ namespace Horde.Build.Models
 		/// <summary>
 		/// Creates a lookup from stream id to whether it's fixed
 		/// </summary>
-		/// <param name="Issue"></param>
+		/// <param name="issue"></param>
 		/// <returns></returns>
-		public static Dictionary<StreamId, bool> GetFixStreamIds(this IIssue Issue)
+		public static Dictionary<StreamId, bool> GetFixStreamIds(this IIssue issue)
 		{
-			Dictionary<StreamId, bool> FixStreamIds = new Dictionary<StreamId, bool>();
-			foreach (IIssueStream Stream in Issue.Streams)
+			Dictionary<StreamId, bool> fixStreamIds = new Dictionary<StreamId, bool>();
+			foreach (IIssueStream stream in issue.Streams)
 			{
-				if (Stream.ContainsFix.HasValue)
+				if (stream.ContainsFix.HasValue)
 				{
-					FixStreamIds[Stream.StreamId] = Stream.ContainsFix.Value;
+					fixStreamIds[stream.StreamId] = stream.ContainsFix.Value;
 				}
 			}
-			return FixStreamIds;
+			return fixStreamIds;
 		}
 
 		/// <summary>
 		/// Find the first fix-failed step for this issue
 		/// </summary>
-		/// <param name="Issue"></param>
-		/// <param name="Spans"></param>
+		/// <param name="issue"></param>
+		/// <param name="spans"></param>
 		/// <returns></returns>
-		public static IIssueStep? FindFixFailedStep(this IIssue Issue, IEnumerable<IIssueSpan> Spans)
+		public static IIssueStep? FindFixFailedStep(this IIssue issue, IEnumerable<IIssueSpan> spans)
 		{
-			IIssueStep? FixFailedStep = null;
-			if (Issue.FixChange != null && Issue.FixChange.Value >= 0)
+			IIssueStep? fixFailedStep = null;
+			if (issue.FixChange != null && issue.FixChange.Value >= 0)
 			{
-				foreach (IIssueSpan Span in Spans)
+				foreach (IIssueSpan span in spans)
 				{
-					IIssueStream? Stream = Issue.Streams.FirstOrDefault(x => x.StreamId == Span.StreamId);
-					if(Stream != null && (Stream.ContainsFix ?? false) && Span.LastFailure.Change >= Issue.FixChange.Value)
+					IIssueStream? stream = issue.Streams.FirstOrDefault(x => x.StreamId == span.StreamId);
+					if(stream != null && (stream.ContainsFix ?? false) && span.LastFailure.Change >= issue.FixChange.Value)
 					{
-						FixFailedStep = Span.LastFailure;
+						fixFailedStep = span.LastFailure;
 						break;
 					}
 				}
 			}
-			return FixFailedStep;
+			return fixFailedStep;
 		}
 	}
 }
\ No newline at end of file
diff --git a/Engine/Source/Programs/Horde/Horde.Build/Issues/IIssueCollection.cs b/Engine/Source/Programs/Horde/Horde.Build/Issues/IIssueCollection.cs
--- a/Engine/Source/Programs/Horde/Horde.Build/Issues/IIssueCollection.cs
+++ b/Engine/Source/Programs/Horde/Horde.Build/Issues/IIssueCollection.cs
@@ -1,24 +1,13 @@
 // Copyright Epic Games, Inc. All Rights Reserved.
 
-using HordeCommon;
-using Horde.Build.Models;
-using Horde.Build.Services;
-using Horde.Build.Utilities;
-using Microsoft.AspNetCore.Connections.Features;
-using Microsoft.Extensions.Logging;
-using MongoDB.Bson;
-using MongoDB.Bson.Serialization;
-using MongoDB.Bson.Serialization.Attributes;
-using MongoDB.Driver;
 using System;
 using System.Collections.Generic;
-using System.Diagnostics.CodeAnalysis;
-using System.IO;
 using System.Linq;
-using System.Linq.Expressions;
-using System.Text.RegularExpressions;
 using System.Threading.Tasks;
-using System.Transactions;
+using Horde.Build.Models;
+using Horde.Build.Utilities;
+using MongoDB.Bson;
+using MongoDB.Driver;
 
 namespace Horde.Build.Collections
 {
@@ -45,71 +34,71 @@ namespace Horde.Build.Collections
 		/// <summary>
 		/// Constructor
 		/// </summary>
-		/// <param name="Type">The type of issue</param>
-		/// <param name="Keys">Keys which uniquely identify this issue</param>
-		/// <param name="RejectKeys">Keys which should not match with this issue</param>
-		public NewIssueFingerprint(string Type, IEnumerable<string> Keys, IEnumerable<string>? RejectKeys)
+		/// <param name="type">The type of issue</param>
+		/// <param name="keys">Keys which uniquely identify this issue</param>
+		/// <param name="rejectKeys">Keys which should not match with this issue</param>
+		public NewIssueFingerprint(string type, IEnumerable<string> keys, IEnumerable<string>? rejectKeys)
 		{
-			this.Type = Type;
-			this.Keys = new CaseInsensitiveStringSet(Keys);
+			Type = type;
+			Keys = new CaseInsensitiveStringSet(keys);
 
-			if (RejectKeys != null && RejectKeys.Any())
+			if (rejectKeys != null && rejectKeys.Any())
 			{
-				this.RejectKeys = new CaseInsensitiveStringSet(RejectKeys);
+				RejectKeys = new CaseInsensitiveStringSet(rejectKeys);
 			}
 		}
 
 		/// <summary>
 		/// Copy constructor
 		/// </summary>
-		/// <param name="Other">The fingerprint to copy from</param>
-		public NewIssueFingerprint(IIssueFingerprint Other)
-			: this(Other.Type, Other.Keys, Other.RejectKeys)
+		/// <param name="other">The fingerprint to copy from</param>
+		public NewIssueFingerprint(IIssueFingerprint other)
+			: this(other.Type, other.Keys, other.RejectKeys)
 		{
 		}
 
 		/// <summary>
 		/// Merges another fingerprint into this one
 		/// </summary>
-		/// <param name="Other">The other fingerprint to merge with</param>
-		public NewIssueFingerprint MergeWith(IIssueFingerprint Other)
+		/// <param name="other">The other fingerprint to merge with</param>
+		public NewIssueFingerprint MergeWith(IIssueFingerprint other)
 		{
-			NewIssueFingerprint Result = new NewIssueFingerprint(this);
+			NewIssueFingerprint result = new NewIssueFingerprint(this);
 
-			Result.Keys.UnionWith(Other.Keys);
-			if (Other.RejectKeys != null)
+			result.Keys.UnionWith(other.Keys);
+			if (other.RejectKeys != null)
 			{
-				if (Result.RejectKeys == null)
+				if (result.RejectKeys == null)
 				{
-					Result.RejectKeys = new CaseInsensitiveStringSet();
+					result.RejectKeys = new CaseInsensitiveStringSet();
 				}
-				Result.RejectKeys.UnionWith(Other.RejectKeys);
+				result.RejectKeys.UnionWith(other.RejectKeys);
 			}
 
-			return Result;
+			return result;
 		}
 
 		/// <inheritdoc/>
-		public override bool Equals(object? Other) => Equals(Other as IIssueFingerprint);
+		public override bool Equals(object? other) => Equals(other as IIssueFingerprint);
 
 		/// <inheritdoc/>
-		public bool Equals(IIssueFingerprint? OtherFingerprint)
+		public bool Equals(IIssueFingerprint? otherFingerprint)
 		{
-			if(OtherFingerprint == null || !Type.Equals(OtherFingerprint.Type, StringComparison.Ordinal) || !Keys.SetEquals(OtherFingerprint.Keys))
+			if(otherFingerprint == null || !Type.Equals(otherFingerprint.Type, StringComparison.Ordinal) || !Keys.SetEquals(otherFingerprint.Keys))
 			{
 				return false;
 			}
 
 			if (RejectKeys == null)
 			{
-				if (OtherFingerprint.RejectKeys != null && OtherFingerprint.RejectKeys.Count > 0)
+				if (otherFingerprint.RejectKeys != null && otherFingerprint.RejectKeys.Count > 0)
 				{
 					return false;
 				}
 			}
 			else
 			{
-				if (OtherFingerprint.RejectKeys == null || !RejectKeys.SetEquals(OtherFingerprint.RejectKeys))
+				if (otherFingerprint.RejectKeys == null || !RejectKeys.SetEquals(otherFingerprint.RejectKeys))
 				{
 					return false;
 				}
@@ -124,42 +113,42 @@ namespace Horde.Build.Collections
 		/// <returns>The hashcode value</returns>
 		public override int GetHashCode()
 		{
-			int Result = StringComparer.Ordinal.GetHashCode(Type);
-			foreach (string Key in Keys)
+			int result = StringComparer.Ordinal.GetHashCode(Type);
+			foreach (string key in Keys)
 			{
-				Result = HashCode.Combine(Result, Key);
+				result = HashCode.Combine(result, key);
 			}
 			if (RejectKeys != null)
 			{
-				foreach (string RejectKey in RejectKeys)
+				foreach (string rejectKey in RejectKeys)
 				{
-					Result = HashCode.Combine(Result, RejectKey);
+					result = HashCode.Combine(result, rejectKey);
 				}
 			}
-			return Result;
+			return result;
 		}
 
 		/// <summary>
 		/// Merges two fingerprints togetherthis issue fingerprint with another
 		/// </summary>
-		/// <param name="A">The first fingerprint</param>
-		/// <param name="B">The second fingerprint to merge with</param>
+		/// <param name="a">The first fingerprint</param>
+		/// <param name="b">The second fingerprint to merge with</param>
 		/// <returns>Merged fingerprint</returns>
-		public static NewIssueFingerprint Merge(IIssueFingerprint A, IIssueFingerprint B)
+		public static NewIssueFingerprint Merge(IIssueFingerprint a, IIssueFingerprint b)
 		{
-			NewIssueFingerprint NewFingerprint = new NewIssueFingerprint(A);
-			NewFingerprint.Keys.UnionWith(B.Keys);
+			NewIssueFingerprint newFingerprint = new NewIssueFingerprint(a);
+			newFingerprint.Keys.UnionWith(b.Keys);
 
-			if (B.RejectKeys != null)
+			if (b.RejectKeys != null)
 			{
-				if (NewFingerprint.RejectKeys == null)
+				if (newFingerprint.RejectKeys == null)
 				{
-					NewFingerprint.RejectKeys = new CaseInsensitiveStringSet();
+					newFingerprint.RejectKeys = new CaseInsensitiveStringSet();
 				}
-				NewFingerprint.RejectKeys.UnionWith(B.RejectKeys);
+				newFingerprint.RejectKeys.UnionWith(b.RejectKeys);
 			}
 
-			return NewFingerprint;
+			return newFingerprint;
 		}
 
 		/// <inheritdoc/>
@@ -204,14 +193,14 @@ namespace Horde.Build.Collections
 		/// <summary>
 		/// Constructor
 		/// </summary>
-		public NewIssueSpanData(StreamId StreamId, string StreamName, TemplateRefId TemplateRefId, string NodeName, NewIssueFingerprint Fingerprint, NewIssueStepData FirstFailure)
+		public NewIssueSpanData(StreamId streamId, string streamName, TemplateRefId templateRefId, string nodeName, NewIssueFingerprint fingerprint, NewIssueStepData firstFailure)
 		{
-			this.StreamId = StreamId;
-			this.StreamName = StreamName;
-			this.TemplateRefId = TemplateRefId;
-			this.NodeName = NodeName;
-			this.Fingerprint = Fingerprint;
-			this.FirstFailure = FirstFailure;
+			StreamId = streamId;
+			StreamName = streamName;
+			TemplateRefId = templateRefId;
+			NodeName = nodeName;
+			Fingerprint = fingerprint;
+			FirstFailure = firstFailure;
 		}
 	}
 
@@ -250,47 +239,47 @@ namespace Horde.Build.Collections
 		/// <summary>
 		/// Constructor
 		/// </summary>
-		/// <param name="Change">The changelist number for this job</param>
-		/// <param name="Severity">Severity of the issue in this step</param>
-		/// <param name="JobName">The job name</param>
-		/// <param name="JobId">The unique job id</param>
-		/// <param name="BatchId">The batch id</param>
-		/// <param name="StepId">The step id</param>
-		/// <param name="StepTime">Time that the step started</param>
-		/// <param name="LogId">Unique id of the log file for this step</param>
-		/// <param name="Promoted">Whether this step is promoted</param>
-		public NewIssueStepData(int Change, IssueSeverity Severity, string JobName, JobId JobId, SubResourceId BatchId, SubResourceId StepId, DateTime StepTime, LogId? LogId, bool Promoted)
+		/// <param name="change">The changelist number for this job</param>
+		/// <param name="severity">Severity of the issue in this step</param>
+		/// <param name="jobName">The job name</param>
+		/// <param name="jobId">The unique job id</param>
+		/// <param name="batchId">The batch id</param>
+		/// <param name="stepId">The step id</param>
+		/// <param name="stepTime">Time that the step started</param>
+		/// <param name="logId">Unique id of the log file for this step</param>
+		/// <param name="promoted">Whether this step is promoted</param>
+		public NewIssueStepData(int change, IssueSeverity severity, string jobName, JobId jobId, SubResourceId batchId, SubResourceId stepId, DateTime stepTime, LogId? logId, bool promoted)
 		{
-			this.Change = Change;
-			this.Severity = Severity;
-			this.JobName = JobName;
-			this.JobId = JobId;
-			this.BatchId = BatchId;
-			this.StepId = StepId;
-			this.StepTime = StepTime;
-			this.LogId = LogId;
-			this.PromoteByDefault = Promoted;
+			Change = change;
+			Severity = severity;
+			JobName = jobName;
+			JobId = jobId;
+			BatchId = batchId;
+			StepId = stepId;
+			StepTime = stepTime;
+			LogId = logId;
+			PromoteByDefault = promoted;
 		}
 
 		/// <summary>
 		/// Constructor
 		/// </summary>
-		/// <param name="Job">The job being built</param>
-		/// <param name="Batch">Batch of the job for the step</param>
-		/// <param name="Step">The step being built</param>
-		/// <param name="Severity">Severity of the issue in this step</param>
-		/// <param name="Promoted">Whether this step is promoted</param>
-		public NewIssueStepData(IJob Job, IJobStepBatch Batch, IJobStep Step, IssueSeverity Severity, bool Promoted)
-			: this(Job.Change, Severity, Job.Name, Job.Id, Batch.Id, Step.Id, Step.StartTimeUtc ?? default, Step.LogId, Promoted)
+		/// <param name="job">The job being built</param>
+		/// <param name="batch">Batch of the job for the step</param>
+		/// <param name="step">The step being built</param>
+		/// <param name="severity">Severity of the issue in this step</param>
+		/// <param name="promoted">Whether this step is promoted</param>
+		public NewIssueStepData(IJob job, IJobStepBatch batch, IJobStep step, IssueSeverity severity, bool promoted)
+			: this(job.Change, severity, job.Name, job.Id, batch.Id, step.Id, step.StartTimeUtc ?? default, step.LogId, promoted)
 		{
 		}
 
 		/// <summary>
 		/// Constructor for sentinel steps
 		/// </summary>
-		/// <param name="JobStepRef">The jobstep reference</param>
-		public NewIssueStepData(IJobStepRef JobStepRef)
-			: this(JobStepRef.Change, IssueSeverity.Unspecified, JobStepRef.JobName, JobStepRef.Id.JobId, JobStepRef.Id.BatchId, JobStepRef.Id.StepId, JobStepRef.StartTimeUtc, JobStepRef.LogId, false)
+		/// <param name="jobStepRef">The jobstep reference</param>
+		public NewIssueStepData(IJobStepRef jobStepRef)
+			: this(jobStepRef.Change, IssueSeverity.Unspecified, jobStepRef.JobName, jobStepRef.Id.JobId, jobStepRef.Id.BatchId, jobStepRef.Id.StepId, jobStepRef.StartTimeUtc, jobStepRef.LogId, false)
 		{
 		}
 	}
@@ -309,12 +298,12 @@ namespace Horde.Build.Collections
 		/// <summary>
 		/// Constructor
 		/// </summary>
-		/// <param name="AuthorId">Author of the change</param>
-		/// <param name="Change">The changelist number</param>
-		public NewIssueSuspectData(UserId AuthorId, int Change)
+		/// <param name="authorId">Author of the change</param>
+		/// <param name="change">The changelist number</param>
+		public NewIssueSuspectData(UserId authorId, int change)
 		{
-			this.AuthorId = AuthorId;
-			this.Change = Change;
+			AuthorId = authorId;
+			Change = change;
 		}
 	}
 
@@ -335,12 +324,12 @@ namespace Horde.Build.Collections
 		/// <summary>
 		/// Constructor
 		/// </summary>
-		/// <param name="Change">The changelist number</param>
-		/// <param name="AuthorId">Author of the change</param>
-		public NewIssueSpanSuspectData(int Change, UserId AuthorId)
+		/// <param name="change">The changelist number</param>
+		/// <param name="authorId">Author of the change</param>
+		public NewIssueSpanSuspectData(int change, UserId authorId)
 		{
-			this.Change = Change;
-			this.AuthorId = AuthorId;
+			Change = change;
+			AuthorId = authorId;
 		}
 	}
 
@@ -358,10 +347,10 @@ namespace Horde.Build.Collections
 		/// <summary>
 		/// Constructor
 		/// </summary>
-		/// <param name="StreamId"></param>
-		public NewIssueStream(StreamId StreamId)
+		/// <param name="streamId"></param>
+		public NewIssueStream(StreamId streamId)
 		{
-			this.StreamId = StreamId;
+			StreamId = streamId;
 		}
 	}
 
@@ -381,80 +370,80 @@ namespace Horde.Build.Collections
 		/// <summary>
 		/// Creates a new issue
 		/// </summary>
-		/// <param name="Summary">Summary text for the issue</param>
+		/// <param name="summary">Summary text for the issue</param>
 		/// <returns>The new issue instance</returns>
-		Task<IIssue> AddIssueAsync(string Summary);
+		Task<IIssue> AddIssueAsync(string summary);
 
 		/// <summary>
 		/// Retrieves and issue by id
 		/// </summary>
-		/// <param name="IssueId">Unique id of the issue</param>
+		/// <param name="issueId">Unique id of the issue</param>
 		/// <returns>The issue matching the given id, or null</returns>
-		Task<IIssue?> GetIssueAsync(int IssueId);
+		Task<IIssue?> GetIssueAsync(int issueId);
 
 		/// <summary>
 		/// Finds the suspects for an issue
 		/// </summary>
-		/// <param name="Issue">The issue to retrieve suspects for</param>
+		/// <param name="issue">The issue to retrieve suspects for</param>
 		/// <returns>List of suspects</returns>
-		Task<List<IIssueSuspect>> FindSuspectsAsync(IIssue Issue);
+		Task<List<IIssueSuspect>> FindSuspectsAsync(IIssue issue);
 
 		/// <summary>
 		/// Searches for open issues
 		/// </summary>
-		/// <param name="Ids">Set of issue ids to find</param>
-		/// <param name="UserId">The user to find issues for</param>
-		/// <param name="StreamId">The stream affected by the issue</param>
-		/// <param name="MinChange">Minimum changelist affected by the issue</param>
-		/// <param name="MaxChange">Maximum changelist affected by the issue</param>
-		/// <param name="Resolved">Include issues that are now resolved</param>
-		/// <param name="Promoted">Include only promoted issues</param>
-		/// <param name="Index">Index within the results to return</param>
-		/// <param name="Count">Number of results</param>
+		/// <param name="ids">Set of issue ids to find</param>
+		/// <param name="userId">The user to find issues for</param>
+		/// <param name="streamId">The stream affected by the issue</param>
+		/// <param name="minChange">Minimum changelist affected by the issue</param>
+		/// <param name="maxChange">Maximum changelist affected by the issue</param>
+		/// <param name="resolved">Include issues that are now resolved</param>
+		/// <param name="promoted">Include only promoted issues</param>
+		/// <param name="index">Index within the results to return</param>
+		/// <param name="count">Number of results</param>
 		/// <returns>List of streams open in the given stream at the given changelist</returns>
-		Task<List<IIssue>> FindIssuesAsync(IEnumerable<int>? Ids = null, UserId? UserId = null, StreamId? StreamId = null, int? MinChange = null, int? MaxChange = null, bool? Resolved = null, bool? Promoted = null, int? Index = null, int? Count = null);
+		Task<List<IIssue>> FindIssuesAsync(IEnumerable<int>? ids = null, UserId? userId = null, StreamId? streamId = null, int? minChange = null, int? maxChange = null, bool? resolved = null, bool? promoted = null, int? index = null, int? count = null);
 
 		/// <summary>
 		/// Searches for open issues
 		/// </summary>
-		/// <param name="Changes">List of suspect changes</param>
+		/// <param name="changes">List of suspect changes</param>
 		/// <returns>List of issues that are affected by the given changes</returns>
-		Task<List<IIssue>> FindIssuesForChangesAsync(List<int> Changes);
+		Task<List<IIssue>> FindIssuesForChangesAsync(List<int> changes);
 
 		/// <summary>
 		/// Try to update the state of an issue
 		/// </summary>
-		/// <param name="Issue">The issue to update</param>
-		/// <param name="NewSeverity">New severity for the issue</param>
-		/// <param name="NewSummary">New summary for the issue</param>
-		/// <param name="NewUserSummary">New user summary for the issue</param>
-		/// <param name="NewDescription">New description for the issue</param>
-		/// <param name="NewPromoted">New promoted state of the issue</param>
-		/// <param name="NewOwnerId">New owner of the issue</param>
-		/// <param name="NewNominatedById">Person that nominated the new owner</param>
-		/// <param name="NewAcknowledged">Whether the issue has been acknowledged</param>
-		/// <param name="NewDeclinedById">Name of a user that has declined the issue</param>
-		/// <param name="NewFixChange">Fix changelist for the issue. Pass 0 to clear the fix changelist, -1 for systemic issue.</param>
-		/// <param name="NewResolvedById">User that resolved the issue (may be ObjectId.Empty to clear)</param>
-		/// <param name="NewExcludeSpanIds">List of span ids to exclude from this issue</param>
-		/// <param name="NewLastSeenAt"></param>
+		/// <param name="issue">The issue to update</param>
+		/// <param name="newSeverity">New severity for the issue</param>
+		/// <param name="newSummary">New summary for the issue</param>
+		/// <param name="newUserSummary">New user summary for the issue</param>
+		/// <param name="newDescription">New description for the issue</param>
+		/// <param name="newPromoted">New promoted state of the issue</param>
+		/// <param name="newOwnerId">New owner of the issue</param>
+		/// <param name="newNominatedById">Person that nominated the new owner</param>
+		/// <param name="newAcknowledged">Whether the issue has been acknowledged</param>
+		/// <param name="newDeclinedById">Name of a user that has declined the issue</param>
+		/// <param name="newFixChange">Fix changelist for the issue. Pass 0 to clear the fix changelist, -1 for systemic issue.</param>
+		/// <param name="newResolvedById">User that resolved the issue (may be ObjectId.Empty to clear)</param>
+		/// <param name="newExcludeSpanIds">List of span ids to exclude from this issue</param>
+		/// <param name="newLastSeenAt"></param>
 		/// <returns>True if the issue was updated</returns>
-		Task<IIssue?> TryUpdateIssueAsync(IIssue Issue, IssueSeverity? NewSeverity = null, string? NewSummary = null, string? NewUserSummary = null, string? NewDescription = null, bool? NewPromoted = null, UserId? NewOwnerId = null, UserId? NewNominatedById = null, bool? NewAcknowledged = null, UserId? NewDeclinedById = null, int? NewFixChange = null, UserId? NewResolvedById = null, List<ObjectId>? NewExcludeSpanIds = null, DateTime? NewLastSeenAt = null);
+		Task<IIssue?> TryUpdateIssueAsync(IIssue issue, IssueSeverity? newSeverity = null, string? newSummary = null, string? newUserSummary = null, string? newDescription = null, bool? newPromoted = null, UserId? newOwnerId = null, UserId? newNominatedById = null, bool? newAcknowledged = null, UserId? newDeclinedById = null, int? newFixChange = null, UserId? newResolvedById = null, List<ObjectId>? newExcludeSpanIds = null, DateTime? newLastSeenAt = null);
 
 		/// <summary>
 		/// Updates derived data for an issue (ie. data computed from the spans attached to it). Also clears the issue's 'modified' state.
 		/// </summary>
-		/// <param name="Issue">Issue to update</param>
-		/// <param name="NewSummary">New summary for the issue</param>
-		/// <param name="NewSeverity">New severity for the issue</param>
-		/// <param name="NewFingerprints">New fingerprints for the issue</param>
-		/// <param name="NewStreams">New streams for the issue</param>
-		/// <param name="NewSuspects">New suspects for the issue</param>
-		/// <param name="NewResolvedAt">Time for the last resolved change</param>
-		/// <param name="NewVerifiedAt">Time that the issue was resolved</param>
-		/// <param name="NewLastSeenAt">Last time the issue was seen</param>
+		/// <param name="issue">Issue to update</param>
+		/// <param name="newSummary">New summary for the issue</param>
+		/// <param name="newSeverity">New severity for the issue</param>
+		/// <param name="newFingerprints">New fingerprints for the issue</param>
+		/// <param name="newStreams">New streams for the issue</param>
+		/// <param name="newSuspects">New suspects for the issue</param>
+		/// <param name="newResolvedAt">Time for the last resolved change</param>
+		/// <param name="newVerifiedAt">Time that the issue was resolved</param>
+		/// <param name="newLastSeenAt">Last time the issue was seen</param>
 		/// <returns>Updated issue, or null if the issue is modified in the interim</returns>
-		Task<IIssue?> TryUpdateIssueDerivedDataAsync(IIssue Issue, string NewSummary, IssueSeverity NewSeverity, List<NewIssueFingerprint> NewFingerprints, List<NewIssueStream> NewStreams, List<NewIssueSuspectData> NewSuspects, DateTime? NewResolvedAt, DateTime? NewVerifiedAt, DateTime NewLastSeenAt);
+		Task<IIssue?> TryUpdateIssueDerivedDataAsync(IIssue issue, string newSummary, IssueSeverity newSeverity, List<NewIssueFingerprint> newFingerprints, List<NewIssueStream> newStreams, List<NewIssueSuspectData> newSuspects, DateTime? newResolvedAt, DateTime? newVerifiedAt, DateTime newLastSeenAt);
 
 		#endregion
 
@@ -463,67 +452,67 @@ namespace Horde.Build.Collections
 		/// <summary>
 		/// Creates a new span from the given failure
 		/// </summary>
-		/// <param name="IssueId">The issue that the span belongs to</param>
-		/// <param name="NewSpan">Information about the new span</param>
+		/// <param name="issueId">The issue that the span belongs to</param>
+		/// <param name="newSpan">Information about the new span</param>
 		/// <returns>New span, or null if the sequence token is not valid</returns>
-		Task<IIssueSpan> AddSpanAsync(int IssueId, NewIssueSpanData NewSpan);
+		Task<IIssueSpan> AddSpanAsync(int issueId, NewIssueSpanData newSpan);
 
 		/// <summary>
 		/// Gets a particular span
 		/// </summary>
-		/// <param name="SpanId">Unique id of the span</param>
+		/// <param name="spanId">Unique id of the span</param>
 		/// <returns>New span, or null if the sequence token is not valid</returns>
-		Task<IIssueSpan?> GetSpanAsync(ObjectId SpanId);
+		Task<IIssueSpan?> GetSpanAsync(ObjectId spanId);
 
 		/// <summary>
 		/// Updates the given span. Note that data in the span's issue may be derived from this, and the issue should be updated afterwards.
 		/// </summary>
-		/// <param name="Span">Span to update</param>
-		/// <param name="NewLastSuccess">New last successful step</param>
-		/// <param name="NewFailure">New failed step</param>
-		/// <param name="NewNextSuccess">New next successful step</param>
-		/// <param name="NewSuspects">New suspects for the span</param>
-		/// <param name="NewIssueId">The new issue id for this span</param>
+		/// <param name="span">Span to update</param>
+		/// <param name="newLastSuccess">New last successful step</param>
+		/// <param name="newFailure">New failed step</param>
+		/// <param name="newNextSuccess">New next successful step</param>
+		/// <param name="newSuspects">New suspects for the span</param>
+		/// <param name="newIssueId">The new issue id for this span</param>
 		/// <returns>The updated span, or null on failure</returns>
-		Task<IIssueSpan?> TryUpdateSpanAsync(IIssueSpan Span, NewIssueStepData? NewLastSuccess = null, NewIssueStepData? NewFailure = null, NewIssueStepData? NewNextSuccess = null, List<NewIssueSpanSuspectData>? NewSuspects = null, int? NewIssueId = null);
+		Task<IIssueSpan?> TryUpdateSpanAsync(IIssueSpan span, NewIssueStepData? newLastSuccess = null, NewIssueStepData? newFailure = null, NewIssueStepData? newNextSuccess = null, List<NewIssueSpanSuspectData>? newSuspects = null, int? newIssueId = null);
 
 		/// <summary>
 		/// Gets all the spans for a particular issue
 		/// </summary>
-		/// <param name="IssueId">Issue id</param>
+		/// <param name="issueId">Issue id</param>
 		/// <returns>List of spans</returns>
-		Task<List<IIssueSpan>> FindSpansAsync(int IssueId);
+		Task<List<IIssueSpan>> FindSpansAsync(int issueId);
 
 		/// <summary>
 		/// Retrieves multiple spans
 		/// </summary>
-		/// <param name="SpanIds">The span ids</param>
+		/// <param name="spanIds">The span ids</param>
 		/// <returns>List of spans</returns>
-		Task<List<IIssueSpan>> FindSpansAsync(IEnumerable<ObjectId> SpanIds);
+		Task<List<IIssueSpan>> FindSpansAsync(IEnumerable<ObjectId> spanIds);
 
 		/// <summary>
 		/// Finds the open issues for a given stream
 		/// </summary>
-		/// <param name="StreamId">The stream id</param>
-		/// <param name="TemplateId">The template id</param>
-		/// <param name="Name">Name of the node</param>
-		/// <param name="Change">Changelist number to query</param>
+		/// <param name="streamId">The stream id</param>
+		/// <param name="templateId">The template id</param>
+		/// <param name="name">Name of the node</param>
+		/// <param name="change">Changelist number to query</param>
 		/// <returns>List of open issues</returns>
-		Task<List<IIssueSpan>> FindOpenSpansAsync(StreamId StreamId, TemplateRefId TemplateId, string Name, int Change);
+		Task<List<IIssueSpan>> FindOpenSpansAsync(StreamId streamId, TemplateRefId templateId, string name, int change);
 
 		/// <summary>
 		/// Searches for open issues
 		/// </summary>
-		/// <param name="Ids">Set of issue ids to find</param>
-		/// <param name="IssueIds">The issue ids to retrieve spans for</param>
-		/// <param name="StreamId">The stream affected by the issue</param>
-		/// <param name="MinChange">Minimum changelist affected by the issue</param>
-		/// <param name="MaxChange">Maximum changelist affected by the issue</param>
-		/// <param name="Resolved">Include issues that are now resolved</param>
-		/// <param name="Index">Index within the results to return</param>
-		/// <param name="Count">Number of results</param>
+		/// <param name="ids">Set of issue ids to find</param>
+		/// <param name="issueIds">The issue ids to retrieve spans for</param>
+		/// <param name="streamId">The stream affected by the issue</param>
+		/// <param name="minChange">Minimum changelist affected by the issue</param>
+		/// <param name="maxChange">Maximum changelist affected by the issue</param>
+		/// <param name="resolved">Include issues that are now resolved</param>
+		/// <param name="index">Index within the results to return</param>
+		/// <param name="count">Number of results</param>
 		/// <returns>List of streams open in the given stream at the given changelist</returns>
-		Task<List<IIssueSpan>> FindSpansAsync(IEnumerable<ObjectId>? Ids = null, IEnumerable<int>? IssueIds = null, StreamId? StreamId = null, int? MinChange = null, int? MaxChange = null, bool? Resolved = null, int? Index = null, int? Count = null);
+		Task<List<IIssueSpan>> FindSpansAsync(IEnumerable<ObjectId>? ids = null, IEnumerable<int>? issueIds = null, StreamId? streamId = null, int? minChange = null, int? maxChange = null, bool? resolved = null, int? index = null, int? count = null);
 
 		#endregion
 
@@ -532,34 +521,34 @@ namespace Horde.Build.Collections
 		/// <summary>
 		/// Adds a new step
 		/// </summary>
-		/// <param name="SpanId">Initial span for the step</param>
-		/// <param name="NewStep">Information about the new step</param>
+		/// <param name="spanId">Initial span for the step</param>
+		/// <param name="newStep">Information about the new step</param>
 		/// <returns>New step object</returns>
-		Task<IIssueStep> AddStepAsync(ObjectId SpanId, NewIssueStepData NewStep);
+		Task<IIssueStep> AddStepAsync(ObjectId spanId, NewIssueStepData newStep);
 
 		/// <summary>
 		/// Find steps for the given spans
 		/// </summary>
-		/// <param name="SpanIds">Span ids</param>
+		/// <param name="spanIds">Span ids</param>
 		/// <returns>List of steps</returns>
-		Task<List<IIssueStep>> FindStepsAsync(IEnumerable<ObjectId> SpanIds);
+		Task<List<IIssueStep>> FindStepsAsync(IEnumerable<ObjectId> spanIds);
 
 		/// <summary>
 		/// Find steps for the given spans
 		/// </summary>
-		/// <param name="JobId">The job id</param>
-		/// <param name="BatchId">The batch id</param>
-		/// <param name="StepId">The step id</param>
+		/// <param name="jobId">The job id</param>
+		/// <param name="batchId">The batch id</param>
+		/// <param name="stepId">The step id</param>
 		/// <returns>List of steps</returns>
-		Task<List<IIssueStep>> FindStepsAsync(JobId JobId, SubResourceId? BatchId, SubResourceId? StepId);
+		Task<List<IIssueStep>> FindStepsAsync(JobId jobId, SubResourceId? batchId, SubResourceId? stepId);
 
 		#endregion
 
 		/// <summary>
 		/// Gets the logger for a particular issue
 		/// </summary>
-		/// <param name="IssueId">The issue id</param>
+		/// <param name="issueId">The issue id</param>
 		/// <returns>Logger for this issue</returns>
-		IAuditLogChannel<int> GetLogger(int IssueId);
+		IAuditLogChannel<int> GetLogger(int issueId);
 	}
 }
diff --git a/Engine/Source/Programs/Horde/Horde.Build/Issues/IIssueDetails.cs b/Engine/Source/Programs/Horde/Horde.Build/Issues/IIssueDetails.cs
--- a/Engine/Source/Programs/Horde/Horde.Build/Issues/IIssueDetails.cs
+++ b/Engine/Source/Programs/Horde/Horde.Build/Issues/IIssueDetails.cs
@@ -1,12 +1,9 @@
 // Copyright Epic Games, Inc. All Rights Reserved.
 
-using Horde.Build.Services;
-using Horde.Build.Utilities;
-using MongoDB.Bson;
-using System;
 using System.Collections.Generic;
-using System.Linq;
 using System.Threading.Tasks;
+using Horde.Build.Services;
+using Horde.Build.Utilities;
 
 namespace Horde.Build.Models
 {
@@ -66,9 +63,9 @@ namespace Horde.Build.Models
 		/// <summary>
 		/// Determines if the issue is relevant to the given user
 		/// </summary>
-		/// <param name="UserId">The user to query</param>
+		/// <param name="userId">The user to query</param>
 		/// <returns>True if the issue is relevant to the given user</returns>
-		bool IncludeForUser(UserId UserId);
+		bool IncludeForUser(UserId userId);
 	}
 
 	/// <summary>
@@ -79,17 +76,17 @@ namespace Horde.Build.Models
 		/// <summary>
 		/// Gets an issue details object for a specific issue id
 		/// </summary>
-		/// <param name="IssueService">The issue service</param>
-		/// <param name="IssueId">Issue id to query </param>
+		/// <param name="issueService">The issue service</param>
+		/// <param name="issueId">Issue id to query </param>
 		/// <returns></returns>
-		public static async Task<IIssueDetails?> GetIssueDetailsAsync(this IIssueService IssueService, int IssueId)
+		public static async Task<IIssueDetails?> GetIssueDetailsAsync(this IIssueService issueService, int issueId)
 		{
-			IIssue? Issue = await IssueService.GetIssueAsync(IssueId);
-			if(Issue == null)
+			IIssue? issue = await issueService.GetIssueAsync(issueId);
+			if(issue == null)
 			{
 				return null;
 			}
-			return await IssueService.GetIssueDetailsAsync(Issue);
+			return await issueService.GetIssueDetailsAsync(issue);
 		}
 	}
 }
diff --git a/Engine/Source/Programs/Horde/Horde.Build/Issues/IIssueFingerprint.cs b/Engine/Source/Programs/Horde/Horde.Build/Issues/IIssueFingerprint.cs
--- a/Engine/Source/Programs/Horde/Horde.Build/Issues/IIssueFingerprint.cs
+++ b/Engine/Source/Programs/Horde/Horde.Build/Issues/IIssueFingerprint.cs
@@ -1,10 +1,8 @@
 // Copyright Epic Games, Inc. All Rights Reserved.
 
-using Horde.Build.Utilities;
 using System;
-using System.Collections.Generic;
 using System.Linq;
-using System.Threading.Tasks;
+using Horde.Build.Utilities;
 
 namespace Horde.Build.Models
 {
@@ -37,20 +35,20 @@ namespace Horde.Build.Models
 		/// <summary>
 		/// Checks if a fingerprint matches another fingerprint
 		/// </summary>
-		/// <param name="Fingerprint">The first fingerprint to compare</param>
-		/// <param name="Other">The other fingerprint to compare to</param>
+		/// <param name="fingerprint">The first fingerprint to compare</param>
+		/// <param name="other">The other fingerprint to compare to</param>
 		/// <returns>True is the fingerprints match</returns>
-		public static bool IsMatch(this IIssueFingerprint Fingerprint, IIssueFingerprint Other)
+		public static bool IsMatch(this IIssueFingerprint fingerprint, IIssueFingerprint other)
 		{
-			if (!Fingerprint.Type.Equals(Other.Type, StringComparison.Ordinal))
+			if (!fingerprint.Type.Equals(other.Type, StringComparison.Ordinal))
 			{
 				return false;
 			}
-			if (Fingerprint.Keys.Count > 0 && Other.Keys.Count > 0 && !Fingerprint.Keys.Any(x => Other.Keys.Contains(x)))
+			if (fingerprint.Keys.Count > 0 && other.Keys.Count > 0 && !fingerprint.Keys.Any(x => other.Keys.Contains(x)))
 			{
 				return false;
 			}
-			if (Fingerprint.RejectKeys != null && Fingerprint.RejectKeys.Any(x => Other.Keys.Contains(x)))
+			if (fingerprint.RejectKeys != null && fingerprint.RejectKeys.Any(x => other.Keys.Contains(x)))
 			{
 				return false;
 			}
@@ -60,16 +58,16 @@ namespace Horde.Build.Models
 		/// <summary>
 		/// Checks if a fingerprint matches another fingerprint for creating a new span
 		/// </summary>
-		/// <param name="Fingerprint">The first fingerprint to compare</param>
-		/// <param name="Other">The other fingerprint to compare to</param>
+		/// <param name="fingerprint">The first fingerprint to compare</param>
+		/// <param name="other">The other fingerprint to compare to</param>
 		/// <returns>True is the fingerprints match</returns>
-		public static bool IsMatchForNewSpan(this IIssueFingerprint Fingerprint, IIssueFingerprint Other)
+		public static bool IsMatchForNewSpan(this IIssueFingerprint fingerprint, IIssueFingerprint other)
 		{
-			if (!Fingerprint.Type.Equals(Other.Type, StringComparison.Ordinal))
+			if (!fingerprint.Type.Equals(other.Type, StringComparison.Ordinal))
 			{
 				return false;
 			}
-			if (Fingerprint.RejectKeys != null && Fingerprint.RejectKeys.Any(x => Other.Keys.Contains(x)))
+			if (fingerprint.RejectKeys != null && fingerprint.RejectKeys.Any(x => other.Keys.Contains(x)))
 			{
 				return false;
 			}
diff --git a/Engine/Source/Programs/Horde/Horde.Build/Issues/IIssueHandler.cs b/Engine/Source/Programs/Horde/Horde.Build/Issues/IIssueHandler.cs
--- a/Engine/Source/Programs/Horde/Horde.Build/Issues/IIssueHandler.cs
+++ b/Engine/Source/Programs/Horde/Horde.Build/Issues/IIssueHandler.cs
@@ -1,13 +1,11 @@
 // Copyright Epic Games, Inc. All Rights Reserved.
 
-using EpicGames.Core;
-using Horde.Build.Models;
-using Horde.Build.Collections;
 using System;
 using System.Collections.Generic;
-using System.Linq;
-using System.Threading.Tasks;
 using System.Diagnostics.CodeAnalysis;
+using EpicGames.Core;
+using Horde.Build.Collections;
+using Horde.Build.Models;
 
 namespace Horde.Build.IssueHandlers
 {
@@ -19,7 +17,7 @@ namespace Horde.Build.IssueHandlers
 		/// <summary>
 		/// Set of extensions to treat as code
 		/// </summary>
-		static readonly HashSet<StringView> CodeExtensions = new HashSet<StringView>(StringViewComparer.OrdinalIgnoreCase)
+		static readonly HashSet<StringView> s_codeExtensions = new HashSet<StringView>(StringViewComparer.OrdinalIgnoreCase)
 		{
 			".c",
 			".cc",
@@ -42,7 +40,7 @@ namespace Horde.Build.IssueHandlers
 		/// <summary>
 		/// Set of file extensions to treat as content
 		/// </summary>
-		static readonly HashSet<StringView> ContentExtensions = new HashSet<StringView>(StringViewComparer.OrdinalIgnoreCase)
+		static readonly HashSet<StringView> s_contentExtensions = new HashSet<StringView>(StringViewComparer.OrdinalIgnoreCase)
 		{
 			".uasset",
 			".umap",
@@ -73,22 +71,22 @@ namespace Horde.Build.IssueHandlers
 		/// <summary>
 		/// Constructor
 		/// </summary>
-		/// <param name="Details">The changelist details</param>
-		public SuspectChange(ChangeDetails Details)
+		/// <param name="details">The changelist details</param>
+		public SuspectChange(ChangeDetails details)
 		{
-			this.Details = Details;
+			Details = details;
 
-			foreach (ChangeFile File in Details.Files)
+			foreach (ChangeFile file in details.Files)
 			{
-				int Idx = File.Path.LastIndexOf('.');
-				if (Idx != -1)
+				int idx = file.Path.LastIndexOf('.');
+				if (idx != -1)
 				{
-					StringView Extension = new StringView(File.Path, Idx);
-					if (CodeExtensions.Contains(Extension))
+					StringView extension = new StringView(file.Path, idx);
+					if (s_codeExtensions.Contains(extension))
 					{
 						ContainsCode = true;
 					}
-					if (ContentExtensions.Contains(Extension))
+					if (s_contentExtensions.Contains(extension))
 					{
 						ContainsContent = true;
 					}
@@ -103,13 +101,13 @@ namespace Horde.Build.IssueHandlers
 		/// <summary>
 		/// Determines whether this change modifies the given file
 		/// </summary>
-		/// <param name="FileToCheck">The file to look for</param>
+		/// <param name="fileToCheck">The file to look for</param>
 		/// <returns>True if the change modifies the given file</returns>
-		public bool ModifiesFile(string FileToCheck)
+		public bool ModifiesFile(string fileToCheck)
 		{
-			foreach (ChangeFile File in Details.Files)
+			foreach (ChangeFile file in Details.Files)
 			{
-				if (File.Path.EndsWith(FileToCheck, StringComparison.OrdinalIgnoreCase) && (File.Length == FileToCheck.Length || File.Path[File.Path.Length - FileToCheck.Length - 1] == '/'))
+				if (file.Path.EndsWith(fileToCheck, StringComparison.OrdinalIgnoreCase) && (file.Length == fileToCheck.Length || file.Path[file.Path.Length - fileToCheck.Length - 1] == '/'))
 				{
 					return true;
 				}
@@ -136,26 +134,26 @@ namespace Horde.Build.IssueHandlers
 		/// <summary>
 		/// Match the given event and produce a fingerprint
 		/// </summary>
-		/// <param name="Job">The job that spawned the event</param>
-		/// <param name="Node">Node that was executed</param>
-		/// <param name="EventData">The event data</param>
-		/// <param name="Fingerprint">Receives the fingerprint on success</param>
+		/// <param name="job">The job that spawned the event</param>
+		/// <param name="node">Node that was executed</param>
+		/// <param name="eventData">The event data</param>
+		/// <param name="fingerprint">Receives the fingerprint on success</param>
 		/// <returns>True if the match is successful</returns>
-		bool TryGetFingerprint(IJob Job, INode Node, ILogEventData EventData, [NotNullWhen(true)] out NewIssueFingerprint? Fingerprint);
+		bool TryGetFingerprint(IJob job, INode node, ILogEventData eventData, [NotNullWhen(true)] out NewIssueFingerprint? fingerprint);
 
 		/// <summary>
 		/// Rank all the suspect changes for a given fingerprint
 		/// </summary>
-		/// <param name="Fingerprint">The issue fingerprint</param>
-		/// <param name="Suspects">Potential suspects</param>
-		void RankSuspects(IIssueFingerprint Fingerprint, List<SuspectChange> Suspects);
+		/// <param name="fingerprint">The issue fingerprint</param>
+		/// <param name="suspects">Potential suspects</param>
+		void RankSuspects(IIssueFingerprint fingerprint, List<SuspectChange> suspects);
 
 		/// <summary>
 		/// Gets the summary text for an issue
 		/// </summary>
-		/// <param name="Fingerprint">The fingerprint</param>
-		/// <param name="Severity">Severity of the issue</param>
+		/// <param name="fingerprint">The fingerprint</param>
+		/// <param name="severity">Severity of the issue</param>
 		/// <returns>The summary text</returns>
-		string GetSummary(IIssueFingerprint Fingerprint, IssueSeverity Severity);
+		string GetSummary(IIssueFingerprint fingerprint, IssueSeverity severity);
 	}
 }
