Horde: Fix coding conventions within Horde.Build.

diff --git a/Engine/Source/Programs/Horde/Horde.Build/Issues/IIssueService.cs b/Engine/Source/Programs/Horde/Horde.Build/Issues/IIssueService.cs
--- a/Engine/Source/Programs/Horde/Horde.Build/Issues/IIssueService.cs
+++ b/Engine/Source/Programs/Horde/Horde.Build/Issues/IIssueService.cs
@@ -1,13 +1,11 @@
 // Copyright Epic Games, Inc. All Rights Reserved.
 
+using System.Collections.Generic;
+using System.Threading.Tasks;
 using Horde.Build.Models;
 using Horde.Build.Utilities;
 using Microsoft.Extensions.Hosting;
 using MongoDB.Bson;
-using System;
-using System.Collections.Generic;
-using System.Linq;
-using System.Threading.Tasks;
 
 namespace Horde.Build.Services
 {
@@ -23,7 +21,7 @@ namespace Horde.Build.Services
 		/// <summary>
 		/// Delegate for issue creation events
 		/// </summary>
-		public delegate void IssueUpdatedEvent(IIssue Issue);
+		public delegate void IssueUpdatedEvent(IIssue issue);
 
 		/// <summary>
 		/// Event triggered when a label state changes
@@ -39,121 +37,121 @@ namespace Horde.Build.Services
 		/// <summary>
 		/// Find a cached issue, or update it
 		/// </summary>
-		/// <param name="IssueId">The issue id</param>
+		/// <param name="issueId">The issue id</param>
 		/// <returns>The cached issue data</returns>
-		Task<IIssueDetails?> GetCachedIssueDetailsAsync(int IssueId);
+		Task<IIssueDetails?> GetCachedIssueDetailsAsync(int issueId);
 
 		/// <summary>
 		/// Gets the details for a particular issue
 		/// </summary>
-		/// <param name="Issue">The issue id</param>
+		/// <param name="issue">The issue id</param>
 		/// <returns>The issue details</returns>
-		Task<IIssueDetails> GetIssueDetailsAsync(IIssue Issue);
+		Task<IIssueDetails> GetIssueDetailsAsync(IIssue issue);
 
 		/// <summary>
 		/// Determines whether to show alerts for a given issue
 		/// </summary>
-		/// <param name="Issue">The issue to query</param>
-		/// <param name="Spans">Spans for the issue</param>
+		/// <param name="issue">The issue to query</param>
+		/// <param name="spans">Spans for the issue</param>
 		/// <returns>True if we should show alerts for the given issue</returns>
-		bool ShowDesktopAlertsForIssue(IIssue Issue, IReadOnlyList<IIssueSpan> Spans);
+		bool ShowDesktopAlertsForIssue(IIssue issue, IReadOnlyList<IIssueSpan> spans);
 
 		/// <summary>
 		/// Gets an issue with the given id
 		/// </summary>
-		/// <param name="IssueId"></param>
+		/// <param name="issueId"></param>
 		/// <returns></returns>
-		Task<IIssue?> GetIssueAsync(int IssueId);
+		Task<IIssue?> GetIssueAsync(int issueId);
 
 		/// <summary>
 		/// Gets the spans for the given issue
 		/// </summary>
-		/// <param name="Issue">The issue id</param>
+		/// <param name="issue">The issue id</param>
 		/// <returns>List of spans for the issue</returns>
-		Task<List<IIssueSpan>> GetIssueSpansAsync(IIssue Issue);
+		Task<List<IIssueSpan>> GetIssueSpansAsync(IIssue issue);
 
 		/// <summary>
 		/// Gets the spans for the given issue
 		/// </summary>
-		/// <param name="Span">The issue span</param>
+		/// <param name="span">The issue span</param>
 		/// <returns>List of steps for the issue</returns>
-		Task<List<IIssueStep>> GetIssueStepsAsync(IIssueSpan Span);
+		Task<List<IIssueStep>> GetIssueStepsAsync(IIssueSpan span);
 
 		/// <summary>
 		/// Gets the suspects for an issue
 		/// </summary>
-		/// <param name="Issue"></param>
+		/// <param name="issue"></param>
 		/// <returns></returns>
-		Task<List<IIssueSuspect>> GetIssueSuspectsAsync(IIssue Issue);
+		Task<List<IIssueSuspect>> GetIssueSuspectsAsync(IIssue issue);
 
 		/// <summary>
 		/// Searches for open issues
 		/// </summary>
-		/// <param name="Ids">List of issue ids to include</param>
-		/// <param name="UserId">User to filter issues returned by</param>
-		/// <param name="StreamId">The stream to return issues for</param>
-		/// <param name="MinChange">The minimum changelist number to include</param>
-		/// <param name="MaxChange">The maximum changelist number to include</param>
-		/// <param name="Resolved">Whether to include results that are resolved</param>
-		/// <param name="Promoted">Whether to filter to issues that are promoted</param>
-		/// <param name="Index">Index within the results to return</param>
-		/// <param name="Count">Number of results</param>
+		/// <param name="ids">List of issue ids to include</param>
+		/// <param name="userId">User to filter issues returned by</param>
+		/// <param name="streamId">The stream to return issues for</param>
+		/// <param name="minChange">The minimum changelist number to include</param>
+		/// <param name="maxChange">The maximum changelist number to include</param>
+		/// <param name="resolved">Whether to include results that are resolved</param>
+		/// <param name="promoted">Whether to filter to issues that are promoted</param>
+		/// <param name="index">Index within the results to return</param>
+		/// <param name="count">Number of results</param>
 		/// <returns>List of streams open in the given stream at the given changelist</returns>
-		Task<List<IIssue>> FindIssuesAsync(IEnumerable<int>? Ids = null, UserId? UserId = null, StreamId? StreamId = null, int? MinChange = null, int? MaxChange = null, bool? Resolved = null, bool? Promoted = null, int? Index = null, int? Count = null);
+		Task<List<IIssue>> FindIssuesAsync(IEnumerable<int>? ids = null, UserId? userId = null, StreamId? streamId = null, int? minChange = null, int? maxChange = null, bool? resolved = null, bool? promoted = null, int? index = null, int? count = null);
 
 		/// <summary>
 		/// Searches for open issues affecting a job
 		/// </summary>
-		/// <param name="Ids"></param>
-		/// <param name="Job"></param>
-		/// <param name="Graph"></param>
-		/// <param name="StepId"></param>
-		/// <param name="BatchId"></param>
-		/// <param name="LabelIdx"></param>
-		/// <param name="UserId"></param>
-		/// <param name="Resolved">Whether to include results that are resolved</param>
-		/// <param name="Promoted">Whether to filter by promoted issues</param>
-		/// <param name="Index">Index within the results to return</param>
-		/// <param name="Count">Number of results</param>
+		/// <param name="ids"></param>
+		/// <param name="job"></param>
+		/// <param name="graph"></param>
+		/// <param name="stepId"></param>
+		/// <param name="batchId"></param>
+		/// <param name="labelIdx"></param>
+		/// <param name="userId"></param>
+		/// <param name="resolved">Whether to include results that are resolved</param>
+		/// <param name="promoted">Whether to filter by promoted issues</param>
+		/// <param name="index">Index within the results to return</param>
+		/// <param name="count">Number of results</param>
 		/// <returns></returns>
-		Task<List<IIssue>> FindIssuesForJobAsync(int[]? Ids, IJob Job, IGraph Graph, SubResourceId? StepId = null, SubResourceId? BatchId = null, int? LabelIdx = null, UserId? UserId = null, bool? Resolved = null, bool? Promoted = null, int? Index = null, int? Count = null);
+		Task<List<IIssue>> FindIssuesForJobAsync(int[]? ids, IJob job, IGraph graph, SubResourceId? stepId = null, SubResourceId? batchId = null, int? labelIdx = null, UserId? userId = null, bool? resolved = null, bool? promoted = null, int? index = null, int? count = null);
 
 		/// <summary>
 		/// Find the events for a particular issue
 		/// </summary>
-		/// <param name="IssueId">The issue id</param>
-		/// <param name="LogIds">Log ids to include</param>
-		/// <param name="Index">Index of the first event to return</param>
-		/// <param name="Count">Number of issues to return</param>
+		/// <param name="issueId">The issue id</param>
+		/// <param name="logIds">Log ids to include</param>
+		/// <param name="index">Index of the first event to return</param>
+		/// <param name="count">Number of issues to return</param>
 		/// <returns>List of log events</returns>
-		Task<List<ILogEvent>> FindEventsForIssueAsync(int IssueId, LogId[]? LogIds = null, int Index = 0, int Count = 10);
+		Task<List<ILogEvent>> FindEventsForIssueAsync(int issueId, LogId[]? logIds = null, int index = 0, int count = 10);
 
 		/// <summary>
 		/// Updates the state of an issue
 		/// </summary>
-		/// <param name="Id">The current issue id</param>s
-		/// <param name="Summary">New summary for the issue</param>
-		/// <param name="Description">New description for the issue</param>
-		/// <param name="Promoted">Whether the issue should be set as promoted</param>
-		/// <param name="OwnerId">New owner of the issue</param>
-		/// <param name="NominatedById">Person that nominated the new owner</param>
-		/// <param name="Acknowledged">Whether the issue has been acknowledged</param>
-		/// <param name="DeclinedById">Name of the user that has declined this issue</param>
-		/// <param name="FixChange">Fix changelist for the issue</param>
-		/// <param name="ResolvedById">Whether the issue has been resolved</param>
-		/// <param name="AddSpanIds">Add spans to this issue</param>
-		/// <param name="RemoveSpanIds">Remove spans from this issue</param>
+		/// <param name="id">The current issue id</param>s
+		/// <param name="summary">New summary for the issue</param>
+		/// <param name="description">New description for the issue</param>
+		/// <param name="promoted">Whether the issue should be set as promoted</param>
+		/// <param name="ownerId">New owner of the issue</param>
+		/// <param name="nominatedById">Person that nominated the new owner</param>
+		/// <param name="acknowledged">Whether the issue has been acknowledged</param>
+		/// <param name="declinedById">Name of the user that has declined this issue</param>
+		/// <param name="fixChange">Fix changelist for the issue</param>
+		/// <param name="resolvedById">Whether the issue has been resolved</param>
+		/// <param name="addSpanIds">Add spans to this issue</param>
+		/// <param name="removeSpanIds">Remove spans from this issue</param>
 		/// <returns>True if the issue was updated</returns>
-		Task<bool> UpdateIssueAsync(int Id, string? Summary = null, string? Description = null, bool? Promoted = null, UserId? OwnerId = null, UserId? NominatedById = null, bool? Acknowledged = null, UserId? DeclinedById = null, int? FixChange = null, UserId? ResolvedById = null, List<ObjectId>? AddSpanIds = null, List<ObjectId>? RemoveSpanIds = null);
+		Task<bool> UpdateIssueAsync(int id, string? summary = null, string? description = null, bool? promoted = null, UserId? ownerId = null, UserId? nominatedById = null, bool? acknowledged = null, UserId? declinedById = null, int? fixChange = null, UserId? resolvedById = null, List<ObjectId>? addSpanIds = null, List<ObjectId>? removeSpanIds = null);
 
 		/// <summary>
 		/// Marks a step as complete
 		/// </summary>
-		/// <param name="Job">The job to update</param>
-		/// <param name="Graph">Graph for the job</param>
-		/// <param name="BatchId">Unique id of the batch</param>
-		/// <param name="StepId">Unique id of the step</param>
+		/// <param name="job">The job to update</param>
+		/// <param name="graph">Graph for the job</param>
+		/// <param name="batchId">Unique id of the batch</param>
+		/// <param name="stepId">Unique id of the step</param>
 		/// <returns>Async task</returns>
-		Task UpdateCompleteStep(IJob Job, IGraph Graph, SubResourceId BatchId, SubResourceId StepId);
+		Task UpdateCompleteStep(IJob job, IGraph graph, SubResourceId batchId, SubResourceId stepId);
 	}
 }
diff --git a/Engine/Source/Programs/Horde/Horde.Build/Issues/IIssueSpan.cs b/Engine/Source/Programs/Horde/Horde.Build/Issues/IIssueSpan.cs
--- a/Engine/Source/Programs/Horde/Horde.Build/Issues/IIssueSpan.cs
+++ b/Engine/Source/Programs/Horde/Horde.Build/Issues/IIssueSpan.cs
@@ -1,12 +1,8 @@
 // Copyright Epic Games, Inc. All Rights Reserved.
 
-using EpicGames.Core;
+using System.Collections.Generic;
 using Horde.Build.Utilities;
 using MongoDB.Bson;
-using System;
-using System.Collections.Generic;
-using System.Linq;
-using System.Threading.Tasks;
 
 namespace Horde.Build.Models
 {
diff --git a/Engine/Source/Programs/Horde/Horde.Build/Issues/IIssueStep.cs b/Engine/Source/Programs/Horde/Horde.Build/Issues/IIssueStep.cs
--- a/Engine/Source/Programs/Horde/Horde.Build/Issues/IIssueStep.cs
+++ b/Engine/Source/Programs/Horde/Horde.Build/Issues/IIssueStep.cs
@@ -1,13 +1,8 @@
 // Copyright Epic Games, Inc. All Rights Reserved.
 
+using System;
 using Horde.Build.Utilities;
 using MongoDB.Bson;
-using System;
-using System.Collections.Generic;
-using System.Linq;
-using System.Threading.Tasks;
-
-using SubResourceId = Horde.Build.Models.SubResourceId;
 
 namespace Horde.Build.Models
 {
diff --git a/Engine/Source/Programs/Horde/Horde.Build/Issues/Impl/Handlers/CompileIssueHandler.cs b/Engine/Source/Programs/Horde/Horde.Build/Issues/Impl/Handlers/CompileIssueHandler.cs
--- a/Engine/Source/Programs/Horde/Horde.Build/Issues/Impl/Handlers/CompileIssueHandler.cs
+++ b/Engine/Source/Programs/Horde/Horde.Build/Issues/Impl/Handlers/CompileIssueHandler.cs
@@ -1,25 +1,16 @@
 // Copyright Epic Games, Inc. All Rights Reserved.
 
-using HordeCommon;
-using Horde.Build.Collections;
-using HordeCommon.Rpc;
-using Horde.Build.Services;
-using Horde.Build.Utilities;
-using Microsoft.Extensions.Logging;
-using MongoDB.Bson;
-using MongoDB.Bson.Serialization;
-using MongoDB.Driver;
 using System;
 using System.Collections.Generic;
 using System.Diagnostics.CodeAnalysis;
-using System.IO;
 using System.Linq;
-using System.Security.Policy;
-using System.Text;
-using System.Threading.Tasks;
 using System.Text.Json;
-using Horde.Build.Models;
 using EpicGames.Core;
+using Horde.Build.Collections;
+using Horde.Build.Models;
+using Horde.Build.Utilities;
+using Microsoft.Extensions.Logging;
+using MongoDB.Driver;
 
 namespace Horde.Build.IssueHandlers.Impl
 {
@@ -42,42 +33,42 @@ namespace Horde.Build.IssueHandlers.Impl
 		/// <summary>
 		/// Determines if the given event id matches
 		/// </summary>
-		/// <param name="EventId">The event id to compare</param>
+		/// <param name="eventId">The event id to compare</param>
 		/// <returns>True if the given event id matches</returns>
-		public static bool IsMatchingEventId(EventId? EventId)
+		public static bool IsMatchingEventId(EventId? eventId)
 		{
-			return EventId == KnownLogEvents.Compiler || EventId == KnownLogEvents.AutomationTool_SourceFileLine || EventId == KnownLogEvents.MSBuild;
+			return eventId == KnownLogEvents.Compiler || eventId == KnownLogEvents.AutomationTool_SourceFileLine || eventId == KnownLogEvents.MSBuild;
 		}
 
 		/// <summary>
 		/// Extracts a list of source files from an event
 		/// </summary>
-		/// <param name="Event">The event data</param>
-		/// <param name="SourceFiles">List of source files</param>
-		static void GetSourceFiles(ILogEventData Event, List<string> SourceFiles)
+		/// <param name="logEventData">The event data</param>
+		/// <param name="sourceFiles">List of source files</param>
+		static void GetSourceFiles(ILogEventData logEventData, List<string> sourceFiles)
 		{
-			foreach (ILogEventLine Line in Event.Lines)
+			foreach (ILogEventLine line in logEventData.Lines)
 			{
-				JsonElement Properties;
-				if (Line.Data.TryGetProperty("properties", out Properties) && Properties.ValueKind == JsonValueKind.Object)
+				JsonElement properties;
+				if (line.Data.TryGetProperty("properties", out properties) && properties.ValueKind == JsonValueKind.Object)
 				{
-					string? Prefix = null;
+					string? prefix = null;
 
-					JsonElement NoteElement;
-					if (Properties.TryGetProperty("note", out NoteElement) && NoteElement.GetBoolean())
+					JsonElement noteElement;
+					if (properties.TryGetProperty("note", out noteElement) && noteElement.GetBoolean())
 					{
-						Prefix = NotePrefix;
+						prefix = NotePrefix;
 					}
 
-					foreach (JsonProperty Property in Properties.EnumerateObject())
+					foreach (JsonProperty property in properties.EnumerateObject())
 					{
-						if (Property.NameEquals("file") && Property.Value.ValueKind == JsonValueKind.String)
+						if (property.NameEquals("file") && property.Value.ValueKind == JsonValueKind.String)
 						{
-							AddSourceFile(SourceFiles, Property.Value.GetString()!, Prefix);
+							AddSourceFile(sourceFiles, property.Value.GetString()!, prefix);
 						}
-						if (Property.Value.HasStringProperty("$type", "SourceFile") && Property.Value.TryGetStringProperty("relativePath", out string? Value))
+						if (property.Value.HasStringProperty("$type", "SourceFile") && property.Value.TryGetStringProperty("relativePath", out string? value))
 						{
-							AddSourceFile(SourceFiles, Value, Prefix);
+							AddSourceFile(sourceFiles, value, prefix);
 						}
 					}
 				}
@@ -87,75 +78,75 @@ namespace Horde.Build.IssueHandlers.Impl
 		/// <summary>
 		/// Add a new source file to a list of unique source files
 		/// </summary>
-		/// <param name="SourceFiles">List of source files</param>
-		/// <param name="RelativePath">File to add</param>
-		/// <param name="Prefix">Prefix to insert at the start of the filename</param>
-		static void AddSourceFile(List<string> SourceFiles, string RelativePath, string? Prefix)
+		/// <param name="sourceFiles">List of source files</param>
+		/// <param name="relativePath">File to add</param>
+		/// <param name="prefix">Prefix to insert at the start of the filename</param>
+		static void AddSourceFile(List<string> sourceFiles, string relativePath, string? prefix)
 		{
-			int EndIdx = RelativePath.LastIndexOfAny(new char[] { '/', '\\' }) + 1;
+			int endIdx = relativePath.LastIndexOfAny(new char[] { '/', '\\' }) + 1;
 
-			string FileName = RelativePath.Substring(EndIdx);
-			if (Prefix != null)
+			string fileName = relativePath.Substring(endIdx);
+			if (prefix != null)
 			{
-				FileName = Prefix + FileName;
+				fileName = prefix + fileName;
 			}
 
-			if (!SourceFiles.Any(x => x.Equals(FileName, StringComparison.OrdinalIgnoreCase)))
+			if (!sourceFiles.Any(x => x.Equals(fileName, StringComparison.OrdinalIgnoreCase)))
 			{
-				SourceFiles.Add(FileName);
+				sourceFiles.Add(fileName);
 			}
 		}
 
 		/// <inheritdoc/>
-		public bool TryGetFingerprint(IJob Job, INode Node, ILogEventData EventData, [NotNullWhen(true)] out NewIssueFingerprint? Fingerprint)
+		public bool TryGetFingerprint(IJob job, INode node, ILogEventData eventData, [NotNullWhen(true)] out NewIssueFingerprint? fingerprint)
 		{
-			if (!IsMatchingEventId(EventData.EventId))
+			if (!IsMatchingEventId(eventData.EventId))
 			{
-				Fingerprint = null;
+				fingerprint = null;
 				return false;
 			}
 
-			List<string> NewFileNames = new List<string>();
-			GetSourceFiles(EventData, NewFileNames);
-			Fingerprint = new NewIssueFingerprint(Type, NewFileNames, null);
+			List<string> newFileNames = new List<string>();
+			GetSourceFiles(eventData, newFileNames);
+			fingerprint = new NewIssueFingerprint(Type, newFileNames, null);
 			return true;
 		}
 
 		/// <inheritdoc/>
-		public string GetSummary(IIssueFingerprint Fingerprint, IssueSeverity Severity)
+		public string GetSummary(IIssueFingerprint fingerprint, IssueSeverity severity)
 		{
-			string Type = (Severity == IssueSeverity.Warning) ? "Compile warnings" : "Compile errors";
-			string List = StringUtils.FormatList(Fingerprint.Keys.Where(x => !x.StartsWith(NotePrefix, StringComparison.Ordinal)).ToArray(), 2);
-			return $"{Type} in {List}";
+			string type = (severity == IssueSeverity.Warning) ? "Compile warnings" : "Compile errors";
+			string list = StringUtils.FormatList(fingerprint.Keys.Where(x => !x.StartsWith(NotePrefix, StringComparison.Ordinal)).ToArray(), 2);
+			return $"{type} in {list}";
 		}
 
 		/// <inheritdoc/>
-		public void RankSuspects(IIssueFingerprint Fingerprint, List<SuspectChange> Suspects)
+		public void RankSuspects(IIssueFingerprint fingerprint, List<SuspectChange> suspects)
 		{
-			List<string> FileNames = new List<string>();
-			foreach (string Key in Fingerprint.Keys)
+			List<string> fileNames = new List<string>();
+			foreach (string key in fingerprint.Keys)
 			{
-				if (Key.StartsWith(NotePrefix, StringComparison.Ordinal))
+				if (key.StartsWith(NotePrefix, StringComparison.Ordinal))
 				{
-					FileNames.Add(Key.Substring(NotePrefix.Length));
+					fileNames.Add(key.Substring(NotePrefix.Length));
 				}
 				else
 				{
-					FileNames.Add(Key);
+					fileNames.Add(key);
 				}
 			}
 
-			foreach (SuspectChange Change in Suspects)
+			foreach (SuspectChange change in suspects)
 			{
-				if (Change.ContainsCode)
+				if (change.ContainsCode)
 				{
-					if (FileNames.Any(x => Change.ModifiesFile(x)))
+					if (fileNames.Any(x => change.ModifiesFile(x)))
 					{
-						Change.Rank += 20;
+						change.Rank += 20;
 					}
 					else
 					{
-						Change.Rank += 10;
+						change.Rank += 10;
 					}
 				}
 			}
diff --git a/Engine/Source/Programs/Horde/Horde.Build/Issues/Impl/Handlers/ContentIssueHandler.cs b/Engine/Source/Programs/Horde/Horde.Build/Issues/Impl/Handlers/ContentIssueHandler.cs
--- a/Engine/Source/Programs/Horde/Horde.Build/Issues/Impl/Handlers/ContentIssueHandler.cs
+++ b/Engine/Source/Programs/Horde/Horde.Build/Issues/Impl/Handlers/ContentIssueHandler.cs
@@ -1,22 +1,15 @@
 // Copyright Epic Games, Inc. All Rights Reserved.
 
-using HordeCommon;
-using Horde.Build.Collections;
-using Horde.Build.Models;
-using Horde.Build.Utilities;
-using Microsoft.Extensions.Logging;
-using MongoDB.Bson;
-using MongoDB.Bson.Serialization;
-using MongoDB.Driver;
 using System;
 using System.Collections.Generic;
 using System.Diagnostics.CodeAnalysis;
-using System.IO;
 using System.Linq;
-using System.Security.Policy;
-using System.Text;
-using System.Threading.Tasks;
 using EpicGames.Core;
+using Horde.Build.Collections;
+using Horde.Build.Models;
+using Horde.Build.Utilities;
+using Microsoft.Extensions.Logging;
+using MongoDB.Driver;
 
 namespace Horde.Build.IssueHandlers.Impl
 {
@@ -34,69 +27,69 @@ namespace Horde.Build.IssueHandlers.Impl
 		/// <summary>
 		/// Determines if the given event id matches
 		/// </summary>
-		/// <param name="EventId">The event id to compare</param>
+		/// <param name="eventId">The event id to compare</param>
 		/// <returns>True if the given event id matches</returns>
-		public static bool IsMatchingEventId(EventId? EventId)
+		public static bool IsMatchingEventId(EventId? eventId)
 		{
-			return EventId == KnownLogEvents.Engine_AssetLog;
+			return eventId == KnownLogEvents.Engine_AssetLog;
 		}
 
 		/// <summary>
 		/// Adds all the assets from the given log event
 		/// </summary>
-		/// <param name="Event">The log event to parse</param>
-		/// <param name="AssetNames">Receives the referenced asset names</param>
-		public static void GetAssetNames(ILogEventData Event, HashSet<string> AssetNames)
+		/// <param name="eventData">The log event to parse</param>
+		/// <param name="assetNames">Receives the referenced asset names</param>
+		public static void GetAssetNames(ILogEventData eventData, HashSet<string> assetNames)
 		{
-			foreach (ILogEventLine Line in Event.Lines)
+			foreach (ILogEventLine line in eventData.Lines)
 			{
-				string? RelativePath;
-				if (Line.Data.TryGetNestedProperty("properties.asset.relativePath", out RelativePath))
+				string? relativePath;
+				if (line.Data.TryGetNestedProperty("properties.asset.relativePath", out relativePath))
 				{
-					int EndIdx = RelativePath.LastIndexOfAny(new char[] { '/', '\\' }) + 1;
-					string FileName = RelativePath.Substring(EndIdx);
-					AssetNames.Add(FileName);
+					int endIdx = relativePath.LastIndexOfAny(new char[] { '/', '\\' }) + 1;
+					string fileName = relativePath.Substring(endIdx);
+					assetNames.Add(fileName);
 				}
 			}
 		}
 
 		/// <inheritdoc/>
-		public bool TryGetFingerprint(IJob Job, INode Node, ILogEventData EventData, [NotNullWhen(true)] out NewIssueFingerprint? Fingerprint)
+		public bool TryGetFingerprint(IJob job, INode node, ILogEventData eventData, [NotNullWhen(true)] out NewIssueFingerprint? fingerprint)
 		{
-			if (!IsMatchingEventId(EventData.EventId))
+			if (!IsMatchingEventId(eventData.EventId))
 			{
-				Fingerprint = null;
+				fingerprint = null;
 				return false;
 			}
 
-			HashSet<string> NewAssetNames = new HashSet<string>();
-			GetAssetNames(EventData, NewAssetNames);
-			Fingerprint = new NewIssueFingerprint(Type, NewAssetNames, null);
+			HashSet<string> newAssetNames = new HashSet<string>();
+			GetAssetNames(eventData, newAssetNames);
+			fingerprint = new NewIssueFingerprint(Type, newAssetNames, null);
 			return true;
 		}
 
 		/// <inheritdoc/>
-		public string GetSummary(IIssueFingerprint Fingerprint, IssueSeverity Severity)
+		public string GetSummary(IIssueFingerprint fingerprint, IssueSeverity severity)
 		{
-			string Type = (Severity == IssueSeverity.Warning) ? "Warnings" : "Errors";
-			string List = StringUtils.FormatList(Fingerprint.Keys.ToArray(), 2);
-			return $"{Type} in {List}";
+			string type = (severity == IssueSeverity.Warning) ? "Warnings" : "Errors";
+			string list = StringUtils.FormatList(fingerprint.Keys.ToArray(), 2);
+			return $"{type} in {list}";
 		}
 
 		/// <inheritdoc/>
-		public void RankSuspects(IIssueFingerprint Fingerprint, List<SuspectChange> Suspects)
+		public void RankSuspects(IIssueFingerprint fingerprint, List<SuspectChange> suspects)
 		{
-			foreach (SuspectChange Suspect in Suspects)
+			foreach (SuspectChange suspect in suspects)
 			{
-				if (Suspect.ContainsContent)
+				if (suspect.ContainsContent)
 				{
-					if (Suspect.Details.Files.Any(x => Fingerprint.Keys.Any(y => x.Path.Contains(y, StringComparison.OrdinalIgnoreCase))))
+					if (suspect.Details.Files.Any(x => fingerprint.Keys.Any(y => x.Path.Contains(y, StringComparison.OrdinalIgnoreCase))))
 					{
-						Suspect.Rank += 20;
+						suspect.Rank += 20;
 					}
 					else
 					{
-						Suspect.Rank += 10;
+						suspect.Rank += 10;
 					}
 				}
 			}
diff --git a/Engine/Source/Programs/Horde/Horde.Build/Issues/Impl/Handlers/CopyrightIssueHandler.cs b/Engine/Source/Programs/Horde/Horde.Build/Issues/Impl/Handlers/CopyrightIssueHandler.cs
--- a/Engine/Source/Programs/Horde/Horde.Build/Issues/Impl/Handlers/CopyrightIssueHandler.cs
+++ b/Engine/Source/Programs/Horde/Horde.Build/Issues/Impl/Handlers/CopyrightIssueHandler.cs
@@ -1,23 +1,15 @@
 // Copyright Epic Games, Inc. All Rights Reserved.
 
-using HordeCommon;
-using Horde.Build.Collections;
-using HordeCommon.Rpc;
-using Horde.Build.Models;
-using Horde.Build.Utilities;
-using Microsoft.Extensions.Logging;
-using MongoDB.Bson;
-using MongoDB.Bson.Serialization;
-using MongoDB.Driver;
 using System;
 using System.Collections.Generic;
 using System.Diagnostics.CodeAnalysis;
-using System.IO;
 using System.Linq;
-using System.Security.Policy;
-using System.Text;
-using System.Threading.Tasks;
 using EpicGames.Core;
+using Horde.Build.Collections;
+using Horde.Build.Models;
+using Horde.Build.Utilities;
+using Microsoft.Extensions.Logging;
+using MongoDB.Driver;
 
 namespace Horde.Build.IssueHandlers.Impl
 {
@@ -40,70 +32,70 @@ namespace Horde.Build.IssueHandlers.Impl
 		/// <summary>
 		/// Determines if the given event id matches
 		/// </summary>
-		/// <param name="EventId">The event id to compare</param>
+		/// <param name="eventId">The event id to compare</param>
 		/// <returns>True if the given event id matches</returns>
-		public static bool IsMatchingEventId(EventId? EventId)
+		public static bool IsMatchingEventId(EventId? eventId)
 		{
-			return EventId == KnownLogEvents.AutomationTool_MissingCopyright;
+			return eventId == KnownLogEvents.AutomationTool_MissingCopyright;
 		}
 
 		/// <summary>
 		/// Extracts a list of source files from an event
 		/// </summary>
-		/// <param name="Event">The event data</param>
-		/// <param name="SourceFiles">List of source files</param>
-		public static void GetSourceFiles(ILogEventData Event, HashSet<string> SourceFiles)
+		/// <param name="logEventData">The event data</param>
+		/// <param name="sourceFiles">List of source files</param>
+		public static void GetSourceFiles(ILogEventData logEventData, HashSet<string> sourceFiles)
 		{
-			foreach (ILogEventLine Line in Event.Lines)
+			foreach (ILogEventLine line in logEventData.Lines)
 			{
-				string? RelativePath;
-				if (Line.Data.TryGetNestedProperty("properties.file.relativePath", out RelativePath) || Line.Data.TryGetNestedProperty("properties.file", out RelativePath))
+				string? relativePath;
+				if (line.Data.TryGetNestedProperty("properties.file.relativePath", out relativePath) || line.Data.TryGetNestedProperty("properties.file", out relativePath))
 				{
-					int EndIdx = RelativePath.LastIndexOfAny(new char[] { '/', '\\' }) + 1;
-					string FileName = RelativePath.Substring(EndIdx);
-					SourceFiles.Add(FileName);
+					int endIdx = relativePath.LastIndexOfAny(new char[] { '/', '\\' }) + 1;
+					string fileName = relativePath.Substring(endIdx);
+					sourceFiles.Add(fileName);
 				}
 			}
 		}
 
 		/// <inheritdoc/>
-		public bool TryGetFingerprint(IJob Job, INode Node, ILogEventData EventData, [NotNullWhen(true)] out NewIssueFingerprint? Fingerprint)
+		public bool TryGetFingerprint(IJob job, INode node, ILogEventData eventData, [NotNullWhen(true)] out NewIssueFingerprint? fingerprint)
 		{
-			if (!IsMatchingEventId(EventData.EventId))
+			if (!IsMatchingEventId(eventData.EventId))
 			{
-				Fingerprint = null;
+				fingerprint = null;
 				return false;
 			}
 
-			HashSet<string> NewFileNames = new HashSet<string>(StringComparer.OrdinalIgnoreCase);
-			GetSourceFiles(EventData, NewFileNames);
-			Fingerprint = new NewIssueFingerprint(Type, NewFileNames, null);
+			HashSet<string> newFileNames = new HashSet<string>(StringComparer.OrdinalIgnoreCase);
+			GetSourceFiles(eventData, newFileNames);
+			fingerprint = new NewIssueFingerprint(Type, newFileNames, null);
 			return true;
 		}
 
 		/// <inheritdoc/>
-		public void RankSuspects(IIssueFingerprint Fingerprint, List<SuspectChange> Suspects)
+		public void RankSuspects(IIssueFingerprint fingerprint, List<SuspectChange> suspects)
 		{
-			foreach (SuspectChange Suspect in Suspects)
+			foreach (SuspectChange suspect in suspects)
 			{
-				if (Suspect.ContainsCode)
+				if (suspect.ContainsCode)
 				{
-					if (FileNames.Any(x => Suspect.ModifiesFile(x)))
+					if (FileNames.Any(x => suspect.ModifiesFile(x)))
 					{
-						Suspect.Rank += 20;
+						suspect.Rank += 20;
 					}
 					else
 					{
-						Suspect.Rank += 10;
+						suspect.Rank += 10;
 					}
 				}
 			}
 		}
 
 		/// <inheritdoc/>
-		public string GetSummary(IIssueFingerprint Fingerprint, IssueSeverity Severity)
+		public string GetSummary(IIssueFingerprint fingerprint, IssueSeverity severity)
 		{
-			return $"Missing copyright notice in {StringUtils.FormatList(Fingerprint.Keys.ToArray(), 2)}";
+			return $"Missing copyright notice in {StringUtils.FormatList(fingerprint.Keys.ToArray(), 2)}";
 		}
 	}
 }
diff --git a/Engine/Source/Programs/Horde/Horde.Build/Issues/Impl/Handlers/DefaultIssueHandler.cs b/Engine/Source/Programs/Horde/Horde.Build/Issues/Impl/Handlers/DefaultIssueHandler.cs
--- a/Engine/Source/Programs/Horde/Horde.Build/Issues/Impl/Handlers/DefaultIssueHandler.cs
+++ b/Engine/Source/Programs/Horde/Horde.Build/Issues/Impl/Handlers/DefaultIssueHandler.cs
@@ -1,12 +1,10 @@
 // Copyright Epic Games, Inc. All Rights Reserved.
 
-using Horde.Build.Collections;
-using Horde.Build.Models;
-using System;
 using System.Collections.Generic;
 using System.Diagnostics.CodeAnalysis;
 using System.Linq;
-using System.Threading.Tasks;
+using Horde.Build.Collections;
+using Horde.Build.Models;
 
 namespace Horde.Build.IssueHandlers.Impl
 {
@@ -27,28 +25,28 @@ namespace Horde.Build.IssueHandlers.Impl
 		public int Priority => 0;
 
 		/// <inheritdoc/>
-		public bool TryGetFingerprint(IJob Job, INode Node, ILogEventData EventData, [NotNullWhen(true)] out NewIssueFingerprint? Fingerprint)
+		public bool TryGetFingerprint(IJob job, INode node, ILogEventData eventData, [NotNullWhen(true)] out NewIssueFingerprint? fingerprint)
 		{
-			Fingerprint = new NewIssueFingerprint(Type, new[] { Node.Name }, null);
+			fingerprint = new NewIssueFingerprint(Type, new[] { node.Name }, null);
 			return true;
 		}
 
 		/// <inheritdoc/>
-		public void RankSuspects(IIssueFingerprint Fingerprint, List<SuspectChange> Suspects)
+		public void RankSuspects(IIssueFingerprint fingerprint, List<SuspectChange> suspects)
 		{
 		}
 
 		/// <inheritdoc/>
-		public string GetSummary(IIssueFingerprint Fingerprint, IssueSeverity Severity)
+		public string GetSummary(IIssueFingerprint fingerprint, IssueSeverity severity)
 		{
-			string NodeName = Fingerprint.Keys.FirstOrDefault() ?? "(unknown)";
-			if(Severity == IssueSeverity.Warning)
+			string nodeName = fingerprint.Keys.FirstOrDefault() ?? "(unknown)";
+			if(severity == IssueSeverity.Warning)
 			{
-				return $"Warnings in {NodeName}";
+				return $"Warnings in {nodeName}";
 			}
 			else
 			{
-				return $"Errors in {NodeName}";
+				return $"Errors in {nodeName}";
 			}
 		}
 	}
diff --git a/Engine/Source/Programs/Horde/Horde.Build/Issues/Impl/Handlers/GauntletIssueHandler.cs b/Engine/Source/Programs/Horde/Horde.Build/Issues/Impl/Handlers/GauntletIssueHandler.cs
--- a/Engine/Source/Programs/Horde/Horde.Build/Issues/Impl/Handlers/GauntletIssueHandler.cs
+++ b/Engine/Source/Programs/Horde/Horde.Build/Issues/Impl/Handlers/GauntletIssueHandler.cs
@@ -1,21 +1,15 @@
 // Copyright Epic Games, Inc. All Rights Reserved.
 
-using HordeCommon;
-using Horde.Build.Models;
-using Horde.Build.Utilities;
-using Microsoft.AspNetCore.Authentication;
-using Microsoft.AspNetCore.Connections.Features;
-using Microsoft.Extensions.Logging;
-using Serilog.Formatting.Json;
 using System;
 using System.Collections.Generic;
+using System.Diagnostics.CodeAnalysis;
 using System.Linq;
-using System.Security.Cryptography.X509Certificates;
 using System.Text.Json;
-using System.Threading.Tasks;
-using System.Diagnostics.CodeAnalysis;
-using Horde.Build.Collections;
 using EpicGames.Core;
+using Horde.Build.Collections;
+using Horde.Build.Models;
+using Horde.Build.Utilities;
+using Microsoft.Extensions.Logging;
 
 namespace Horde.Build.IssueHandlers.Impl
 {
@@ -43,38 +37,38 @@ namespace Horde.Build.IssueHandlers.Impl
 		/// <summary>
 		/// Determines if the given event id matches
 		/// </summary>
-		/// <param name="EventId">The event id to compare</param>
+		/// <param name="eventId">The event id to compare</param>
 		/// <returns>True if the given event id matches</returns>
-		public static bool IsMatchingEventId(EventId? EventId)
+		public static bool IsMatchingEventId(EventId? eventId)
 		{
-			return EventId == KnownLogEvents.Gauntlet_UnitTest || EventId == KnownLogEvents.Gauntlet_ScreenshotTest;
+			return eventId == KnownLogEvents.Gauntlet_UnitTest || eventId == KnownLogEvents.Gauntlet_ScreenshotTest;
 		}
 
 		/// <summary>
 		/// Parses symbol names from a log event
 		/// </summary>
-		/// <param name="EventData">The log event data</param>
-		/// <param name="UnitTestNames">Receives a set of the unit test names</param>
-		private static void GetUnitTestNames(ILogEventData EventData, HashSet<string> UnitTestNames)
+		/// <param name="eventData">The log event data</param>
+		/// <param name="unitTestNames">Receives a set of the unit test names</param>
+		private static void GetUnitTestNames(ILogEventData eventData, HashSet<string> unitTestNames)
 		{
-			foreach (ILogEventLine Line in EventData.Lines)
+			foreach (ILogEventLine line in eventData.Lines)
 			{
-				string? Group = null;
-				string? Name = null;
+				string? group = null;
+				string? name = null;
 
-				string? Value;
-				if (Line.Data.TryGetNestedProperty("properties.group", out Value))
+				string? value;
+				if (line.Data.TryGetNestedProperty("properties.group", out value))
 				{
-					Group = Value;
+					group = value;
 				}
-				if (Line.Data.TryGetNestedProperty("properties.name", out Value))
+				if (line.Data.TryGetNestedProperty("properties.name", out value))
 				{
-					Name = Value;
+					name = value;
 				}
 
-				if (Group != null && Name != null)
+				if (group != null && name != null)
 				{
-					UnitTestNames.Add($"{UnitTestPrefix}:{Group}/{Name}");
+					unitTestNames.Add($"{UnitTestPrefix}:{group}/{name}");
 				}
 			}
 		}
@@ -82,20 +76,20 @@ namespace Horde.Build.IssueHandlers.Impl
 		/// <summary>
 		/// Parses screenshot test names from a log event
 		/// </summary>
-		/// <param name="EventData">The event data</param>
-		/// <param name="ScreenshotTestNames">Receives the parsed screenshot test names</param>
-		private static void GetScreenshotTestNames(ILogEventData EventData, HashSet<string> ScreenshotTestNames)
+		/// <param name="eventData">The event data</param>
+		/// <param name="screenshotTestNames">Receives the parsed screenshot test names</param>
+		private static void GetScreenshotTestNames(ILogEventData eventData, HashSet<string> screenshotTestNames)
 		{
-			foreach (ILogEventLine Line in EventData.Lines)
+			foreach (ILogEventLine line in eventData.Lines)
 			{
-				if (Line.Data.TryGetProperty("properties", JsonValueKind.Object, out JsonElement Properties))
+				if (line.Data.TryGetProperty("properties", JsonValueKind.Object, out JsonElement properties))
 				{
-					foreach (JsonProperty Property in Properties.EnumerateObject())
+					foreach (JsonProperty property in properties.EnumerateObject())
 					{
-						JsonElement Value = Property.Value;
-						if (Value.ValueKind == JsonValueKind.Object && Value.HasStringProperty("type", "Screenshot") && Value.TryGetStringProperty("name", out string? Name))
+						JsonElement value = property.Value;
+						if (value.ValueKind == JsonValueKind.Object && value.HasStringProperty("type", "Screenshot") && value.TryGetStringProperty("name", out string? name))
 						{
-							ScreenshotTestNames.Add($"{ScreenshotTestPrefix}:{Name}");
+							screenshotTestNames.Add($"{ScreenshotTestPrefix}:{name}");
 						}
 					}
 				}
@@ -103,72 +97,72 @@ namespace Horde.Build.IssueHandlers.Impl
 		}
 
 		/// <inheritdoc/>
-		public bool TryGetFingerprint(IJob Job, INode Node, ILogEventData EventData, [NotNullWhen(true)] out NewIssueFingerprint? Fingerprint)
+		public bool TryGetFingerprint(IJob job, INode node, ILogEventData eventData, [NotNullWhen(true)] out NewIssueFingerprint? fingerprint)
 		{
-			if(!IsMatchingEventId(EventData.EventId))
+			if(!IsMatchingEventId(eventData.EventId))
 			{
-				Fingerprint = null;
+				fingerprint = null;
 				return false;
 			}
 
-			HashSet<string> Keys = new HashSet<string>();
-			GetUnitTestNames(EventData, Keys);
-			GetScreenshotTestNames(EventData, Keys);
+			HashSet<string> keys = new HashSet<string>();
+			GetUnitTestNames(eventData, keys);
+			GetScreenshotTestNames(eventData, keys);
 
-			if (Keys.Count == 0)
+			if (keys.Count == 0)
 			{
-				Fingerprint = null;
+				fingerprint = null;
 				return false;
 			}
 
-			Fingerprint = new NewIssueFingerprint(Type, Keys, null);
+			fingerprint = new NewIssueFingerprint(Type, keys, null);
 			return true;
 		}
 
 		/// <inheritdoc/>
-		public void RankSuspects(IIssueFingerprint Fingerprint, List<SuspectChange> Changes)
+		public void RankSuspects(IIssueFingerprint fingerprint, List<SuspectChange> changes)
 		{
-			foreach (SuspectChange Change in Changes)
+			foreach (SuspectChange change in changes)
 			{
-				if (Change.ContainsCode)
+				if (change.ContainsCode)
 				{
-					Change.Rank += 10;
+					change.Rank += 10;
 				}
 			}
 		}
 
 		/// <inheritdoc/>
-		public string GetSummary(IIssueFingerprint Fingerprint, IssueSeverity Severity)
+		public string GetSummary(IIssueFingerprint fingerprint, IssueSeverity severity)
 		{
-			List<string> UnitTestNames = Fingerprint.Keys.Where(x => x.StartsWith(UnitTestPrefix, StringComparison.Ordinal)).Select(x => x.Substring(UnitTestPrefix.Length + 1)).ToList();
-			if (UnitTestNames.Count > 0)
+			List<string> unitTestNames = fingerprint.Keys.Where(x => x.StartsWith(UnitTestPrefix, StringComparison.Ordinal)).Select(x => x.Substring(UnitTestPrefix.Length + 1)).ToList();
+			if (unitTestNames.Count > 0)
 			{
-				HashSet<string> GroupNames = new HashSet<string>(StringComparer.OrdinalIgnoreCase);
-				HashSet<string> TestNames = new HashSet<string>(StringComparer.OrdinalIgnoreCase);
-				foreach (string UnitTestName in UnitTestNames)
+				HashSet<string> groupNames = new HashSet<string>(StringComparer.OrdinalIgnoreCase);
+				HashSet<string> testNames = new HashSet<string>(StringComparer.OrdinalIgnoreCase);
+				foreach (string unitTestName in unitTestNames)
 				{
-					int Idx = UnitTestName.IndexOf('/', StringComparison.OrdinalIgnoreCase);
-					if (Idx != -1)
+					int idx = unitTestName.IndexOf('/', StringComparison.OrdinalIgnoreCase);
+					if (idx != -1)
 					{
-						GroupNames.Add(UnitTestName.Substring(0, Idx));
+						groupNames.Add(unitTestName.Substring(0, idx));
 					}
-					TestNames.Add(UnitTestName.Substring(Idx + 1));
+					testNames.Add(unitTestName.Substring(idx + 1));
 				}
 
-				if (GroupNames.Count == 1)
+				if (groupNames.Count == 1)
 				{
-					return $"{GroupNames.First()} test failures: {StringUtils.FormatList(TestNames.ToArray(), 3)}";
+					return $"{groupNames.First()} test failures: {StringUtils.FormatList(testNames.ToArray(), 3)}";
 				}
 				else
 				{
-					return $"{StringUtils.FormatList(GroupNames.OrderBy(x => x).ToArray(), 100)} test failures";
+					return $"{StringUtils.FormatList(groupNames.OrderBy(x => x).ToArray(), 100)} test failures";
 				}
 			}
 
-			List<string> ScreenshotTestNames = Fingerprint.Keys.Where(x => x.StartsWith(ScreenshotTestPrefix, StringComparison.Ordinal)).Select(x => x.Substring(ScreenshotTestPrefix.Length + 1)).ToList();
-			if (ScreenshotTestNames.Count > 0)
+			List<string> screenshotTestNames = fingerprint.Keys.Where(x => x.StartsWith(ScreenshotTestPrefix, StringComparison.Ordinal)).Select(x => x.Substring(ScreenshotTestPrefix.Length + 1)).ToList();
+			if (screenshotTestNames.Count > 0)
 			{
-				return $"Screenshot test failures: {StringUtils.FormatList(ScreenshotTestNames.ToArray(), 3)}";
+				return $"Screenshot test failures: {StringUtils.FormatList(screenshotTestNames.ToArray(), 3)}";
 			}
 
 			return "Test failures";
diff --git a/Engine/Source/Programs/Horde/Horde.Build/Issues/Impl/Handlers/SymbolIssueHandler.cs b/Engine/Source/Programs/Horde/Horde.Build/Issues/Impl/Handlers/SymbolIssueHandler.cs
--- a/Engine/Source/Programs/Horde/Horde.Build/Issues/Impl/Handlers/SymbolIssueHandler.cs
+++ b/Engine/Source/Programs/Horde/Horde.Build/Issues/Impl/Handlers/SymbolIssueHandler.cs
@@ -1,19 +1,14 @@
 // Copyright Epic Games, Inc. All Rights Reserved.
 
-using HordeCommon;
-using Horde.Build.Collections;
-using Horde.Build.Models;
-using Horde.Build.Utilities;
-using Microsoft.Extensions.Logging;
-using MongoDB.Bson;
-using MongoDB.Bson.Serialization;
 using System;
 using System.Collections.Generic;
 using System.Diagnostics.CodeAnalysis;
 using System.Linq;
-using System.Text.RegularExpressions;
-using System.Threading.Tasks;
 using EpicGames.Core;
+using Horde.Build.Collections;
+using Horde.Build.Models;
+using Horde.Build.Utilities;
+using Microsoft.Extensions.Logging;
 
 namespace Horde.Build.IssueHandlers.Impl
 {
@@ -31,74 +26,74 @@ namespace Horde.Build.IssueHandlers.Impl
 		/// <summary>
 		/// Determines if the given event id matches
 		/// </summary>
-		/// <param name="EventId">The event id to compare</param>
+		/// <param name="eventId">The event id to compare</param>
 		/// <returns>True if the given event id matches</returns>
-		public static bool IsMatchingEventId(EventId? EventId)
+		public static bool IsMatchingEventId(EventId? eventId)
 		{
-			return EventId == KnownLogEvents.Linker_UndefinedSymbol || EventId == KnownLogEvents.Linker_DuplicateSymbol || EventId == KnownLogEvents.Linker;
+			return eventId == KnownLogEvents.Linker_UndefinedSymbol || eventId == KnownLogEvents.Linker_DuplicateSymbol || eventId == KnownLogEvents.Linker;
 		}
 
 		/// <summary>
 		/// Parses symbol names from a log event
 		/// </summary>
-		/// <param name="EventData">The log event data</param>
-		/// <param name="SymbolNames">Receives the list of symbol names</param>
-		public static void GetSymbolNames(ILogEventData EventData, SortedSet<string> SymbolNames)
+		/// <param name="eventData">The log event data</param>
+		/// <param name="symbolNames">Receives the list of symbol names</param>
+		public static void GetSymbolNames(ILogEventData eventData, SortedSet<string> symbolNames)
 		{
-			foreach (ILogEventLine Line in EventData.Lines)
+			foreach (ILogEventLine line in eventData.Lines)
 			{
-				string? Identifier;
-				if (Line.Data.TryGetNestedProperty("properties.symbol.identifier", out Identifier))
+				string? identifier;
+				if (line.Data.TryGetNestedProperty("properties.symbol.identifier", out identifier))
 				{
-					SymbolNames.Add(Identifier);
+					symbolNames.Add(identifier);
 				}
 			}
 		}
 
 		/// <inheritdoc/>
-		public void RankSuspects(IIssueFingerprint Fingerprint, List<SuspectChange> Changes)
+		public void RankSuspects(IIssueFingerprint fingerprint, List<SuspectChange> changes)
 		{
-			HashSet<string> Names = new HashSet<string>();
-			foreach (string Name in Fingerprint.Keys)
+			HashSet<string> names = new HashSet<string>();
+			foreach (string name in fingerprint.Keys)
 			{
-				Names.UnionWith(Name.Split("::", StringSplitOptions.RemoveEmptyEntries));
+				names.UnionWith(name.Split("::", StringSplitOptions.RemoveEmptyEntries));
 			}
 
-			foreach (SuspectChange Change in Changes)
+			foreach (SuspectChange change in changes)
 			{
-				if (Change.ContainsCode)
+				if (change.ContainsCode)
 				{
-					int Matches = Names.Count(x => Change.Details.Files.Any(y => y.Path.Contains(x, StringComparison.OrdinalIgnoreCase)));
-					Change.Rank += 10 + (10 * Matches);
+					int matches = names.Count(x => change.Details.Files.Any(y => y.Path.Contains(x, StringComparison.OrdinalIgnoreCase)));
+					change.Rank += 10 + (10 * matches);
 				}
 			}
 		}
 
 		/// <inheritdoc/>
-		public string GetSummary(IIssueFingerprint Fingerprint, IssueSeverity Severity)
+		public string GetSummary(IIssueFingerprint fingerprint, IssueSeverity severity)
 		{
-			HashSet<string> Symbols = Fingerprint.Keys;
-			if (Symbols.Count == 1)
+			HashSet<string> symbols = fingerprint.Keys;
+			if (symbols.Count == 1)
 			{
-				return $"Undefined symbol '{Symbols.First()}'";
+				return $"Undefined symbol '{symbols.First()}'";
 			}
 			else
 			{
-				return $"Undefined symbols: {StringUtils.FormatList(Symbols.ToArray(), 3)}";
+				return $"Undefined symbols: {StringUtils.FormatList(symbols.ToArray(), 3)}";
 			}
 		}
 
-		public bool TryGetFingerprint(IJob Job, INode Node, ILogEventData EventData, [NotNullWhen(true)] out NewIssueFingerprint? Fingerprint)
+		public bool TryGetFingerprint(IJob job, INode node, ILogEventData eventData, [NotNullWhen(true)] out NewIssueFingerprint? fingerprint)
 		{
-			if (!IsMatchingEventId(EventData.EventId))
+			if (!IsMatchingEventId(eventData.EventId))
 			{
-				Fingerprint = null;
+				fingerprint = null;
 				return false;
 			}
 
-			SortedSet<string> SymbolNames = new SortedSet<string>();
-			GetSymbolNames(EventData, SymbolNames);
-			Fingerprint = new NewIssueFingerprint(Type, SymbolNames, null);
+			SortedSet<string> symbolNames = new SortedSet<string>();
+			GetSymbolNames(eventData, symbolNames);
+			fingerprint = new NewIssueFingerprint(Type, symbolNames, null);
 			return true;
 		}
 	}
diff --git a/Engine/Source/Programs/Horde/Horde.Build/Issues/Impl/Handlers/SystemicIssueHandler.cs b/Engine/Source/Programs/Horde/Horde.Build/Issues/Impl/Handlers/SystemicIssueHandler.cs
--- a/Engine/Source/Programs/Horde/Horde.Build/Issues/Impl/Handlers/SystemicIssueHandler.cs
+++ b/Engine/Source/Programs/Horde/Horde.Build/Issues/Impl/Handlers/SystemicIssueHandler.cs
@@ -1,25 +1,14 @@
 // Copyright Epic Games, Inc. All Rights Reserved.
 
-using HordeCommon;
-using Horde.Build.Collections;
-using HordeCommon.Rpc;
-using Horde.Build.Services;
-using Horde.Build.Utilities;
-using Microsoft.Extensions.Logging;
-using MongoDB.Bson;
-using MongoDB.Bson.Serialization;
-using MongoDB.Driver;
 using System;
 using System.Collections.Generic;
 using System.Diagnostics.CodeAnalysis;
-using System.IO;
 using System.Linq;
-using System.Security.Policy;
-using System.Text;
-using System.Threading.Tasks;
-using System.Text.Json;
-using Horde.Build.Models;
 using EpicGames.Core;
+using Horde.Build.Collections;
+using Horde.Build.Models;
+using Microsoft.Extensions.Logging;
+using MongoDB.Driver;
 
 namespace Horde.Build.IssueHandlers.Impl
 {
@@ -37,7 +26,7 @@ namespace Horde.Build.IssueHandlers.Impl
 		/// <summary>
 		///  Known systemic errors
 		/// </summary>
-		static HashSet<EventId?> KnownSystemic = new HashSet<EventId?> {  KnownLogEvents.Systemic, KnownLogEvents.Systemic_Xge, KnownLogEvents.Systemic_Xge_Standalone, 
+		static readonly HashSet<EventId?> s_knownSystemic = new HashSet<EventId?> {  KnownLogEvents.Systemic, KnownLogEvents.Systemic_Xge, KnownLogEvents.Systemic_Xge_Standalone, 
 																		KnownLogEvents.Systemic_Xge_ServiceNotRunning, KnownLogEvents.Systemic_Xge_BuildFailed, 
 																		KnownLogEvents.Systemic_SlowDDC, KnownLogEvents.Systemic_Horde, KnownLogEvents.Systemic_Horde_ArtifactUpload,
 																		KnownLogEvents.Horde, KnownLogEvents.Horde_InvalidPreflight};
@@ -45,58 +34,57 @@ namespace Horde.Build.IssueHandlers.Impl
 		/// <summary>
 		/// Determines if the given event id matches
 		/// </summary>
-		/// <param name="EventId">The event id to compare</param>
+		/// <param name="eventId">The event id to compare</param>
 		/// <returns>True if the given event id matches</returns>
-		public static bool IsMatchingEventId(EventId? EventId)
+		public static bool IsMatchingEventId(EventId? eventId)
 		{
-			return KnownSystemic.Contains(EventId);
+			return s_knownSystemic.Contains(eventId);
 		}
 
 		/// <inheritdoc/>
-		public void RankSuspects(IIssueFingerprint Fingerprint, List<SuspectChange> Suspects)
+		public void RankSuspects(IIssueFingerprint fingerprint, List<SuspectChange> suspects)
 		{
 		}
 
 		/// <inheritdoc/>
-		public bool TryGetFingerprint(IJob Job, INode Node, ILogEventData EventData, [NotNullWhen(true)] out NewIssueFingerprint? Fingerprint)
+		public bool TryGetFingerprint(IJob job, INode node, ILogEventData eventData, [NotNullWhen(true)] out NewIssueFingerprint? fingerprint)
 		{
-			Fingerprint = null;
+			fingerprint = null;
 
-			if (EventData.EventId == KnownLogEvents.ExitCode)
+			if (eventData.EventId == KnownLogEvents.ExitCode)
 			{
-				for (int i = 0; i < EventData.Lines.Count; i++)
+				for (int i = 0; i < eventData.Lines.Count; i++)
 				{
-					string Message = EventData.Lines[i].Message;
+					string message = eventData.Lines[i].Message;
 
-					string[] SystemicExitMessages = new string[] { "AutomationTool exiting with ExitCode", "BUILD FAILED", "tool returned code" };
+					string[] systemicExitMessages = new string[] { "AutomationTool exiting with ExitCode", "BUILD FAILED", "tool returned code" };
 
-					for (int j = 0; j < SystemicExitMessages.Length; j++)
+					for (int j = 0; j < systemicExitMessages.Length; j++)
 					{
-						if (Message.Contains(SystemicExitMessages[j], StringComparison.InvariantCultureIgnoreCase))
+						if (message.Contains(systemicExitMessages[j], StringComparison.InvariantCultureIgnoreCase))
 						{
-							Fingerprint = new NewIssueFingerprint(Type, new[] { Node.Name }, null);
+							fingerprint = new NewIssueFingerprint(Type, new[] { node.Name }, null);
 							return true;
 						}
 					}
 				}
 			}
 
-			if (!IsMatchingEventId(EventData.EventId))
+			if (!IsMatchingEventId(eventData.EventId))
 			{				
 				return false;
 			}
 
-			Fingerprint = new NewIssueFingerprint(Type, new[] { Node.Name }, null);
+			fingerprint = new NewIssueFingerprint(Type, new[] { node.Name }, null);
 			return true;
 		}
 
 		/// <inheritdoc/>
-		public string GetSummary(IIssueFingerprint Fingerprint, IssueSeverity Severity)
+		public string GetSummary(IIssueFingerprint fingerprint, IssueSeverity severity)
 		{
-			string Type = (Severity == IssueSeverity.Warning) ? "Systemic warnings" : "Systemic errors";
-			string NodeName = Fingerprint.Keys.FirstOrDefault() ?? "(unknown)";
-			return $"{Type} in {NodeName}";
+			string type = (severity == IssueSeverity.Warning) ? "Systemic warnings" : "Systemic errors";
+			string nodeName = fingerprint.Keys.FirstOrDefault() ?? "(unknown)";
+			return $"{type} in {nodeName}";
 		}
-
 	}
 }
