Horde: Fix coding conventions within Horde.Build.

diff --git a/Engine/Source/Programs/Horde/Horde.Build/Issues/Impl/IssueCollection.cs b/Engine/Source/Programs/Horde/Horde.Build/Issues/Impl/IssueCollection.cs
--- a/Engine/Source/Programs/Horde/Horde.Build/Issues/Impl/IssueCollection.cs
+++ b/Engine/Source/Programs/Horde/Horde.Build/Issues/Impl/IssueCollection.cs
@@ -1,6 +1,10 @@
 // Copyright Epic Games, Inc. All Rights Reserved.
 
-using EpicGames.Core;
+using System;
+using System.Collections.Generic;
+using System.Diagnostics;
+using System.Linq;
+using System.Threading.Tasks;
 using EpicGames.Redis.Utility;
 using Horde.Build.Models;
 using Horde.Build.Services;
@@ -8,17 +12,8 @@ using Horde.Build.Utilities;
 using Microsoft.Extensions.Caching.Memory;
 using Microsoft.Extensions.Logging;
 using MongoDB.Bson;
-using MongoDB.Bson.Serialization;
 using MongoDB.Bson.Serialization.Attributes;
 using MongoDB.Driver;
-using System;
-using System.Collections.Concurrent;
-using System.Collections.Generic;
-using System.Diagnostics;
-using System.Linq;
-using System.Linq.Expressions;
-using System.Security.Claims;
-using System.Threading.Tasks;
 
 namespace Horde.Build.Collections.Impl
 {
@@ -62,7 +57,7 @@ namespace Horde.Build.Collections.Impl
 			public bool? ManuallyPromoted { get; set; }
 
 			[BsonIgnoreIfNull, BsonElement("Promoted")]
-			public bool? ManuallyPromoted_DEPRECATED { get; set; }
+			public bool? ManuallyPromotedDeprecated { get; set; }
 
 			[BsonIgnoreIfNull]
 			public UserId? OwnerId { get; set; }
@@ -101,7 +96,7 @@ namespace Horde.Build.Collections.Impl
 			public int MaxSuspectChange { get; set; }
 
 			[BsonElement("Suspects"), BsonIgnoreIfNull]
-			public List<IssueSuspect>? Suspects_DEPRECATED { get; set; }
+			public List<IssueSuspect>? SuspectsDeprecated { get; set; }
 
 			[BsonIgnoreIfNull]
 			public List<ObjectId>? ExcludeSpans { get; set; }
@@ -120,22 +115,22 @@ namespace Horde.Build.Collections.Impl
 				Fingerprint = null!;
 			}
 
-			public Issue(int Id, string Summary)
+			public Issue(int id, string summary)
 			{
-				this.Id = Id;
-				this.Summary = Summary;
-				this.CreatedAt = DateTime.UtcNow;
-				this.LastSeenAt = DateTime.UtcNow;
+				Id = id;
+				Summary = summary;
+				CreatedAt = DateTime.UtcNow;
+				LastSeenAt = DateTime.UtcNow;
 			}
 
 			UserId? GetDefaultOwnerId()
 			{
-				if (Suspects_DEPRECATED != null && Suspects_DEPRECATED.Count > 0)
+				if (SuspectsDeprecated != null && SuspectsDeprecated.Count > 0)
 				{
-					UserId PossibleOwner = Suspects_DEPRECATED[0].AuthorId;
-					if (Suspects_DEPRECATED.All(x => x.AuthorId == PossibleOwner) && Suspects_DEPRECATED.Any(x => x.DeclinedAt == null))
+					UserId possibleOwner = SuspectsDeprecated[0].AuthorId;
+					if (SuspectsDeprecated.All(x => x.AuthorId == possibleOwner) && SuspectsDeprecated.Any(x => x.DeclinedAt == null))
 					{
-						return Suspects_DEPRECATED[0].AuthorId;
+						return SuspectsDeprecated[0].AuthorId;
 					}
 				}
 				return null;
@@ -147,11 +142,11 @@ namespace Horde.Build.Collections.Impl
 				{
 					if (Fingerprints == null || Fingerprints.Count == 0)
 					{
-						return string.Empty;
+						return String.Empty;
 					}
 
-					return string.Join(", ", Fingerprints.Select(x => {
-					   return $"(Type: {x.Type} / Keys: {string.Join(", ", x.Keys)} / RejectKeys: {string.Join(", ", x.RejectKeys ?? new CaseInsensitiveStringSet(new string[] {"No Reject Keys"}))})";
+					return String.Join(", ", Fingerprints.Select(x => {
+					   return $"(Type: {x.Type} / Keys: {String.Join(", ", x.Keys)} / RejectKeys: {String.Join(", ", x.RejectKeys ?? new CaseInsensitiveStringSet(new string[] {"No Reject Keys"}))})";
 				   }));
 				}
 			}
@@ -166,10 +161,10 @@ namespace Horde.Build.Collections.Impl
 			{
 			}
 
-			public IssueStream(IIssueStream Other)
+			public IssueStream(IIssueStream other)
 			{
-				this.StreamId = Other.StreamId;
-				this.ContainsFix = Other.ContainsFix;
+				StreamId = other.StreamId;
+				ContainsFix = other.ContainsFix;
 			}
 		}
 
@@ -186,28 +181,28 @@ namespace Horde.Build.Collections.Impl
 			{
 			}
 
-			public IssueSuspect(int IssueId, NewIssueSuspectData NewSuspect, DateTime? ResolvedAt)
+			public IssueSuspect(int issueId, NewIssueSuspectData newSuspect, DateTime? resolvedAt)
 			{
-				this.Id = ObjectId.GenerateNewId();
-				this.IssueId = IssueId;
-				this.AuthorId = NewSuspect.AuthorId;
-				this.Change = NewSuspect.Change;
-				this.ResolvedAt = ResolvedAt;
+				Id = ObjectId.GenerateNewId();
+				IssueId = issueId;
+				AuthorId = newSuspect.AuthorId;
+				Change = newSuspect.Change;
+				ResolvedAt = resolvedAt;
 			}
 
-			public IssueSuspect(int IssueId, IIssueSpanSuspect Suspect)
-				: this(IssueId, Suspect.AuthorId, Suspect.OriginatingChange ?? Suspect.Change, null, null)
+			public IssueSuspect(int issueId, IIssueSpanSuspect suspect)
+				: this(issueId, suspect.AuthorId, suspect.OriginatingChange ?? suspect.Change, null, null)
 			{
 			}
 
-			public IssueSuspect(int IssueId, UserId AuthorId, int Change, DateTime? DeclinedAt, DateTime? ResolvedAt)
+			public IssueSuspect(int issueId, UserId authorId, int change, DateTime? declinedAt, DateTime? resolvedAt)
 			{
-				this.Id = ObjectId.GenerateNewId();
-				this.IssueId = IssueId;
-				this.AuthorId = AuthorId;
-				this.Change = Change;
-				this.DeclinedAt = DeclinedAt;
-				this.ResolvedAt = ResolvedAt;
+				Id = ObjectId.GenerateNewId();
+				IssueId = issueId;
+				AuthorId = authorId;
+				Change = change;
+				DeclinedAt = declinedAt;
+				ResolvedAt = resolvedAt;
 			}
 		}
 
@@ -224,11 +219,11 @@ namespace Horde.Build.Collections.Impl
 				Keys = new CaseInsensitiveStringSet();
 			}
 
-			public IssueFingerprint(IIssueFingerprint Fingerprint)
+			public IssueFingerprint(IIssueFingerprint fingerprint)
 			{
-				this.Type = Fingerprint.Type;
-				this.Keys = Fingerprint.Keys;
-				this.RejectKeys = Fingerprint.RejectKeys;
+				Type = fingerprint.Type;
+				Keys = fingerprint.Keys;
+				RejectKeys = fingerprint.RejectKeys;
 			}
 		}
 
@@ -253,7 +248,7 @@ namespace Horde.Build.Collections.Impl
 			public IssueFingerprint Fingerprint { get; set; }
 
 			public int MinChange { get; set; }
-			public int MaxChange { get; set; } = int.MaxValue;
+			public int MaxChange { get; set; } = Int32.MaxValue;
 
 			public IssueStep? LastSuccess { get; set; }
 
@@ -265,12 +260,12 @@ namespace Horde.Build.Collections.Impl
 
 			public IssueStep? NextSuccess { get; set; }
 
-			public bool? PromoteByDefault;
+			public bool? PromoteByDefault { get; set; }
 
 			[BsonElement("NotifySuspects"), BsonIgnoreIfDefault(false)]
-			public bool NotifySuspects_DEPRECATED { get; set; }
+			public bool NotifySuspectsDeprecated { get; set; }
 
-			bool IIssueSpan.PromoteByDefault => PromoteByDefault ?? NotifySuspects_DEPRECATED;
+			bool IIssueSpan.PromoteByDefault => PromoteByDefault ?? NotifySuspectsDeprecated;
 
 			public List<IssueSpanSuspect> Suspects { get; set; }
 			public int IssueId { get; set; }
@@ -285,37 +280,37 @@ namespace Horde.Build.Collections.Impl
 
 			private IssueSpan()
 			{
-				this.StreamName = null!;
-				this.NodeName = null!;
-				this.Fingerprint = null!;
-				this.FirstFailure = null!;
-				this.LastFailure = null!;
-				this.Suspects = new List<IssueSpanSuspect>();
+				StreamName = null!;
+				NodeName = null!;
+				Fingerprint = null!;
+				FirstFailure = null!;
+				LastFailure = null!;
+				Suspects = new List<IssueSpanSuspect>();
 			}
 
-			public IssueSpan(int IssueId, NewIssueSpanData NewSpan)
+			public IssueSpan(int issueId, NewIssueSpanData newSpan)
 			{
-				this.Id = ObjectId.GenerateNewId();
-				this.StreamId = NewSpan.StreamId;
-				this.StreamName = NewSpan.StreamName;
-				this.TemplateRefId = NewSpan.TemplateRefId;
-				this.NodeName = NewSpan.NodeName;
-				this.Fingerprint = new IssueFingerprint(NewSpan.Fingerprint);
-				if (NewSpan.LastSuccess != null)
+				Id = ObjectId.GenerateNewId();
+				StreamId = newSpan.StreamId;
+				StreamName = newSpan.StreamName;
+				TemplateRefId = newSpan.TemplateRefId;
+				NodeName = newSpan.NodeName;
+				Fingerprint = new IssueFingerprint(newSpan.Fingerprint);
+				if (newSpan.LastSuccess != null)
 				{
-					this.MinChange = NewSpan.LastSuccess.Change;
-					this.LastSuccess = new IssueStep(Id, NewSpan.LastSuccess);
+					MinChange = newSpan.LastSuccess.Change;
+					LastSuccess = new IssueStep(Id, newSpan.LastSuccess);
 				}
-				this.FirstFailure = new IssueStep(Id, NewSpan.FirstFailure);
-				this.LastFailure = new IssueStep(Id, NewSpan.FirstFailure);
-				if (NewSpan.NextSuccess != null)
+				FirstFailure = new IssueStep(Id, newSpan.FirstFailure);
+				LastFailure = new IssueStep(Id, newSpan.FirstFailure);
+				if (newSpan.NextSuccess != null)
 				{
-					this.MaxChange = NewSpan.NextSuccess.Change;
-					this.NextSuccess = new IssueStep(Id, NewSpan.NextSuccess);
+					MaxChange = newSpan.NextSuccess.Change;
+					NextSuccess = new IssueStep(Id, newSpan.NextSuccess);
 				}
-				this.PromoteByDefault = NewSpan.FirstFailure.PromoteByDefault;
-				this.Suspects = NewSpan.Suspects.ConvertAll(x => new IssueSpanSuspect(x));
-				this.IssueId = IssueId;
+				PromoteByDefault = newSpan.FirstFailure.PromoteByDefault;
+				Suspects = newSpan.Suspects.ConvertAll(x => new IssueSpanSuspect(x));
+				IssueId = issueId;
 			}
 		}
 
@@ -330,11 +325,11 @@ namespace Horde.Build.Collections.Impl
 			{
 			}
 
-			public IssueSpanSuspect(NewIssueSpanSuspectData NewSuspectData)
+			public IssueSpanSuspect(NewIssueSpanSuspectData newSuspectData)
 			{
-				this.Change = NewSuspectData.Change;
-				this.AuthorId = NewSuspectData.AuthorId;
-				this.OriginatingChange = NewSuspectData.OriginatingChange;
+				Change = newSuspectData.Change;
+				AuthorId = newSuspectData.AuthorId;
+				OriginatingChange = newSuspectData.OriginatingChange;
 			}
 		}
 
@@ -362,738 +357,724 @@ namespace Horde.Build.Collections.Impl
 
 			public LogId? LogId { get; set; }
 
-			public bool? PromoteByDefault;
+			public bool? PromoteByDefault { get; set; }
 
 			[BsonElement("NotifySuspects"), BsonIgnoreIfDefault(false)]
-			public bool NotifySuspects_DEPRECATED { get; set; }
+			public bool NotifySuspectsDeprecated { get; set; }
 
-			bool IIssueStep.PromoteByDefault => PromoteByDefault ?? NotifySuspects_DEPRECATED;
+			bool IIssueStep.PromoteByDefault => PromoteByDefault ?? NotifySuspectsDeprecated;
 
 			[BsonConstructor]
 			private IssueStep()
 			{
-				this.JobName = null!;
+				JobName = null!;
 			}
 
-			public IssueStep(ObjectId SpanId, NewIssueStepData StepData)
+			public IssueStep(ObjectId spanId, NewIssueStepData stepData)
 			{
-				this.Id = ObjectId.GenerateNewId();
-				this.SpanId = SpanId;
-				this.Change = StepData.Change;
-				this.Severity = StepData.Severity;
-				this.JobName = StepData.JobName;
-				this.JobId = StepData.JobId;
-				this.BatchId = StepData.BatchId;
-				this.StepId = StepData.StepId;
-				this.StepTime = StepData.StepTime;
-				this.LogId = StepData.LogId;
-				this.PromoteByDefault = StepData.PromoteByDefault;
+				Id = ObjectId.GenerateNewId();
+				SpanId = spanId;
+				Change = stepData.Change;
+				Severity = stepData.Severity;
+				JobName = stepData.JobName;
+				JobId = stepData.JobId;
+				BatchId = stepData.BatchId;
+				StepId = stepData.StepId;
+				StepTime = stepData.StepTime;
+				LogId = stepData.LogId;
+				PromoteByDefault = stepData.PromoteByDefault;
 			}
 		}
 
-		RedisService RedisService;
-		ISingletonDocument<IssueLedger> LedgerSingleton;
-		IMongoCollection<Issue> Issues;
-		IMongoCollection<IssueSpan> IssueSpans;
-		IMongoCollection<IssueStep> IssueSteps;
-		IMongoCollection<IssueSuspect> IssueSuspects;
-		IAuditLog<int> AuditLog;
-		ILogger Logger;
+		readonly RedisService _redisService;
+		readonly ISingletonDocument<IssueLedger> _ledgerSingleton;
+		readonly IMongoCollection<Issue> _issues;
+		readonly IMongoCollection<IssueSpan> _issueSpans;
+		readonly IMongoCollection<IssueStep> _issueSteps;
+		readonly IMongoCollection<IssueSuspect> _issueSuspects;
+		readonly IAuditLog<int> _auditLog;
+		readonly ILogger _logger;
 
-		public IssueCollection(DatabaseService DatabaseService, RedisService RedisService, IAuditLogFactory<int> AuditLogFactory, ILogger<IssueCollection> Logger)
+		public IssueCollection(DatabaseService databaseService, RedisService redisService, IAuditLogFactory<int> auditLogFactory, ILogger<IssueCollection> logger)
 		{
-			this.RedisService = RedisService;
-			this.Logger = Logger;
+			_redisService = redisService;
+			_logger = logger;
 
-			LedgerSingleton = new SingletonDocument<IssueLedger>(DatabaseService);
+			_ledgerSingleton = new SingletonDocument<IssueLedger>(databaseService);
 
-			Issues = DatabaseService.GetCollection<Issue>("IssuesV2");
-			IssueSpans = DatabaseService.GetCollection<IssueSpan>("IssuesV2.Spans");
-			IssueSteps = DatabaseService.GetCollection<IssueStep>("IssuesV2.Steps");
-			IssueSuspects = DatabaseService.GetCollection<IssueSuspect>("IssuesV2.Suspects");
-			AuditLog = AuditLogFactory.Create("IssuesV2.History", "IssueId");
+			_issues = databaseService.GetCollection<Issue>("IssuesV2");
+			_issueSpans = databaseService.GetCollection<IssueSpan>("IssuesV2.Spans");
+			_issueSteps = databaseService.GetCollection<IssueStep>("IssuesV2.Steps");
+			_issueSuspects = databaseService.GetCollection<IssueSuspect>("IssuesV2.Suspects");
+			_auditLog = auditLogFactory.Create("IssuesV2.History", "IssueId");
 
-			if (!DatabaseService.ReadOnlyMode)
+			if (!databaseService.ReadOnlyMode)
 			{
-				Issues.Indexes.CreateOne(new CreateIndexModel<Issue>(Builders<Issue>.IndexKeys.Ascending(x => x.ResolvedAt)));
-				Issues.Indexes.CreateOne(new CreateIndexModel<Issue>(Builders<Issue>.IndexKeys.Ascending(x => x.VerifiedAt)));
+				_issues.Indexes.CreateOne(new CreateIndexModel<Issue>(Builders<Issue>.IndexKeys.Ascending(x => x.ResolvedAt)));
+				_issues.Indexes.CreateOne(new CreateIndexModel<Issue>(Builders<Issue>.IndexKeys.Ascending(x => x.VerifiedAt)));
 
-				IssueSpans.Indexes.CreateOne(new CreateIndexModel<IssueSpan>(Builders<IssueSpan>.IndexKeys.Ascending(x => x.IssueId)));
-				IssueSpans.Indexes.CreateOne(new CreateIndexModel<IssueSpan>(Builders<IssueSpan>.IndexKeys.Ascending(x => x.StreamId).Ascending(x => x.MinChange).Ascending(x => x.MaxChange)));
-				IssueSpans.Indexes.CreateOne(new CreateIndexModel<IssueSpan>(Builders<IssueSpan>.IndexKeys.Ascending(x => x.StreamId).Ascending(x => x.TemplateRefId).Ascending(x => x.NodeName).Ascending(x => x.MinChange).Ascending(x => x.MaxChange), new CreateIndexOptions { Name = "StreamChanges" }));
+				_issueSpans.Indexes.CreateOne(new CreateIndexModel<IssueSpan>(Builders<IssueSpan>.IndexKeys.Ascending(x => x.IssueId)));
+				_issueSpans.Indexes.CreateOne(new CreateIndexModel<IssueSpan>(Builders<IssueSpan>.IndexKeys.Ascending(x => x.StreamId).Ascending(x => x.MinChange).Ascending(x => x.MaxChange)));
+				_issueSpans.Indexes.CreateOne(new CreateIndexModel<IssueSpan>(Builders<IssueSpan>.IndexKeys.Ascending(x => x.StreamId).Ascending(x => x.TemplateRefId).Ascending(x => x.NodeName).Ascending(x => x.MinChange).Ascending(x => x.MaxChange), new CreateIndexOptions { Name = "StreamChanges" }));
 				
-				IssueSteps.Indexes.CreateOne(new CreateIndexModel<IssueStep>(Builders<IssueStep>.IndexKeys.Ascending(x => x.SpanId)));
-				IssueSteps.Indexes.CreateOne(new CreateIndexModel<IssueStep>(Builders<IssueStep>.IndexKeys.Ascending(x => x.JobId).Ascending(x => x.BatchId).Ascending(x => x.StepId)));
+				_issueSteps.Indexes.CreateOne(new CreateIndexModel<IssueStep>(Builders<IssueStep>.IndexKeys.Ascending(x => x.SpanId)));
+				_issueSteps.Indexes.CreateOne(new CreateIndexModel<IssueStep>(Builders<IssueStep>.IndexKeys.Ascending(x => x.JobId).Ascending(x => x.BatchId).Ascending(x => x.StepId)));
 
-				IssueSuspects.Indexes.CreateOne(new CreateIndexModel<IssueSuspect>(Builders<IssueSuspect>.IndexKeys.Ascending(x => x.Change)));
-				IssueSuspects.Indexes.CreateOne(new CreateIndexModel<IssueSuspect>(Builders<IssueSuspect>.IndexKeys.Ascending(x => x.AuthorId).Ascending(x => x.ResolvedAt)));
-				IssueSuspects.Indexes.CreateOne(new CreateIndexModel<IssueSuspect>(Builders<IssueSuspect>.IndexKeys.Ascending(x => x.IssueId).Ascending(x => x.Change), new CreateIndexOptions { Unique = true }));
+				_issueSuspects.Indexes.CreateOne(new CreateIndexModel<IssueSuspect>(Builders<IssueSuspect>.IndexKeys.Ascending(x => x.Change)));
+				_issueSuspects.Indexes.CreateOne(new CreateIndexModel<IssueSuspect>(Builders<IssueSuspect>.IndexKeys.Ascending(x => x.AuthorId).Ascending(x => x.ResolvedAt)));
+				_issueSuspects.Indexes.CreateOne(new CreateIndexModel<IssueSuspect>(Builders<IssueSuspect>.IndexKeys.Ascending(x => x.IssueId).Ascending(x => x.Change), new CreateIndexOptions { Unique = true }));
 			}
 		}
 
 		/// <inheritdoc/>
 		public async Task<IAsyncDisposable> EnterCriticalSectionAsync()
 		{
-			Stopwatch Timer = Stopwatch.StartNew();
-			TimeSpan NextNotifyTime = TimeSpan.FromSeconds(2.0);
+			Stopwatch timer = Stopwatch.StartNew();
+			TimeSpan nextNotifyTime = TimeSpan.FromSeconds(2.0);
 
-			RedisLock Lock = new RedisLock(RedisService.Database, "issues/lock");
-			while (!await Lock.AcquireAsync(TimeSpan.FromMinutes(1)))
+			RedisLock issueLock = new RedisLock(_redisService.Database, "issues/lock");
+			while (!await issueLock.AcquireAsync(TimeSpan.FromMinutes(1)))
 			{
-				if (Timer.Elapsed > NextNotifyTime)
+				if (timer.Elapsed > nextNotifyTime)
 				{
-					Logger.LogWarning("Waiting on lock over issue collection for {TimeSpan}", Timer.Elapsed);
-					NextNotifyTime *= 2;
+					_logger.LogWarning("Waiting on lock over issue collection for {TimeSpan}", timer.Elapsed);
+					nextNotifyTime *= 2;
 				}
 				await Task.Delay(TimeSpan.FromMilliseconds(100));
 			}
-			return Lock;
+			return issueLock;
 		}
 
-		async Task<Issue?> TryUpdateIssueAsync(IIssue Issue, UpdateDefinition<Issue> Update)
+		async Task<Issue?> TryUpdateIssueAsync(IIssue issue, UpdateDefinition<Issue> update)
 		{
-			Issue IssueDocument = (Issue)Issue;
+			Issue issueDocument = (Issue)issue;
 
-			int PrevUpdateIndex = IssueDocument.UpdateIndex;
-			Update = Update.Set(x => x.UpdateIndex, PrevUpdateIndex + 1);
+			int prevUpdateIndex = issueDocument.UpdateIndex;
+			update = update.Set(x => x.UpdateIndex, prevUpdateIndex + 1);
 
-			FindOneAndUpdateOptions<Issue, Issue> Options = new FindOneAndUpdateOptions<Issue, Issue> { ReturnDocument = ReturnDocument.After };
-			return await Issues.FindOneAndUpdateAsync<Issue>(x => x.Id == IssueDocument.Id && x.UpdateIndex == PrevUpdateIndex, Update, Options);
+			FindOneAndUpdateOptions<Issue, Issue> options = new FindOneAndUpdateOptions<Issue, Issue> { ReturnDocument = ReturnDocument.After };
+			return await _issues.FindOneAndUpdateAsync<Issue>(x => x.Id == issueDocument.Id && x.UpdateIndex == prevUpdateIndex, update, options);
 		}
 
-		async Task<IssueSpan?> TryUpdateSpanAsync(IIssueSpan IssueSpan, UpdateDefinition<IssueSpan> Update)
+		async Task<IssueSpan?> TryUpdateSpanAsync(IIssueSpan issueSpan, UpdateDefinition<IssueSpan> update)
 		{
-			IssueSpan IssueSpanDocument = (IssueSpan)IssueSpan;
+			IssueSpan issueSpanDocument = (IssueSpan)issueSpan;
 
-			int PrevUpdateIndex = IssueSpanDocument.UpdateIndex;
-			Update = Update.Set(x => x.UpdateIndex, PrevUpdateIndex + 1);
+			int prevUpdateIndex = issueSpanDocument.UpdateIndex;
+			update = update.Set(x => x.UpdateIndex, prevUpdateIndex + 1);
 
-			FindOneAndUpdateOptions<IssueSpan, IssueSpan> Options = new FindOneAndUpdateOptions<IssueSpan, IssueSpan> { ReturnDocument = ReturnDocument.After };
-			return await IssueSpans.FindOneAndUpdateAsync<IssueSpan>(x => x.Id == IssueSpanDocument.Id && x.UpdateIndex == PrevUpdateIndex, Update, Options);
+			FindOneAndUpdateOptions<IssueSpan, IssueSpan> options = new FindOneAndUpdateOptions<IssueSpan, IssueSpan> { ReturnDocument = ReturnDocument.After };
+			return await _issueSpans.FindOneAndUpdateAsync<IssueSpan>(x => x.Id == issueSpanDocument.Id && x.UpdateIndex == prevUpdateIndex, update, options);
 		}
 
 		#region Issues
 
 		/// <inheritdoc/>
-		public async Task<IIssue> AddIssueAsync(string Summary)
+		public async Task<IIssue> AddIssueAsync(string summary)
 		{
-			IssueLedger Ledger = await LedgerSingleton.UpdateAsync(x => x.NextId++);
+			IssueLedger ledger = await _ledgerSingleton.UpdateAsync(x => x.NextId++);
 
-			Issue NewIssue = new Issue(Ledger.NextId, Summary);
-			await Issues.InsertOneAsync(NewIssue);
+			Issue newIssue = new Issue(ledger.NextId, summary);
+			await _issues.InsertOneAsync(newIssue);
 
-			ILogger IssueLogger = GetLogger(NewIssue.Id);
-			IssueLogger.LogInformation("Created issue {IssueId}", NewIssue.Id);
+			ILogger issueLogger = GetLogger(newIssue.Id);
+			issueLogger.LogInformation("Created issue {IssueId}", newIssue.Id);
 
-			return NewIssue;
+			return newIssue;
 		}
 
-		static void LogIssueChanges(ILogger IssueLogger, Issue OldIssue, Issue NewIssue)
+		static void LogIssueChanges(ILogger issueLogger, Issue oldIssue, Issue newIssue)
 		{
-			if (NewIssue.Severity != OldIssue.Severity)
+			if (newIssue.Severity != oldIssue.Severity)
 			{
-				IssueLogger.LogInformation("Changed severity to {Severity}", NewIssue.Severity);
+				issueLogger.LogInformation("Changed severity to {Severity}", newIssue.Severity);
 			}
-			if (NewIssue.Summary != OldIssue.Summary)
+			if (newIssue.Summary != oldIssue.Summary)
 			{
-				IssueLogger.LogInformation("Changed summary to \"{Summary}\"", NewIssue.Summary);
+				issueLogger.LogInformation("Changed summary to \"{Summary}\"", newIssue.Summary);
 			}
-			if (NewIssue.Description != OldIssue.Description)
+			if (newIssue.Description != oldIssue.Description)
 			{
-				IssueLogger.LogInformation("Description set to {Value}", NewIssue.Description);
+				issueLogger.LogInformation("Description set to {Value}", newIssue.Description);
 			}
-			if (((IIssue)NewIssue).Promoted != ((IIssue)OldIssue).Promoted)
+			if (((IIssue)newIssue).Promoted != ((IIssue)oldIssue).Promoted)
 			{
-				IssueLogger.LogInformation("Promoted set to {Value}", ((IIssue)NewIssue).Promoted);
+				issueLogger.LogInformation("Promoted set to {Value}", ((IIssue)newIssue).Promoted);
 			}
-			if (NewIssue.OwnerId != OldIssue.OwnerId)
+			if (newIssue.OwnerId != oldIssue.OwnerId)
 			{
-				if (NewIssue.NominatedById != null)
+				if (newIssue.NominatedById != null)
 				{
-					IssueLogger.LogInformation("User {UserId} was nominated by {NominatedByUserId}", NewIssue.OwnerId, NewIssue.NominatedById);
+					issueLogger.LogInformation("User {UserId} was nominated by {NominatedByUserId}", newIssue.OwnerId, newIssue.NominatedById);
 				}
 				else
 				{
-					IssueLogger.LogInformation("User {UserId} was nominated by default", NewIssue.OwnerId);
+					issueLogger.LogInformation("User {UserId} was nominated by default", newIssue.OwnerId);
 				}
 			}
-			if (NewIssue.AcknowledgedAt != OldIssue.AcknowledgedAt)
+			if (newIssue.AcknowledgedAt != oldIssue.AcknowledgedAt)
 			{
-				if (NewIssue.AcknowledgedAt == null)
+				if (newIssue.AcknowledgedAt == null)
 				{
-					IssueLogger.LogInformation("Issue was un-acknowledged by {UserId}", OldIssue.OwnerId);
+					issueLogger.LogInformation("Issue was un-acknowledged by {UserId}", oldIssue.OwnerId);
 				}
 				else
 				{
-					IssueLogger.LogInformation("Issue was acknowledged by {UserId}", OldIssue.OwnerId);
+					issueLogger.LogInformation("Issue was acknowledged by {UserId}", oldIssue.OwnerId);
 				}
 			}
-			if (NewIssue.FixChange != OldIssue.FixChange)
+			if (newIssue.FixChange != oldIssue.FixChange)
 			{
-				if (NewIssue.FixChange == 0)
+				if (newIssue.FixChange == 0)
 				{
-					IssueLogger.LogInformation("Issue was marked as not fixed");
+					issueLogger.LogInformation("Issue was marked as not fixed");
 				}
 				else
 				{
-					IssueLogger.LogInformation("Issue was marked as fixed in {Change}", NewIssue.FixChange);
+					issueLogger.LogInformation("Issue was marked as fixed in {Change}", newIssue.FixChange);
 				}
 			}
-			if (NewIssue.ResolvedById != OldIssue.ResolvedById)
+			if (newIssue.ResolvedById != oldIssue.ResolvedById)
 			{
-				if (NewIssue.ResolvedById == null)
+				if (newIssue.ResolvedById == null)
 				{
-					IssueLogger.LogInformation("Marking as unresolved");
+					issueLogger.LogInformation("Marking as unresolved");
 				}
 				else
 				{
-					IssueLogger.LogInformation("Resolved by {UserId}", NewIssue.ResolvedById);
+					issueLogger.LogInformation("Resolved by {UserId}", newIssue.ResolvedById);
 				}
 			}
 
-			string OldFingerprints = OldIssue.FingerprintsDesc;
-			string NewFingerprints = NewIssue.FingerprintsDesc;
-			if (OldFingerprints != NewFingerprints)
+			string oldFingerprints = oldIssue.FingerprintsDesc;
+			string newFingerprints = newIssue.FingerprintsDesc;
+			if (oldFingerprints != newFingerprints)
 			{
-				IssueLogger.LogInformation("Fingerprints changed {Fingerprints}", NewFingerprints);
+				issueLogger.LogInformation("Fingerprints changed {Fingerprints}", newFingerprints);
 			}
 
-
-			HashSet<StreamId> OldFixStreams = new HashSet<StreamId>(OldIssue.Streams.Where(x => x.ContainsFix ?? false).Select(x => x.StreamId));
-			HashSet<StreamId> NewFixStreams = new HashSet<StreamId>(NewIssue.Streams.Where(x => x.ContainsFix ?? false).Select(x => x.StreamId));
-			foreach (StreamId StreamId in NewFixStreams.Where(x => !OldFixStreams.Contains(x)))
+			HashSet<StreamId> oldFixStreams = new HashSet<StreamId>(oldIssue.Streams.Where(x => x.ContainsFix ?? false).Select(x => x.StreamId));
+			HashSet<StreamId> newFixStreams = new HashSet<StreamId>(newIssue.Streams.Where(x => x.ContainsFix ?? false).Select(x => x.StreamId));
+			foreach (StreamId streamId in newFixStreams.Where(x => !oldFixStreams.Contains(x)))
 			{
-				IssueLogger.LogInformation("Marking stream {StreamId} as fixed", StreamId);
+				issueLogger.LogInformation("Marking stream {StreamId} as fixed", streamId);
 			}
-			foreach (StreamId StreamId in OldFixStreams.Where(x => !NewFixStreams.Contains(x)))
+			foreach (StreamId streamId in oldFixStreams.Where(x => !newFixStreams.Contains(x)))
 			{
-				IssueLogger.LogInformation("Marking stream {StreamId} as not fixed", StreamId);
+				issueLogger.LogInformation("Marking stream {StreamId} as not fixed", streamId);
 			}
 		}
 
-		static void LogIssueSuspectChanges(ILogger IssueLogger, List<IssueSuspect> OldIssueSuspects, List<IssueSuspect> NewIssueSuspects)
+		static void LogIssueSuspectChanges(ILogger issueLogger, List<IssueSuspect> oldIssueSuspects, List<IssueSuspect> newIssueSuspects)
 		{
-			HashSet<(UserId, int)> OldSuspects = new HashSet<(UserId, int)>(OldIssueSuspects.Select(x => (x.AuthorId, x.Change)));
-			HashSet<(UserId, int)> NewSuspects = new HashSet<(UserId, int)>(NewIssueSuspects.Select(x => (x.AuthorId, x.Change)));
-			foreach ((UserId UserId, int Change) in NewSuspects.Where(x => !OldSuspects.Contains(x)))
+			HashSet<(UserId, int)> oldSuspects = new HashSet<(UserId, int)>(oldIssueSuspects.Select(x => (x.AuthorId, x.Change)));
+			HashSet<(UserId, int)> newSuspects = new HashSet<(UserId, int)>(newIssueSuspects.Select(x => (x.AuthorId, x.Change)));
+			foreach ((UserId userId, int change) in newSuspects.Where(x => !oldSuspects.Contains(x)))
 			{
-				IssueLogger.LogInformation("Added suspect {UserId} for change {Change}", UserId, Change);
+				issueLogger.LogInformation("Added suspect {UserId} for change {Change}", userId, change);
 			}
-			foreach ((UserId UserId, int Change) in OldSuspects.Where(x => !NewSuspects.Contains(x)))
+			foreach ((UserId userId, int change) in oldSuspects.Where(x => !newSuspects.Contains(x)))
 			{
-				IssueLogger.LogInformation("Removed suspect {UserId} for change {Change}", UserId, Change);
+				issueLogger.LogInformation("Removed suspect {UserId} for change {Change}", userId, change);
 			}
 
-			HashSet<UserId> OldDeclinedBy = new HashSet<UserId>(OldIssueSuspects.Where(x => x.DeclinedAt != null).Select(x => x.AuthorId));
-			HashSet<UserId> NewDeclinedBy = new HashSet<UserId>(NewIssueSuspects.Where(x => x.DeclinedAt != null).Select(x => x.AuthorId));
-			foreach (UserId AddDeclinedBy in NewDeclinedBy.Where(x => !OldDeclinedBy.Contains(x)))
-			{
-				IssueLogger.LogInformation("Declined by {UserId}", AddDeclinedBy);
-			}
-			foreach (UserId RemoveDeclinedBy in OldDeclinedBy.Where(x => !NewDeclinedBy.Contains(x)))
-			{
-				IssueLogger.LogInformation("Un-declined by {UserId}", RemoveDeclinedBy);
-			}
-		}
-
-		static void LogSpanInfo(ILogger IssueLogger, IIssueSpan Span)
-		{
-			IssueLogger.LogInformation("Added span {SpanId} in {StreamId}", Span.Id, Span.StreamId);
-			IssueLogger.LogInformation("Span {SpanId} first failure: CL {Change} ({JobId})", Span.Id, Span.FirstFailure.Change, Span.FirstFailure.JobId);
-			if (Span.LastSuccess != null)
+			HashSet<UserId> oldDeclinedBy = new HashSet<UserId>(oldIssueSuspects.Where(x => x.DeclinedAt != null).Select(x => x.AuthorId));
+			HashSet<UserId> newDeclinedBy = new HashSet<UserId>(newIssueSuspects.Where(x => x.DeclinedAt != null).Select(x => x.AuthorId));
+			foreach (UserId addDeclinedBy in newDeclinedBy.Where(x => !oldDeclinedBy.Contains(x)))
 			{
-				IssueLogger.LogInformation("Span {SpanId} last success: CL {Change} ({JobId})", Span.Id, Span.LastSuccess.Change, Span.LastSuccess.JobId);
+				issueLogger.LogInformation("Declined by {UserId}", addDeclinedBy);
 			}
-			if (Span.NextSuccess != null)
+			foreach (UserId removeDeclinedBy in oldDeclinedBy.Where(x => !newDeclinedBy.Contains(x)))
 			{
-				IssueLogger.LogInformation("Span {SpanId} next success: CL {Change} ({JobId})", Span.Id, Span.NextSuccess.Change, Span.NextSuccess.JobId);
+				issueLogger.LogInformation("Un-declined by {UserId}", removeDeclinedBy);
 			}
 		}
 
 		/// <inheritdoc/>
-		public async Task<IIssue?> GetIssueAsync(int IssueId)
+		public async Task<IIssue?> GetIssueAsync(int issueId)
 		{
-			Issue Issue = await Issues.Find(x => x.Id == IssueId).FirstOrDefaultAsync();
-			return Issue;
+			Issue issue = await _issues.Find(x => x.Id == issueId).FirstOrDefaultAsync();
+			return issue;
 		}
 
 		/// <inheritdoc/>
-		public Task<List<IIssueSuspect>> FindSuspectsAsync(IIssue Issue)
+		public Task<List<IIssueSuspect>> FindSuspectsAsync(IIssue issue)
 		{
-			return IssueSuspects.Find(x => x.IssueId == Issue.Id).ToListAsync<IssueSuspect, IIssueSuspect>();
+			return _issueSuspects.Find(x => x.IssueId == issue.Id).ToListAsync<IssueSuspect, IIssueSuspect>();
 		}
 
 		class ProjectedIssueId
 		{
+			[System.Diagnostics.CodeAnalysis.SuppressMessage("Style", "IDE1006:Naming Styles")]
 			public int? _id { get; set; }
 		}
 
 		/// <inheritdoc/>
-		public async Task<List<IIssue>> FindIssuesAsync(IEnumerable<int>? Ids = null, UserId? UserId = null, StreamId? StreamId = null, int? MinChange = null, int? MaxChange = null, bool? Resolved = null, bool? Promoted = null, int? Index = null, int? Count = null)
+		public async Task<List<IIssue>> FindIssuesAsync(IEnumerable<int>? ids = null, UserId? userId = null, StreamId? streamId = null, int? minChange = null, int? maxChange = null, bool? resolved = null, bool? promoted = null, int? index = null, int? count = null)
 		{
-			List<Issue> Results = await FilterIssuesByUserIdAsync(Ids, UserId, StreamId, MinChange, MaxChange, Resolved ?? false, Promoted, Index ?? 0, Count);
-			return Results.ConvertAll<IIssue>(x => x);
+			List<Issue> results = await FilterIssuesByUserIdAsync(ids, userId, streamId, minChange, maxChange, resolved ?? false, promoted, index ?? 0, count);
+			return results.ConvertAll<IIssue>(x => x);
 		}
 
-		async Task<List<Issue>> FilterIssuesByUserIdAsync(IEnumerable<int>? Ids, UserId? UserId, StreamId? StreamId, int? MinChange, int? MaxChange, bool? Resolved, bool? Promoted, int Index, int? Count)
+		async Task<List<Issue>> FilterIssuesByUserIdAsync(IEnumerable<int>? ids, UserId? userId, StreamId? streamId, int? minChange, int? maxChange, bool? resolved, bool? promoted, int index, int? count)
 		{
-			if (UserId == null)
+			if (userId == null)
 			{
-				return await FilterIssuesByStreamIdAsync(Ids, StreamId, MinChange, MaxChange, Resolved, Promoted, Index, Count);
+				return await FilterIssuesByStreamIdAsync(ids, streamId, minChange, maxChange, resolved, promoted, index, count);
 			}
 			else
 			{
-				FilterDefinition<IssueSuspect> Filter = Builders<IssueSuspect>.Filter.Eq(x => x.AuthorId, UserId);
-				if (Ids != null)
+				FilterDefinition<IssueSuspect> filter = Builders<IssueSuspect>.Filter.Eq(x => x.AuthorId, userId);
+				if (ids != null)
 				{
-					Filter &= Builders<IssueSuspect>.Filter.In(x => x.IssueId, Ids);
+					filter &= Builders<IssueSuspect>.Filter.In(x => x.IssueId, ids);
 				}
-				if (Resolved != null)
+				if (resolved != null)
 				{
-					if (Resolved.Value)
+					if (resolved.Value)
 					{
-						Filter &= Builders<IssueSuspect>.Filter.Ne(x => x.ResolvedAt, null);
+						filter &= Builders<IssueSuspect>.Filter.Ne(x => x.ResolvedAt, null);
 					}
 					else
 					{
-						Filter &= Builders<IssueSuspect>.Filter.Eq(x => x.ResolvedAt, null);
+						filter &= Builders<IssueSuspect>.Filter.Eq(x => x.ResolvedAt, null);
 					}
 				}
 
-				using (IAsyncCursor<ProjectedIssueId> Cursor = await IssueSuspects.Aggregate().Match(Filter).Group(x => x.IssueId, x => new ProjectedIssueId { _id = x.Key }).SortByDescending(x => x._id).ToCursorAsync())
+				using (IAsyncCursor<ProjectedIssueId> cursor = await _issueSuspects.Aggregate().Match(filter).Group(x => x.IssueId, x => new ProjectedIssueId { _id = x.Key }).SortByDescending(x => x._id).ToCursorAsync())
 				{
-					return await PaginatedJoinAsync(Cursor, (NextIds, NextIndex, NextCount) => FilterIssuesByStreamIdAsync(NextIds, StreamId, MinChange, MaxChange, null, Promoted, NextIndex, NextCount), Index, Count);
+					return await PaginatedJoinAsync(cursor, (nextIds, nextIndex, nextCount) => FilterIssuesByStreamIdAsync(nextIds, streamId, minChange, maxChange, null, promoted, nextIndex, nextCount), index, count);
 				}
 			}
 		}
 
-		async Task<List<Issue>> FilterIssuesByStreamIdAsync(IEnumerable<int>? Ids, StreamId? StreamId, int? MinChange, int? MaxChange, bool? Resolved, bool? Promoted, int Index, int? Count)
+		async Task<List<Issue>> FilterIssuesByStreamIdAsync(IEnumerable<int>? ids, StreamId? streamId, int? minChange, int? maxChange, bool? resolved, bool? promoted, int index, int? count)
 		{
-			if (StreamId == null)
+			if (streamId == null)
 			{
-				return await FilterIssuesByOtherFieldsAsync(Ids, MinChange, MaxChange, Resolved, Promoted, Index, Count);
+				return await FilterIssuesByOtherFieldsAsync(ids, minChange, maxChange, resolved, promoted, index, count);
 			}
 			else
 			{
-				FilterDefinition<IssueSpan> Filter = Builders<IssueSpan>.Filter.Eq(x => x.StreamId, StreamId.Value);
-				if (Ids != null)
+				FilterDefinition<IssueSpan> filter = Builders<IssueSpan>.Filter.Eq(x => x.StreamId, streamId.Value);
+				if (ids != null)
 				{
-					Filter &= Builders<IssueSpan>.Filter.In(x => x.IssueId, Ids.Select<int, int?>(x => x));
+					filter &= Builders<IssueSpan>.Filter.In(x => x.IssueId, ids.Select<int, int?>(x => x));
 				}
 				else
 				{
-					Filter &= Builders<IssueSpan>.Filter.Exists(x => x.IssueId);
+					filter &= Builders<IssueSpan>.Filter.Exists(x => x.IssueId);
 				}
 
-				if (MinChange != null)
+				if (minChange != null)
 				{
-					Filter &= Builders<IssueSpan>.Filter.Not(Builders<IssueSpan>.Filter.Lt(x => x.MaxChange, MinChange.Value));
+					filter &= Builders<IssueSpan>.Filter.Not(Builders<IssueSpan>.Filter.Lt(x => x.MaxChange, minChange.Value));
 				}
-				if (MaxChange != null)
+				if (maxChange != null)
 				{
-					Filter &= Builders<IssueSpan>.Filter.Not(Builders<IssueSpan>.Filter.Gt(x => x.MinChange, MaxChange.Value));
+					filter &= Builders<IssueSpan>.Filter.Not(Builders<IssueSpan>.Filter.Gt(x => x.MinChange, maxChange.Value));
 				}
 
-				if (Resolved != null)
+				if (resolved != null)
 				{
-					if (Resolved.Value)
+					if (resolved.Value)
 					{
-						Filter &= Builders<IssueSpan>.Filter.Ne(x => x.ResolvedAt, null);
+						filter &= Builders<IssueSpan>.Filter.Ne(x => x.ResolvedAt, null);
 					}
 					else
 					{
-						Filter &= Builders<IssueSpan>.Filter.Eq(x => x.ResolvedAt, null);
+						filter &= Builders<IssueSpan>.Filter.Eq(x => x.ResolvedAt, null);
 					}
 				}
 
-				using (IAsyncCursor<ProjectedIssueId> Cursor = await IssueSpans.Aggregate().Match(Filter).Group(x => x.IssueId, x => new ProjectedIssueId { _id = x.Key }).SortByDescending(x => x._id).ToCursorAsync())
+				using (IAsyncCursor<ProjectedIssueId> cursor = await _issueSpans.Aggregate().Match(filter).Group(x => x.IssueId, x => new ProjectedIssueId { _id = x.Key }).SortByDescending(x => x._id).ToCursorAsync())
 				{
-					List<Issue> Results = await PaginatedJoinAsync(Cursor, (NextIds, NextIndex, NextCount) => FilterIssuesByOtherFieldsAsync(NextIds, null, null, null, Promoted, NextIndex, NextCount), Index, Count);
-					if (Resolved != null)
+					List<Issue> results = await PaginatedJoinAsync(cursor, (nextIds, nextIndex, nextCount) => FilterIssuesByOtherFieldsAsync(nextIds, null, null, null, promoted, nextIndex, nextCount), index, count);
+					if (resolved != null)
 					{
-						for (int Idx = Results.Count - 1; Idx >= 0; Idx--)
+						for (int idx = results.Count - 1; idx >= 0; idx--)
 						{
-							Issue Issue = Results[Idx];
-							if ((Issue.ResolvedAt != null) != Resolved.Value)
+							Issue issue = results[idx];
+							if ((issue.ResolvedAt != null) != resolved.Value)
 							{
-								Logger.LogWarning("Issue {IssueId} has resolved state out of sync with spans", Issue.Id);
-								Results.RemoveAt(Idx);
+								_logger.LogWarning("Issue {IssueId} has resolved state out of sync with spans", issue.Id);
+								results.RemoveAt(idx);
 							}
 						}
 					}
-					return Results;
+					return results;
 				}
 			}
 		}
 
-		async Task<List<Issue>> FilterIssuesByOtherFieldsAsync(IEnumerable<int>? Ids, int? MinChange, int? MaxChange, bool? Resolved, bool? Promoted, int Index, int? Count)
+		async Task<List<Issue>> FilterIssuesByOtherFieldsAsync(IEnumerable<int>? ids, int? minChange, int? maxChange, bool? resolved, bool? promoted, int index, int? count)
 		{
-			FilterDefinition<Issue> Filter = FilterDefinition<Issue>.Empty;
-			if (Ids != null)
+			FilterDefinition<Issue> filter = FilterDefinition<Issue>.Empty;
+			if (ids != null)
 			{
-				Filter &= Builders<Issue>.Filter.In(x => x.Id, Ids);
+				filter &= Builders<Issue>.Filter.In(x => x.Id, ids);
 			}
-			if (Resolved != null)
+			if (resolved != null)
 			{
-				if (Resolved.Value)
+				if (resolved.Value)
 				{
-					Filter &= Builders<Issue>.Filter.Ne(x => x.ResolvedAt, null);
+					filter &= Builders<Issue>.Filter.Ne(x => x.ResolvedAt, null);
 				}
 				else
 				{
-					Filter &= Builders<Issue>.Filter.Eq(x => x.ResolvedAt, null);
+					filter &= Builders<Issue>.Filter.Eq(x => x.ResolvedAt, null);
 				}
 			}
-			if (MinChange != null)
+			if (minChange != null)
 			{
-				Filter &= Builders<Issue>.Filter.Not(Builders<Issue>.Filter.Lt(x => x.MaxSuspectChange, MinChange.Value));
+				filter &= Builders<Issue>.Filter.Not(Builders<Issue>.Filter.Lt(x => x.MaxSuspectChange, minChange.Value));
 			}
-			if (MaxChange != null)
+			if (maxChange != null)
 			{
-				Filter &= Builders<Issue>.Filter.Not(Builders<Issue>.Filter.Gt(x => x.MinSuspectChange, MaxChange.Value));
+				filter &= Builders<Issue>.Filter.Not(Builders<Issue>.Filter.Gt(x => x.MinSuspectChange, maxChange.Value));
 			}
-			if (Promoted != null)
+			if (promoted != null)
 			{
-				if (Promoted.Value)
+				if (promoted.Value)
 				{
-					Filter &= Builders<Issue>.Filter.Eq(x => x.Promoted, true);
+					filter &= Builders<Issue>.Filter.Eq(x => x.Promoted, true);
 				}
 				else
 				{
-					Filter &= Builders<Issue>.Filter.Ne(x => x.Promoted, true); // Handle the field not existing as well as being set to false.
+					filter &= Builders<Issue>.Filter.Ne(x => x.Promoted, true); // Handle the field not existing as well as being set to false.
 				}
 			}
-			return await Issues.Find(Filter).SortByDescending(x => x.Id).Range(Index, Count).ToListAsync();
+			return await _issues.Find(filter).SortByDescending(x => x.Id).Range(index, count).ToListAsync();
 		}
 
 		/// <summary>
 		/// Performs a client-side join of a filtered set of issues against another query
 		/// </summary>
-		/// <param name="Cursor"></param>
-		/// <param name="NextStageFunc"></param>
-		/// <param name="Index"></param>
-		/// <param name="Count"></param>
+		/// <param name="cursor"></param>
+		/// <param name="nextStageFunc"></param>
+		/// <param name="index"></param>
+		/// <param name="count"></param>
 		/// <returns></returns>
-		static async Task<List<Issue>> PaginatedJoinAsync(IAsyncCursor<ProjectedIssueId> Cursor, Func<IEnumerable<int>, int, int?, Task<List<Issue>>> NextStageFunc, int Index, int? Count)
+		static async Task<List<Issue>> PaginatedJoinAsync(IAsyncCursor<ProjectedIssueId> cursor, Func<IEnumerable<int>, int, int?, Task<List<Issue>>> nextStageFunc, int index, int? count)
 		{
-			if (Count == null)
+			if (count == null)
 			{
-				List<ProjectedIssueId> IssueIds = await Cursor.ToListAsync();
-				return await NextStageFunc(IssueIds.Where(x => x._id != null).Select(x => x._id!.Value), Index, null);
+				List<ProjectedIssueId> issueIds = await cursor.ToListAsync();
+				return await nextStageFunc(issueIds.Where(x => x._id != null).Select(x => x._id!.Value), index, null);
 			}
 			else
 			{
-				List<Issue> Results = new List<Issue>();
-				while (await Cursor.MoveNextAsync() && Results.Count < Count.Value)
+				List<Issue> results = new List<Issue>();
+				while (await cursor.MoveNextAsync() && results.Count < count.Value)
 				{
-					List<Issue> NextResults = await NextStageFunc(Cursor.Current.Where(x => x._id != null).Select(x => x._id!.Value), 0, Count.Value - Results.Count);
-					int RemoveCount = Math.Min(Index, NextResults.Count);
-					NextResults.RemoveRange(0, RemoveCount);
-					Index -= RemoveCount;
-					Results.AddRange(NextResults);
+					List<Issue> nextResults = await nextStageFunc(cursor.Current.Where(x => x._id != null).Select(x => x._id!.Value), 0, count.Value - results.Count);
+					int removeCount = Math.Min(index, nextResults.Count);
+					nextResults.RemoveRange(0, removeCount);
+					index -= removeCount;
+					results.AddRange(nextResults);
 				}
-				return Results;
+				return results;
 			}
 		}
 
 		/// <inheritdoc/>
-		public async Task<List<IIssue>> FindIssuesForChangesAsync(List<int> Changes)
+		public async Task<List<IIssue>> FindIssuesForChangesAsync(List<int> changes)
 		{
-			List<int> IssueIds = await (await IssueSuspects.DistinctAsync(x => x.IssueId, Builders<IssueSuspect>.Filter.In(x => x.Change, Changes))).ToListAsync();
-			return await Issues.Find(Builders<Issue>.Filter.In(x => x.Id, IssueIds)).ToListAsync<Issue, IIssue>();
+			List<int> issueIds = await (await _issueSuspects.DistinctAsync(x => x.IssueId, Builders<IssueSuspect>.Filter.In(x => x.Change, changes))).ToListAsync();
+			return await _issues.Find(Builders<Issue>.Filter.In(x => x.Id, issueIds)).ToListAsync<Issue, IIssue>();
 		}
 
 		/// <inheritdoc/>
-		public async Task<IIssue?> TryUpdateIssueAsync(IIssue Issue, IssueSeverity? NewSeverity = null, string? NewSummary = null, string? NewUserSummary = null, string? NewDescription = null, bool? NewManuallyPromoted = null, UserId? NewOwnerId = null, UserId? NewNominatedById = null, bool? NewAcknowledged = null, UserId? NewDeclinedById = null, int? NewFixChange = null, UserId? NewResolvedById = null, List<ObjectId>? NewExcludeSpanIds = null, DateTime? NewLastSeenAt = null)
+		public async Task<IIssue?> TryUpdateIssueAsync(IIssue issue, IssueSeverity? newSeverity = null, string? newSummary = null, string? newUserSummary = null, string? newDescription = null, bool? newManuallyPromoted = null, UserId? newOwnerId = null, UserId? newNominatedById = null, bool? newAcknowledged = null, UserId? newDeclinedById = null, int? newFixChange = null, UserId? newResolvedById = null, List<ObjectId>? newExcludeSpanIds = null, DateTime? newLastSeenAt = null)
 		{
-			Issue IssueDocument = (Issue)Issue;
+			Issue issueDocument = (Issue)issue;
 
-			DateTime UtcNow = DateTime.UtcNow;
+			DateTime utcNow = DateTime.UtcNow;
 
-			List<UpdateDefinition<Issue>> Updates = new List<UpdateDefinition<Issue>>();
-			if (NewSeverity != null)
+			List<UpdateDefinition<Issue>> updates = new List<UpdateDefinition<Issue>>();
+			if (newSeverity != null)
 			{
-				Updates.Add(Builders<Issue>.Update.Set(x => x.Severity, NewSeverity.Value));
+				updates.Add(Builders<Issue>.Update.Set(x => x.Severity, newSeverity.Value));
 			}
-			if (NewSummary != null)
+			if (newSummary != null)
 			{
-				Updates.Add(Builders<Issue>.Update.Set(x => x.Summary, NewSummary));
+				updates.Add(Builders<Issue>.Update.Set(x => x.Summary, newSummary));
 			}
-			if (NewUserSummary != null)
+			if (newUserSummary != null)
 			{
-				if (NewUserSummary.Length == 0)
+				if (newUserSummary.Length == 0)
 				{
-					Updates.Add(Builders<Issue>.Update.Unset(x => x.UserSummary!));
+					updates.Add(Builders<Issue>.Update.Unset(x => x.UserSummary!));
 				}
 				else
 				{
-					Updates.Add(Builders<Issue>.Update.Set(x => x.UserSummary, NewUserSummary));
+					updates.Add(Builders<Issue>.Update.Set(x => x.UserSummary, newUserSummary));
 				}
 			}
-			if (NewDescription != null)
+			if (newDescription != null)
 			{
-				if (NewDescription.Length == 0)
+				if (newDescription.Length == 0)
 				{
-					Updates.Add(Builders<Issue>.Update.Unset(x => x.Description));
+					updates.Add(Builders<Issue>.Update.Unset(x => x.Description));
 				}
 				else
 				{
-					Updates.Add(Builders<Issue>.Update.Set(x => x.Description, NewDescription));
+					updates.Add(Builders<Issue>.Update.Set(x => x.Description, newDescription));
 				}
 			}
-			if (NewManuallyPromoted != null)
+			if (newManuallyPromoted != null)
 			{
-				Updates.Add(Builders<Issue>.Update.Set(x => x.ManuallyPromoted, NewManuallyPromoted.Value));
+				updates.Add(Builders<Issue>.Update.Set(x => x.ManuallyPromoted, newManuallyPromoted.Value));
 			}
-			if (NewOwnerId != null)
+			if (newOwnerId != null)
 			{
-				if (NewOwnerId.Value == UserId.Empty)
+				if (newOwnerId.Value == UserId.Empty)
 				{
-					Updates.Add(Builders<Issue>.Update.Unset(x => x.OwnerId!));
-					Updates.Add(Builders<Issue>.Update.Unset(x => x.NominatedAt!));
-					Updates.Add(Builders<Issue>.Update.Unset(x => x.NominatedById!));
+					updates.Add(Builders<Issue>.Update.Unset(x => x.OwnerId!));
+					updates.Add(Builders<Issue>.Update.Unset(x => x.NominatedAt!));
+					updates.Add(Builders<Issue>.Update.Unset(x => x.NominatedById!));
 				}
 				else
 				{
-					Updates.Add(Builders<Issue>.Update.Set(x => x.OwnerId!, NewOwnerId.Value));
+					updates.Add(Builders<Issue>.Update.Set(x => x.OwnerId!, newOwnerId.Value));
 
-					Updates.Add(Builders<Issue>.Update.Set(x => x.NominatedAt, DateTime.UtcNow));
-					if (NewNominatedById == null)
+					updates.Add(Builders<Issue>.Update.Set(x => x.NominatedAt, DateTime.UtcNow));
+					if (newNominatedById == null)
 					{
-						Updates.Add(Builders<Issue>.Update.Unset(x => x.NominatedById!));
+						updates.Add(Builders<Issue>.Update.Unset(x => x.NominatedById!));
 					}
 					else
 					{
-						Updates.Add(Builders<Issue>.Update.Set(x => x.NominatedById, NewNominatedById.Value));
+						updates.Add(Builders<Issue>.Update.Set(x => x.NominatedById, newNominatedById.Value));
 					}
-					NewAcknowledged = NewAcknowledged ?? false;
+					newAcknowledged ??= false;
 				}
 			}
-			if (NewAcknowledged != null)
+			if (newAcknowledged != null)
 			{
-				if (NewAcknowledged.Value)
+				if (newAcknowledged.Value)
 				{
-					if (IssueDocument.AcknowledgedAt == null)
+					if (issueDocument.AcknowledgedAt == null)
 					{
-						Updates.Add(Builders<Issue>.Update.Set(x => x.AcknowledgedAt, UtcNow));
+						updates.Add(Builders<Issue>.Update.Set(x => x.AcknowledgedAt, utcNow));
 					}
 				}
 				else
 				{
-					if (IssueDocument.AcknowledgedAt != null)
+					if (issueDocument.AcknowledgedAt != null)
 					{
-						Updates.Add(Builders<Issue>.Update.Unset(x => x.AcknowledgedAt!));
+						updates.Add(Builders<Issue>.Update.Unset(x => x.AcknowledgedAt!));
 					}
 				}
 			}
-			if (NewFixChange != null)
+			if (newFixChange != null)
 			{
-				if (NewFixChange == 0)
+				if (newFixChange == 0)
 				{
-					Updates.Add(Builders<Issue>.Update.Unset(x => x.FixChange!));
+					updates.Add(Builders<Issue>.Update.Unset(x => x.FixChange!));
 				}
 				else
 				{
-					Updates.Add(Builders<Issue>.Update.Set(x => x.FixChange, NewFixChange));
+					updates.Add(Builders<Issue>.Update.Set(x => x.FixChange, newFixChange));
 				}
 			}
-			if (NewResolvedById != null)
+			if (newResolvedById != null)
 			{
-				if (NewResolvedById.Value != UserId.Empty)
+				if (newResolvedById.Value != UserId.Empty)
 				{
-					if (IssueDocument.ResolvedAt == null || IssueDocument.ResolvedById != NewResolvedById)
+					if (issueDocument.ResolvedAt == null || issueDocument.ResolvedById != newResolvedById)
 					{
-						Updates.Add(Builders<Issue>.Update.Set(x => x.ResolvedAt, UtcNow));
-						Updates.Add(Builders<Issue>.Update.Set(x => x.ResolvedById, NewResolvedById.Value));
+						updates.Add(Builders<Issue>.Update.Set(x => x.ResolvedAt, utcNow));
+						updates.Add(Builders<Issue>.Update.Set(x => x.ResolvedById, newResolvedById.Value));
 					}
 				}
 				else
 				{
-					if (IssueDocument.ResolvedAt != null)
+					if (issueDocument.ResolvedAt != null)
 					{
-						Updates.Add(Builders<Issue>.Update.Unset(x => x.ResolvedAt!));
+						updates.Add(Builders<Issue>.Update.Unset(x => x.ResolvedAt!));
 					}
-					if (IssueDocument.ResolvedById != null)
+					if (issueDocument.ResolvedById != null)
 					{
-						Updates.Add(Builders<Issue>.Update.Unset(x => x.ResolvedById!));
+						updates.Add(Builders<Issue>.Update.Unset(x => x.ResolvedById!));
 					}
 				}
 			}
-			if (NewExcludeSpanIds != null)
+			if (newExcludeSpanIds != null)
 			{
-				List<ObjectId> NewCombinedExcludeSpanIds = NewExcludeSpanIds;
-				if (Issue.ExcludeSpans != null)
+				List<ObjectId> newCombinedExcludeSpanIds = newExcludeSpanIds;
+				if (issue.ExcludeSpans != null)
 				{
-					NewCombinedExcludeSpanIds = NewCombinedExcludeSpanIds.Union(Issue.ExcludeSpans).ToList();
+					newCombinedExcludeSpanIds = newCombinedExcludeSpanIds.Union(issue.ExcludeSpans).ToList();
 				}
-				Updates.Add(Builders<Issue>.Update.Set(x => x.ExcludeSpans, NewCombinedExcludeSpanIds));
+				updates.Add(Builders<Issue>.Update.Set(x => x.ExcludeSpans, newCombinedExcludeSpanIds));
 			}
-			if (NewLastSeenAt != null)
+			if (newLastSeenAt != null)
 			{
-				Updates.Add(Builders<Issue>.Update.Set(x => x.LastSeenAt, NewLastSeenAt.Value));
+				updates.Add(Builders<Issue>.Update.Set(x => x.LastSeenAt, newLastSeenAt.Value));
 			}
 
-			if (NewDeclinedById != null)
+			if (newDeclinedById != null)
 			{
-				await IssueSuspects.UpdateManyAsync(x => x.IssueId == Issue.Id && x.AuthorId == NewDeclinedById.Value, Builders<IssueSuspect>.Update.Set(x => x.DeclinedAt, DateTime.UtcNow));
+				await _issueSuspects.UpdateManyAsync(x => x.IssueId == issue.Id && x.AuthorId == newDeclinedById.Value, Builders<IssueSuspect>.Update.Set(x => x.DeclinedAt, DateTime.UtcNow));
 			}
-			if (Updates.Count == 0)
+			if (updates.Count == 0)
 			{
-				return IssueDocument;
+				return issueDocument;
 			}
 
-			Issue? NewIssue = await TryUpdateIssueAsync(Issue, Builders<Issue>.Update.Combine(Updates));
-			if(NewIssue == null)
+			Issue? newIssue = await TryUpdateIssueAsync(issue, Builders<Issue>.Update.Combine(updates));
+			if(newIssue == null)
 			{
 				return null;
 			}
 
-			ILogger IssueLogger = GetLogger(Issue.Id);
-			LogIssueChanges(IssueLogger, IssueDocument, NewIssue);
-			return NewIssue;
+			ILogger issueLogger = GetLogger(issue.Id);
+			LogIssueChanges(issueLogger, issueDocument, newIssue);
+			return newIssue;
 		}
 
 		/// <inheritdoc/>
-		public async Task<IIssue?> TryUpdateIssueDerivedDataAsync(IIssue Issue, string NewSummary, IssueSeverity NewSeverity, List<NewIssueFingerprint> NewFingerprints, List<NewIssueStream> NewStreams, List<NewIssueSuspectData> NewSuspects, DateTime? NewResolvedAt, DateTime? NewVerifiedAt, DateTime NewLastSeenAt)
+		public async Task<IIssue?> TryUpdateIssueDerivedDataAsync(IIssue issue, string newSummary, IssueSeverity newSeverity, List<NewIssueFingerprint> newFingerprints, List<NewIssueStream> newStreams, List<NewIssueSuspectData> newSuspects, DateTime? newResolvedAt, DateTime? newVerifiedAt, DateTime newLastSeenAt)
 		{
-			Issue IssueImpl = (Issue)Issue;
+			Issue issueImpl = (Issue)issue;
 
 			// Update all the suspects for this issue
-			List<IssueSuspect> OldSuspectImpls = await IssueSuspects.Find(x => x.IssueId == Issue.Id).ToListAsync();
-			List<IssueSuspect> NewSuspectImpls = await UpdateIssueSuspectsAsync(Issue.Id, OldSuspectImpls, NewSuspects, NewResolvedAt);
+			List<IssueSuspect> oldSuspectImpls = await _issueSuspects.Find(x => x.IssueId == issue.Id).ToListAsync();
+			List<IssueSuspect> newSuspectImpls = await UpdateIssueSuspectsAsync(issue.Id, oldSuspectImpls, newSuspects, newResolvedAt);
 
 			// Find the spans for this issue
-			List<IssueSpan> NewSpans = await IssueSpans.Find(x => x.IssueId == Issue.Id).ToListAsync();
+			List<IssueSpan> newSpans = await _issueSpans.Find(x => x.IssueId == issue.Id).ToListAsync();
 
 			// Update the resolved time on any issues
-			List<ObjectId> UpdateSpanIds = NewSpans.Where(x => x.ResolvedAt != NewResolvedAt).Select(x => x.Id).ToList();
-			if (UpdateSpanIds.Count > 0)
+			List<ObjectId> updateSpanIds = newSpans.Where(x => x.ResolvedAt != newResolvedAt).Select(x => x.Id).ToList();
+			if (updateSpanIds.Count > 0)
 			{
-				FilterDefinition<IssueSpan> Filter = Builders<IssueSpan>.Filter.In(x => x.Id, UpdateSpanIds);
-				await IssueSpans.UpdateManyAsync(Filter, Builders<IssueSpan>.Update.Set(x => x.ResolvedAt, NewResolvedAt));
+				FilterDefinition<IssueSpan> filter = Builders<IssueSpan>.Filter.In(x => x.Id, updateSpanIds);
+				await _issueSpans.UpdateManyAsync(filter, Builders<IssueSpan>.Update.Set(x => x.ResolvedAt, newResolvedAt));
 			}
 
 			// Figure out if this issue should be promoted
-			bool NewPromoted;
-			if (IssueImpl.ManuallyPromoted.HasValue)
+			bool newPromoted;
+			if (issueImpl.ManuallyPromoted.HasValue)
 			{
-				NewPromoted = IssueImpl.ManuallyPromoted.Value;
+				newPromoted = issueImpl.ManuallyPromoted.Value;
 			}
-			else if (IssueImpl.ManuallyPromoted_DEPRECATED.HasValue)
+			else if (issueImpl.ManuallyPromotedDeprecated.HasValue)
 			{
-				NewPromoted = IssueImpl.ManuallyPromoted_DEPRECATED.Value;
+				newPromoted = issueImpl.ManuallyPromotedDeprecated.Value;
 			}
 			else
 			{
-				NewPromoted = NewSpans.Any(x => ((IIssueSpan)x).PromoteByDefault);
+				newPromoted = newSpans.Any(x => ((IIssueSpan)x).PromoteByDefault);
 			}
 
 			// Find the default owner
-			UserId? NewDefaultOwnerId = null;
-			if (NewPromoted && NewSuspectImpls.Count > 0)
+			UserId? newDefaultOwnerId = null;
+			if (newPromoted && newSuspectImpls.Count > 0)
 			{
-				UserId PossibleOwnerId = NewSuspectImpls[0].AuthorId;
-				if (NewSuspectImpls.All(x => x.AuthorId == PossibleOwnerId) && NewSuspectImpls.Any(x => x.DeclinedAt == null))
+				UserId possibleOwnerId = newSuspectImpls[0].AuthorId;
+				if (newSuspectImpls.All(x => x.AuthorId == possibleOwnerId) && newSuspectImpls.Any(x => x.DeclinedAt == null))
 				{
-					NewDefaultOwnerId = PossibleOwnerId;
+					newDefaultOwnerId = possibleOwnerId;
 				}
 			}
 
 			// Get the range of suspect changes
-			int NewMinSuspectChange = (NewSuspects.Count > 0) ? NewSuspects.Min(x => x.Change) : 0;
-			int NewMaxSuspectChange = (NewSuspects.Count > 0) ? NewSuspects.Min(x => x.Change) : 0;
+			int newMinSuspectChange = (newSuspects.Count > 0) ? newSuspects.Min(x => x.Change) : 0;
+			int newMaxSuspectChange = (newSuspects.Count > 0) ? newSuspects.Min(x => x.Change) : 0;
 
 			// Perform the actual update with this data
-			List<UpdateDefinition<Issue>> Updates = new List<UpdateDefinition<Issue>>();
-			if (!String.Equals(Issue.Summary, NewSummary, StringComparison.Ordinal))
+			List<UpdateDefinition<Issue>> updates = new List<UpdateDefinition<Issue>>();
+			if (!String.Equals(issue.Summary, newSummary, StringComparison.Ordinal))
 			{
-				Updates.Add(Builders<Issue>.Update.Set(x => x.Summary, NewSummary));
+				updates.Add(Builders<Issue>.Update.Set(x => x.Summary, newSummary));
 			}
-			if (Issue.Severity != NewSeverity)
+			if (issue.Severity != newSeverity)
 			{
-				Updates.Add(Builders<Issue>.Update.Set(x => x.Severity, NewSeverity));
+				updates.Add(Builders<Issue>.Update.Set(x => x.Severity, newSeverity));
 			}
-			if (Issue.Promoted != NewPromoted)
+			if (issue.Promoted != newPromoted)
 			{
-				Updates.Add(Builders<Issue>.Update.Set(x => x.Promoted, NewPromoted));
+				updates.Add(Builders<Issue>.Update.Set(x => x.Promoted, newPromoted));
 			}
-			if (Issue.Fingerprints.Count != NewFingerprints.Count || !NewFingerprints.Zip(Issue.Fingerprints).All(x => x.First.Equals(x.Second)))
+			if (issue.Fingerprints.Count != newFingerprints.Count || !newFingerprints.Zip(issue.Fingerprints).All(x => x.First.Equals(x.Second)))
 			{
-				Updates.Add(Builders<Issue>.Update.Set(x => x.Fingerprints, NewFingerprints.Select(x => new IssueFingerprint(x))));
+				updates.Add(Builders<Issue>.Update.Set(x => x.Fingerprints, newFingerprints.Select(x => new IssueFingerprint(x))));
 			}
-			if (Issue.Streams.Count != NewStreams.Count || !NewStreams.Zip(Issue.Streams).All(x => x.First.StreamId == x.Second.StreamId && x.First.ContainsFix == x.Second.ContainsFix))
+			if (issue.Streams.Count != newStreams.Count || !newStreams.Zip(issue.Streams).All(x => x.First.StreamId == x.Second.StreamId && x.First.ContainsFix == x.Second.ContainsFix))
 			{
-				Updates.Add(Builders<Issue>.Update.Set(x => x.Streams, NewStreams.Select(x => new IssueStream(x))));
+				updates.Add(Builders<Issue>.Update.Set(x => x.Streams, newStreams.Select(x => new IssueStream(x))));
 			}
-			if (IssueImpl.MinSuspectChange != NewMinSuspectChange)
+			if (issueImpl.MinSuspectChange != newMinSuspectChange)
 			{
-				Updates.Add(Builders<Issue>.Update.Set(x => x.MinSuspectChange, NewMinSuspectChange));
+				updates.Add(Builders<Issue>.Update.Set(x => x.MinSuspectChange, newMinSuspectChange));
 			}
-			if (IssueImpl.MaxSuspectChange != NewMaxSuspectChange)
+			if (issueImpl.MaxSuspectChange != newMaxSuspectChange)
 			{
-				Updates.Add(Builders<Issue>.Update.Set(x => x.MaxSuspectChange, NewMaxSuspectChange));
+				updates.Add(Builders<Issue>.Update.Set(x => x.MaxSuspectChange, newMaxSuspectChange));
 			}
-			if (IssueImpl.DefaultOwnerId != NewDefaultOwnerId)
+			if (issueImpl.DefaultOwnerId != newDefaultOwnerId)
 			{
-				Updates.Add(Builders<Issue>.Update.Set(x => x.DefaultOwnerId, NewDefaultOwnerId));
+				updates.Add(Builders<Issue>.Update.Set(x => x.DefaultOwnerId, newDefaultOwnerId));
 			}
-			if (Issue.ResolvedAt != NewResolvedAt)
+			if (issue.ResolvedAt != newResolvedAt)
 			{
-				Updates.Add(Builders<Issue>.Update.SetOrUnsetNull(x => x.ResolvedAt, NewResolvedAt));
+				updates.Add(Builders<Issue>.Update.SetOrUnsetNull(x => x.ResolvedAt, newResolvedAt));
 			}
-			if (NewResolvedAt == null && Issue.ResolvedById != null)
+			if (newResolvedAt == null && issue.ResolvedById != null)
 			{
-				Updates.Add(Builders<Issue>.Update.Unset(x => x.ResolvedById));
+				updates.Add(Builders<Issue>.Update.Unset(x => x.ResolvedById));
 			}
-			if (Issue.VerifiedAt != NewVerifiedAt)
+			if (issue.VerifiedAt != newVerifiedAt)
 			{
-				Updates.Add(Builders<Issue>.Update.SetOrUnsetNull(x => x.VerifiedAt, NewVerifiedAt));
+				updates.Add(Builders<Issue>.Update.SetOrUnsetNull(x => x.VerifiedAt, newVerifiedAt));
 			}
-			if (Issue.LastSeenAt != NewLastSeenAt)
+			if (issue.LastSeenAt != newLastSeenAt)
 			{
-				Updates.Add(Builders<Issue>.Update.Set(x => x.LastSeenAt, NewLastSeenAt));
+				updates.Add(Builders<Issue>.Update.Set(x => x.LastSeenAt, newLastSeenAt));
 			}
 
-			Issue? NewIssue = await TryUpdateIssueAsync(Issue, Builders<Issue>.Update.Combine(Updates));
-			if(NewIssue != null)
+			Issue? newIssue = await TryUpdateIssueAsync(issue, Builders<Issue>.Update.Combine(updates));
+			if(newIssue != null)
 			{
-				ILogger IssueLogger = GetLogger(Issue.Id);
-				LogIssueChanges(GetLogger(Issue.Id), IssueImpl, NewIssue);
-				LogIssueSuspectChanges(IssueLogger, OldSuspectImpls, NewSuspectImpls);
-				return NewIssue;
+				ILogger issueLogger = GetLogger(issue.Id);
+				LogIssueChanges(GetLogger(issue.Id), issueImpl, newIssue);
+				LogIssueSuspectChanges(issueLogger, oldSuspectImpls, newSuspectImpls);
+				return newIssue;
 			}
 			return null;
 		}
 
-		async Task<List<IssueSuspect>> UpdateIssueSuspectsAsync(int IssueId, List<IssueSuspect> OldSuspectImpls, List<NewIssueSuspectData> NewSuspects, DateTime? ResolvedAt)
+		async Task<List<IssueSuspect>> UpdateIssueSuspectsAsync(int issueId, List<IssueSuspect> oldSuspectImpls, List<NewIssueSuspectData> newSuspects, DateTime? resolvedAt)
 		{
-			List<IssueSuspect> NewSuspectImpls = new List<IssueSuspect>(OldSuspectImpls);
+			List<IssueSuspect> newSuspectImpls = new List<IssueSuspect>(oldSuspectImpls);
 
 			// Find the current list of suspects
-			HashSet<(UserId, int)> CurSuspectKeys = new HashSet<(UserId, int)>(OldSuspectImpls.Select(x => (x.AuthorId, x.Change)));
-			List<IssueSuspect> CreateSuspects = NewSuspects.Where(x => !CurSuspectKeys.Contains((x.AuthorId, x.Change))).Select(x => new IssueSuspect(IssueId, x, ResolvedAt)).ToList();
+			HashSet<(UserId, int)> curSuspectKeys = new HashSet<(UserId, int)>(oldSuspectImpls.Select(x => (x.AuthorId, x.Change)));
+			List<IssueSuspect> createSuspects = newSuspects.Where(x => !curSuspectKeys.Contains((x.AuthorId, x.Change))).Select(x => new IssueSuspect(issueId, x, resolvedAt)).ToList();
 
-			HashSet<(UserId, int)> NewSuspectKeys = new HashSet<(UserId, int)>(NewSuspects.Select(x => (x.AuthorId, x.Change)));
-			List<IssueSuspect> DeleteSuspects = OldSuspectImpls.Where(x => !NewSuspectKeys.Contains((x.AuthorId, x.Change))).ToList();
+			HashSet<(UserId, int)> newSuspectKeys = new HashSet<(UserId, int)>(newSuspects.Select(x => (x.AuthorId, x.Change)));
+			List<IssueSuspect> deleteSuspects = oldSuspectImpls.Where(x => !newSuspectKeys.Contains((x.AuthorId, x.Change))).ToList();
 
 			// Apply the suspect changes
-			if (CreateSuspects.Count > 0)
+			if (createSuspects.Count > 0)
 			{
-				await IssueSuspects.InsertManyIgnoreDuplicatesAsync(CreateSuspects);
-				NewSuspectImpls.AddRange(CreateSuspects);
+				await _issueSuspects.InsertManyIgnoreDuplicatesAsync(createSuspects);
+				newSuspectImpls.AddRange(createSuspects);
 			}
-			if (DeleteSuspects.Count > 0)
+			if (deleteSuspects.Count > 0)
 			{
-				await IssueSuspects.DeleteManyAsync(Builders<IssueSuspect>.Filter.In(x => x.Id, DeleteSuspects.Select(y => y.Id)));
-				NewSuspectImpls.RemoveAll(x => !NewSuspectKeys.Contains((x.AuthorId, x.Change)));
+				await _issueSuspects.DeleteManyAsync(Builders<IssueSuspect>.Filter.In(x => x.Id, deleteSuspects.Select(y => y.Id)));
+				newSuspectImpls.RemoveAll(x => !newSuspectKeys.Contains((x.AuthorId, x.Change)));
 			}
 
 			// Make sure all the remaining suspects have the correct resolved time
-			if (NewSuspectImpls.Any(x => x.ResolvedAt != ResolvedAt))
+			if (newSuspectImpls.Any(x => x.ResolvedAt != resolvedAt))
 			{
-				await IssueSuspects.UpdateManyAsync(Builders<IssueSuspect>.Filter.Eq(x => x.IssueId, IssueId), Builders<IssueSuspect>.Update.Set(x => x.ResolvedAt, ResolvedAt));
+				await _issueSuspects.UpdateManyAsync(Builders<IssueSuspect>.Filter.Eq(x => x.IssueId, issueId), Builders<IssueSuspect>.Update.Set(x => x.ResolvedAt, resolvedAt));
 			}
-			return NewSuspectImpls;
+			return newSuspectImpls;
 		}
 
 		#endregion
@@ -1101,141 +1082,141 @@ namespace Horde.Build.Collections.Impl
 		#region Spans
 
 		/// <inheritdoc/>
-		public async Task<IIssueSpan> AddSpanAsync(int IssueId, NewIssueSpanData NewSpan)
+		public async Task<IIssueSpan> AddSpanAsync(int issueId, NewIssueSpanData newSpan)
 		{
-			IssueSpan Span = new IssueSpan(IssueId, NewSpan);
-			await IssueSpans.InsertOneAsync(Span);
-			return Span;
+			IssueSpan span = new IssueSpan(issueId, newSpan);
+			await _issueSpans.InsertOneAsync(span);
+			return span;
 		}
 
 		/// <inheritdoc/>
-		public async Task<IIssueSpan?> GetSpanAsync(ObjectId SpanId)
+		public async Task<IIssueSpan?> GetSpanAsync(ObjectId spanId)
 		{
-			return await IssueSpans.Find(Builders<IssueSpan>.Filter.Eq(x => x.Id, SpanId)).FirstOrDefaultAsync();
+			return await _issueSpans.Find(Builders<IssueSpan>.Filter.Eq(x => x.Id, spanId)).FirstOrDefaultAsync();
 		}
 
 		/// <inheritdoc/>
-		public async Task<IIssueSpan?> TryUpdateSpanAsync(IIssueSpan Span, NewIssueStepData? NewLastSuccess = null, NewIssueStepData? NewFailure = null, NewIssueStepData? NewNextSuccess = null, List<NewIssueSpanSuspectData>? NewSuspects = null, int? NewIssueId = null)
+		public async Task<IIssueSpan?> TryUpdateSpanAsync(IIssueSpan span, NewIssueStepData? newLastSuccess = null, NewIssueStepData? newFailure = null, NewIssueStepData? newNextSuccess = null, List<NewIssueSpanSuspectData>? newSuspects = null, int? newIssueId = null)
 		{
-			List<UpdateDefinition<IssueSpan>> Updates = new List<UpdateDefinition<IssueSpan>>();
-			if (NewLastSuccess != null)
+			List<UpdateDefinition<IssueSpan>> updates = new List<UpdateDefinition<IssueSpan>>();
+			if (newLastSuccess != null)
 			{
-				Updates.Add(Builders<IssueSpan>.Update.Set(x => x.MinChange, NewLastSuccess.Change));
-				Updates.Add(Builders<IssueSpan>.Update.Set(x => x.LastSuccess, new IssueStep(Span.Id, NewLastSuccess)));
+				updates.Add(Builders<IssueSpan>.Update.Set(x => x.MinChange, newLastSuccess.Change));
+				updates.Add(Builders<IssueSpan>.Update.Set(x => x.LastSuccess, new IssueStep(span.Id, newLastSuccess)));
 			}
-			if (NewFailure != null)
+			if (newFailure != null)
 			{
-				if (NewFailure.Change < Span.FirstFailure.Change)
+				if (newFailure.Change < span.FirstFailure.Change)
 				{
-					Updates.Add(Builders<IssueSpan>.Update.Set(x => x.FirstFailure, new IssueStep(Span.Id, NewFailure)));
+					updates.Add(Builders<IssueSpan>.Update.Set(x => x.FirstFailure, new IssueStep(span.Id, newFailure)));
 				}
-				if (NewFailure.Change >= Span.LastFailure.Change)
+				if (newFailure.Change >= span.LastFailure.Change)
 				{
-					Updates.Add(Builders<IssueSpan>.Update.Set(x => x.LastFailure, new IssueStep(Span.Id, NewFailure)));
+					updates.Add(Builders<IssueSpan>.Update.Set(x => x.LastFailure, new IssueStep(span.Id, newFailure)));
 				}
-				if (NewFailure.PromoteByDefault != Span.PromoteByDefault && NewFailure.Change >= Span.LastFailure.Change)
+				if (newFailure.PromoteByDefault != span.PromoteByDefault && newFailure.Change >= span.LastFailure.Change)
 				{
-					Updates.Add(Builders<IssueSpan>.Update.Set(x => x.PromoteByDefault, NewFailure.PromoteByDefault));
+					updates.Add(Builders<IssueSpan>.Update.Set(x => x.PromoteByDefault, newFailure.PromoteByDefault));
 				}
 			}
-			if (NewNextSuccess != null)
+			if (newNextSuccess != null)
 			{
-				Updates.Add(Builders<IssueSpan>.Update.Set(x => x.MaxChange, NewNextSuccess.Change));
-				Updates.Add(Builders<IssueSpan>.Update.Set(x => x.NextSuccess, new IssueStep(Span.Id, NewNextSuccess)));
+				updates.Add(Builders<IssueSpan>.Update.Set(x => x.MaxChange, newNextSuccess.Change));
+				updates.Add(Builders<IssueSpan>.Update.Set(x => x.NextSuccess, new IssueStep(span.Id, newNextSuccess)));
 			}
-			if (NewSuspects != null)
+			if (newSuspects != null)
 			{
-				Updates.Add(Builders<IssueSpan>.Update.Set(x => x.Suspects, NewSuspects.ConvertAll(x => new IssueSpanSuspect(x))));
+				updates.Add(Builders<IssueSpan>.Update.Set(x => x.Suspects, newSuspects.ConvertAll(x => new IssueSpanSuspect(x))));
 			}
-			if (NewIssueId != null)
+			if (newIssueId != null)
 			{
-				Updates.Add(Builders<IssueSpan>.Update.Set(x => x.IssueId, NewIssueId.Value));
+				updates.Add(Builders<IssueSpan>.Update.Set(x => x.IssueId, newIssueId.Value));
 			}
 
-			if (Updates.Count == 0)
+			if (updates.Count == 0)
 			{
-				return Span;
+				return span;
 			}
 
-			IssueSpan? NewSpan = await TryUpdateSpanAsync(Span, Builders<IssueSpan>.Update.Combine(Updates));
-			if (NewSpan != null)
+			IssueSpan? newSpan = await TryUpdateSpanAsync(span, Builders<IssueSpan>.Update.Combine(updates));
+			if (newSpan != null)
 			{
-				ILogger Logger = GetLogger(NewSpan.IssueId);
-				if (NewLastSuccess != null)
+				ILogger logger = GetLogger(newSpan.IssueId);
+				if (newLastSuccess != null)
 				{
-					Logger.LogInformation("Set last success for span {SpanId} to job {JobId} at CL {Change}", NewSpan.Id, NewLastSuccess.JobId, NewLastSuccess.Change);
+					logger.LogInformation("Set last success for span {SpanId} to job {JobId} at CL {Change}", newSpan.Id, newLastSuccess.JobId, newLastSuccess.Change);
 				}
-				if (NewNextSuccess != null)
+				if (newNextSuccess != null)
 				{
-					Logger.LogInformation("Set next success for span {SpanId} to job {JobId} at CL {Change}", NewSpan.Id, NewNextSuccess.JobId, NewNextSuccess.Change);
+					logger.LogInformation("Set next success for span {SpanId} to job {JobId} at CL {Change}", newSpan.Id, newNextSuccess.JobId, newNextSuccess.Change);
 				}
-				if (NewFailure != null)
+				if (newFailure != null)
 				{
-					Logger.LogInformation("Added failure for span {SpanId} in job {JobId} at CL {Change}", NewSpan.Id, NewFailure.JobId, NewFailure.Change);
+					logger.LogInformation("Added failure for span {SpanId} in job {JobId} at CL {Change}", newSpan.Id, newFailure.JobId, newFailure.Change);
 				}
 			}
-			return NewSpan;
+			return newSpan;
 		}
 
 		/// <inheritdoc/>
-		public async Task<List<IIssueSpan>> FindSpansAsync(int IssueId)
+		public async Task<List<IIssueSpan>> FindSpansAsync(int issueId)
 		{
-			return await IssueSpans.Find(x => x.IssueId == IssueId).ToListAsync<IssueSpan, IIssueSpan>();
+			return await _issueSpans.Find(x => x.IssueId == issueId).ToListAsync<IssueSpan, IIssueSpan>();
 		}
 
 		/// <inheritdoc/>
-		public Task<List<IIssueSpan>> FindSpansAsync(IEnumerable<ObjectId> SpanIds)
+		public Task<List<IIssueSpan>> FindSpansAsync(IEnumerable<ObjectId> spanIds)
 		{
-			return IssueSpans.Find(Builders<IssueSpan>.Filter.In(x => x.Id, SpanIds)).ToListAsync<IssueSpan, IIssueSpan>();
+			return _issueSpans.Find(Builders<IssueSpan>.Filter.In(x => x.Id, spanIds)).ToListAsync<IssueSpan, IIssueSpan>();
 		}
 
 		/// <inheritdoc/>
-		public async Task<List<IIssueSpan>> FindOpenSpansAsync(StreamId StreamId, TemplateRefId TemplateId, string NodeName, int Change)
+		public async Task<List<IIssueSpan>> FindOpenSpansAsync(StreamId streamId, TemplateRefId templateId, string nodeName, int change)
 		{
-			List<IssueSpan> Spans = await IssueSpans.Find(x => x.StreamId == StreamId && x.TemplateRefId == TemplateId && x.NodeName == NodeName && Change >= x.MinChange && Change <= x.MaxChange).ToListAsync();
-			return Spans.ConvertAll<IIssueSpan>(x => x);
+			List<IssueSpan> spans = await _issueSpans.Find(x => x.StreamId == streamId && x.TemplateRefId == templateId && x.NodeName == nodeName && change >= x.MinChange && change <= x.MaxChange).ToListAsync();
+			return spans.ConvertAll<IIssueSpan>(x => x);
 		}
 
 		/// <inheritdoc/>
-		public Task<List<IIssueSpan>> FindSpansAsync(IEnumerable<ObjectId>? SpanIds, IEnumerable<int>? IssueIds, StreamId? StreamId, int? MinChange, int? MaxChange, bool? Resolved, int? Index, int? Count)
+		public Task<List<IIssueSpan>> FindSpansAsync(IEnumerable<ObjectId>? spanIds, IEnumerable<int>? issueIds, StreamId? streamId, int? minChange, int? maxChange, bool? resolved, int? index, int? count)
 		{
-			FilterDefinition<IssueSpan> Filter = FilterDefinition<IssueSpan>.Empty;
+			FilterDefinition<IssueSpan> filter = FilterDefinition<IssueSpan>.Empty;
 
-			if(SpanIds != null)
+			if(spanIds != null)
 			{
-				Filter &= Builders<IssueSpan>.Filter.In(x => x.Id, SpanIds);
+				filter &= Builders<IssueSpan>.Filter.In(x => x.Id, spanIds);
 			}
 
-			if (StreamId != null)
+			if (streamId != null)
 			{
-				Filter &= Builders<IssueSpan>.Filter.Eq(x => x.StreamId, StreamId);
+				filter &= Builders<IssueSpan>.Filter.Eq(x => x.StreamId, streamId);
 			}
 
-			if (IssueIds != null)
+			if (issueIds != null)
 			{
-				Filter &= Builders<IssueSpan>.Filter.In(x => x.IssueId, IssueIds.Select<int, int?>(x => x));
+				filter &= Builders<IssueSpan>.Filter.In(x => x.IssueId, issueIds.Select<int, int?>(x => x));
 			}
-			if (MinChange != null)
+			if (minChange != null)
 			{
-				Filter &= Builders<IssueSpan>.Filter.Not(Builders<IssueSpan>.Filter.Lt(x => x.MaxChange, MinChange.Value));
+				filter &= Builders<IssueSpan>.Filter.Not(Builders<IssueSpan>.Filter.Lt(x => x.MaxChange, minChange.Value));
 			}
-			if (MaxChange != null)
+			if (maxChange != null)
 			{
-				Filter &= Builders<IssueSpan>.Filter.Not(Builders<IssueSpan>.Filter.Gt(x => x.MinChange, MaxChange.Value));
+				filter &= Builders<IssueSpan>.Filter.Not(Builders<IssueSpan>.Filter.Gt(x => x.MinChange, maxChange.Value));
 			}
-			if (Resolved != null)
+			if (resolved != null)
 			{
-				if (Resolved.Value)
+				if (resolved.Value)
 				{
-					Filter &= Builders<IssueSpan>.Filter.Ne(x => x.ResolvedAt, null);
+					filter &= Builders<IssueSpan>.Filter.Ne(x => x.ResolvedAt, null);
 				}
 				else
 				{
-					Filter &= Builders<IssueSpan>.Filter.Eq(x => x.ResolvedAt, null);
+					filter &= Builders<IssueSpan>.Filter.Eq(x => x.ResolvedAt, null);
 				}
 			}
 
-			return IssueSpans.Find(Filter).Range(Index, Count).ToListAsync<IssueSpan, IIssueSpan>();
+			return _issueSpans.Find(filter).Range(index, count).ToListAsync<IssueSpan, IIssueSpan>();
 		}
 
 		#endregion
@@ -1243,41 +1224,41 @@ namespace Horde.Build.Collections.Impl
 		#region Steps
 
 		/// <inheritdoc/>
-		public async Task<IIssueStep> AddStepAsync(ObjectId SpanId, NewIssueStepData NewStep)
+		public async Task<IIssueStep> AddStepAsync(ObjectId spanId, NewIssueStepData newStep)
 		{
-			IssueStep Step = new IssueStep(SpanId, NewStep);
-			await IssueSteps.InsertOneAsync(Step);
-			return Step;
+			IssueStep step = new IssueStep(spanId, newStep);
+			await _issueSteps.InsertOneAsync(step);
+			return step;
 		}
 
 		/// <inheritdoc/>
-		public Task<List<IIssueStep>> FindStepsAsync(IEnumerable<ObjectId> SpanIds)
+		public Task<List<IIssueStep>> FindStepsAsync(IEnumerable<ObjectId> spanIds)
 		{
-			FilterDefinition<IssueStep> Filter = Builders<IssueStep>.Filter.In(x => x.SpanId, SpanIds);
-			return IssueSteps.Find(Filter).ToListAsync<IssueStep, IIssueStep>();
+			FilterDefinition<IssueStep> filter = Builders<IssueStep>.Filter.In(x => x.SpanId, spanIds);
+			return _issueSteps.Find(filter).ToListAsync<IssueStep, IIssueStep>();
 		}
 
 		/// <inheritdoc/>
-		public Task<List<IIssueStep>> FindStepsAsync(JobId JobId, SubResourceId? BatchId, SubResourceId? StepId)
+		public Task<List<IIssueStep>> FindStepsAsync(JobId jobId, SubResourceId? batchId, SubResourceId? stepId)
 		{
-			FilterDefinition<IssueStep> Filter = Builders<IssueStep>.Filter.Eq(x => x.JobId, JobId);
-			if (BatchId != null)
+			FilterDefinition<IssueStep> filter = Builders<IssueStep>.Filter.Eq(x => x.JobId, jobId);
+			if (batchId != null)
 			{
-				Filter &= Builders<IssueStep>.Filter.Eq(x => x.BatchId, BatchId.Value);
+				filter &= Builders<IssueStep>.Filter.Eq(x => x.BatchId, batchId.Value);
 			}
-			if (StepId != null)
+			if (stepId != null)
 			{
-				Filter &= Builders<IssueStep>.Filter.Eq(x => x.StepId, StepId.Value);
+				filter &= Builders<IssueStep>.Filter.Eq(x => x.StepId, stepId.Value);
 			}
-			return IssueSteps.Find(Filter).ToListAsync<IssueStep, IIssueStep>();
+			return _issueSteps.Find(filter).ToListAsync<IssueStep, IIssueStep>();
 		}
 
 		#endregion
 
 		/// <inheritdoc/>
-		public IAuditLogChannel<int> GetLogger(int IssueId)
+		public IAuditLogChannel<int> GetLogger(int issueId)
 		{
-			return AuditLog[IssueId];
+			return _auditLog[issueId];
 		}
 	}
 }
