Horde: Fix coding conventions within Horde.Build.

diff --git a/Engine/Source/Programs/Horde/Horde.Build/Jobs/IGraph.cs b/Engine/Source/Programs/Horde/Horde.Build/Jobs/IGraph.cs
--- a/Engine/Source/Programs/Horde/Horde.Build/Jobs/IGraph.cs
+++ b/Engine/Source/Programs/Horde/Horde.Build/Jobs/IGraph.cs
@@ -1,17 +1,13 @@
 // Copyright Epic Games, Inc. All Rights Reserved.
 
-using EpicGames.Core;
-using Horde.Build.Api;
-using HordeCommon;
-using Horde.Build.Utilities;
-using MongoDB.Bson;
-using MongoDB.Bson.Serialization.Attributes;
-using Newtonsoft.Json;
 using System;
 using System.Collections.Generic;
 using System.Diagnostics;
 using System.Diagnostics.CodeAnalysis;
 using System.Linq;
+using EpicGames.Core;
+using HordeCommon;
+using MongoDB.Bson.Serialization.Attributes;
 
 namespace Horde.Build.Models
 {
@@ -111,19 +107,19 @@ namespace Horde.Build.Models
 		/// <summary>
 		/// Constructor
 		/// </summary>
-		/// <param name="GroupIdx">Index of thr group containing the node</param>
-		/// <param name="NodeIdx">Index of the node within the group</param>
-		public NodeRef(int GroupIdx, int NodeIdx)
+		/// <param name="groupIdx">Index of thr group containing the node</param>
+		/// <param name="nodeIdx">Index of the node within the group</param>
+		public NodeRef(int groupIdx, int nodeIdx)
 		{
-			this.GroupIdx = GroupIdx;
-			this.NodeIdx = NodeIdx;
+			GroupIdx = groupIdx;
+			NodeIdx = nodeIdx;
 		}
 
 		/// <inheritdoc/>
-		public override bool Equals(object? Other)
+		public override bool Equals(object? other)
 		{
-			NodeRef? OtherNodeRef = Other as NodeRef;
-			return OtherNodeRef != null && OtherNodeRef.GroupIdx == GroupIdx && OtherNodeRef.NodeIdx == NodeIdx;
+			NodeRef? otherNodeRef = other as NodeRef;
+			return otherNodeRef != null && otherNodeRef.GroupIdx == GroupIdx && otherNodeRef.NodeIdx == NodeIdx;
 		}
 
 		/// <inheritdoc/>
@@ -135,11 +131,11 @@ namespace Horde.Build.Models
 		/// <summary>
 		/// Converts this reference to a node name
 		/// </summary>
-		/// <param name="Groups">List of groups that this reference points to</param>
+		/// <param name="groups">List of groups that this reference points to</param>
 		/// <returns>Name of the referenced node</returns>
-		public INode ToNode(IReadOnlyList<INodeGroup> Groups)
+		public INode ToNode(IReadOnlyList<INodeGroup> groups)
 		{
-			return Groups[GroupIdx].Nodes[NodeIdx];
+			return groups[GroupIdx].Nodes[NodeIdx];
 		}
 	}
 
@@ -208,18 +204,18 @@ namespace Horde.Build.Models
 		/// <summary>
 		/// Enumerate all the required dependencies of this node group
 		/// </summary>
-		/// <param name="Label">The label instance</param>
-		/// <param name="Groups">List of groups for the job containing this aggregate</param>
+		/// <param name="label">The label instance</param>
+		/// <param name="groups">List of groups for the job containing this aggregate</param>
 		/// <returns>Sequence of nodes</returns>
-		public static IEnumerable<INode> GetDependencies(this ILabel Label, IReadOnlyList<INodeGroup> Groups)
+		public static IEnumerable<INode> GetDependencies(this ILabel label, IReadOnlyList<INodeGroup> groups)
 		{
-			foreach (NodeRef RequiredNode in Label.RequiredNodes)
+			foreach (NodeRef requiredNode in label.RequiredNodes)
 			{
-				yield return RequiredNode.ToNode(Groups);
+				yield return requiredNode.ToNode(groups);
 			}
-			foreach (NodeRef IncludedNode in Label.IncludedNodes)
+			foreach (NodeRef includedNode in label.IncludedNodes)
 			{
-				yield return IncludedNode.ToNode(Groups);
+				yield return includedNode.ToNode(groups);
 			}
 		}
 	}
@@ -263,59 +259,59 @@ namespace Horde.Build.Models
 		/// <summary>
 		/// Gets the node from a node reference
 		/// </summary>
-		/// <param name="Graph">The graph instance</param>
-		/// <param name="Ref">The node reference</param>
+		/// <param name="graph">The graph instance</param>
+		/// <param name="nodeRef">The node reference</param>
 		/// <returns>The node for the given reference</returns>
-		public static INode GetNode(this IGraph Graph, NodeRef Ref)
+		public static INode GetNode(this IGraph graph, NodeRef nodeRef)
 		{
-			return Graph.Groups[Ref.GroupIdx].Nodes[Ref.NodeIdx];
+			return graph.Groups[nodeRef.GroupIdx].Nodes[nodeRef.NodeIdx];
 		}
 
 		/// <summary>
 		/// Tries to find a node by name
 		/// </summary>
-		/// <param name="Graph">The graph to search</param>
-		/// <param name="NodeName">Name of the node</param>
-		/// <param name="Ref">Receives the node reference</param>
+		/// <param name="graph">The graph to search</param>
+		/// <param name="nodeName">Name of the node</param>
+		/// <param name="nodeRef">Receives the node reference</param>
 		/// <returns>True if the node was found, false otherwise</returns>
-		public static bool TryFindNode(this IGraph Graph, string NodeName, out NodeRef Ref)
+		public static bool TryFindNode(this IGraph graph, string nodeName, out NodeRef nodeRef)
 		{
-			for (int GroupIdx = 0; GroupIdx < Graph.Groups.Count; GroupIdx++)
+			for (int groupIdx = 0; groupIdx < graph.Groups.Count; groupIdx++)
 			{
-				INodeGroup Group = Graph.Groups[GroupIdx];
-				for (int NodeIdx = 0; NodeIdx < Group.Nodes.Count; NodeIdx++)
+				INodeGroup group = graph.Groups[groupIdx];
+				for (int nodeIdx = 0; nodeIdx < group.Nodes.Count; nodeIdx++)
 				{
-					INode Node = Group.Nodes[NodeIdx];
-					if (String.Equals(Node.Name, NodeName, StringComparison.OrdinalIgnoreCase))
+					INode node = group.Nodes[nodeIdx];
+					if (String.Equals(node.Name, nodeName, StringComparison.OrdinalIgnoreCase))
 					{
-						Ref = new NodeRef(GroupIdx, NodeIdx);
+						nodeRef = new NodeRef(groupIdx, nodeIdx);
 						return true;
 					}
 				}
 			}
 
-			Ref = new NodeRef(0, 0);
+			nodeRef = new NodeRef(0, 0);
 			return false;
 		}
 
 		/// <summary>
 		/// Tries to find a node by name
 		/// </summary>
-		/// <param name="Graph">The graph to search</param>
-		/// <param name="NodeName">Name of the node</param>
-		/// <param name="Node">Receives the node</param>
+		/// <param name="graph">The graph to search</param>
+		/// <param name="nodeName">Name of the node</param>
+		/// <param name="node">Receives the node</param>
 		/// <returns>True if the node was found, false otherwise</returns>
-		public static bool TryFindNode(this IGraph Graph, string NodeName, out INode? Node)
+		public static bool TryFindNode(this IGraph graph, string nodeName, out INode? node)
 		{
-			NodeRef Ref;
-			if (TryFindNode(Graph, NodeName, out Ref))
+			NodeRef nodeRef;
+			if (TryFindNode(graph, nodeName, out nodeRef))
 			{
-				Node = Graph.Groups[Ref.GroupIdx].Nodes[Ref.NodeIdx];
+				node = graph.Groups[nodeRef.GroupIdx].Nodes[nodeRef.NodeIdx];
 				return true;
 			}
 			else
 			{
-				Node = null;
+				node = null;
 				return false;
 			}
 		}
@@ -323,34 +319,34 @@ namespace Horde.Build.Models
 		/// <summary>
 		/// Tries to find a node by name
 		/// </summary>
-		/// <param name="Graph">The graph to search</param>
-		/// <param name="Name">Name of the node</param>
-		/// <param name="AggregateIdx">Receives the aggregate index</param>
+		/// <param name="graph">The graph to search</param>
+		/// <param name="name">Name of the node</param>
+		/// <param name="aggregateIdx">Receives the aggregate index</param>
 		/// <returns>True if the node was found, false otherwise</returns>
-		public static bool TryFindAggregate(this IGraph Graph, string Name, out int AggregateIdx)
+		public static bool TryFindAggregate(this IGraph graph, string name, out int aggregateIdx)
 		{
-			AggregateIdx = Graph.Aggregates.FindIndex(x => x.Name.Equals(Name, StringComparison.OrdinalIgnoreCase));
-			return AggregateIdx != -1;
+			aggregateIdx = graph.Aggregates.FindIndex(x => x.Name.Equals(name, StringComparison.OrdinalIgnoreCase));
+			return aggregateIdx != -1;
 		}
 
 		/// <summary>
 		/// Tries to find a node by name
 		/// </summary>
-		/// <param name="Graph">The graph to search</param>
-		/// <param name="Name">Name of the node</param>
-		/// <param name="Aggregate">Receives the aggregate</param>
+		/// <param name="graph">The graph to search</param>
+		/// <param name="name">Name of the node</param>
+		/// <param name="aggregate">Receives the aggregate</param>
 		/// <returns>True if the node was found, false otherwise</returns>
-		public static bool TryFindAggregate(this IGraph Graph, string Name, [NotNullWhen(true)] out IAggregate? Aggregate)
+		public static bool TryFindAggregate(this IGraph graph, string name, [NotNullWhen(true)] out IAggregate? aggregate)
 		{
-			int AggregateIdx;
-			if (TryFindAggregate(Graph, Name, out AggregateIdx))
+			int aggregateIdx;
+			if (TryFindAggregate(graph, name, out aggregateIdx))
 			{
-				Aggregate = Graph.Aggregates[AggregateIdx];
+				aggregate = graph.Aggregates[aggregateIdx];
 				return true;
 			}
 			else
 			{
-				Aggregate = null;
+				aggregate = null;
 				return false;
 			}
 		}
@@ -358,12 +354,12 @@ namespace Horde.Build.Models
 		/// <summary>
 		/// Gets a list of dependencies for the given node
 		/// </summary>
-		/// <param name="Graph">The graph instance</param>
-		/// <param name="Node">The node to return dependencies for</param>
+		/// <param name="graph">The graph instance</param>
+		/// <param name="node">The node to return dependencies for</param>
 		/// <returns>List of dependencies</returns>
-		public static IEnumerable<INode> GetDependencies(this IGraph Graph, INode Node)
+		public static IEnumerable<INode> GetDependencies(this IGraph graph, INode node)
 		{
-			return Enumerable.Concat(Node.InputDependencies, Node.OrderDependencies).Select(x => Graph.GetNode(x));
+			return Enumerable.Concat(node.InputDependencies, node.OrderDependencies).Select(x => graph.GetNode(x));
 		}
 	}
 
@@ -420,26 +416,26 @@ namespace Horde.Build.Models
 		/// <summary>
 		/// Constructor
 		/// </summary>
-		/// <param name="Name">Name of the node</param>
-		/// <param name="InputDependencies">List of nodes which must have completed succesfully for this node to run</param>
-		/// <param name="OrderDependencies">List of nodes which must have completed for this node to run</param>
-		/// <param name="Priority">Priority of this node</param>
-		/// <param name="AllowRetry">Whether the node can be run multiple times</param>
-		/// <param name="RunEarly">Whether the node can run early, before dependencies of other nodes in the same group complete</param>
-		/// <param name="Warnings">Whether to include warnings in the diagnostic output (defaults to true)</param>
-		/// <param name="Credentials">Credentials required for this node to run</param>
-		/// <param name="Properties">Properties for the node</param>
-		public NewNode(string Name, List<string>? InputDependencies = null, List<string>? OrderDependencies = null, Priority? Priority = null, bool? AllowRetry = null, bool? RunEarly = null, bool? Warnings = null, Dictionary<string, string>? Credentials = null, Dictionary<string, string>? Properties = null)
+		/// <param name="name">Name of the node</param>
+		/// <param name="inputDependencies">List of nodes which must have completed succesfully for this node to run</param>
+		/// <param name="orderDependencies">List of nodes which must have completed for this node to run</param>
+		/// <param name="priority">Priority of this node</param>
+		/// <param name="allowRetry">Whether the node can be run multiple times</param>
+		/// <param name="runEarly">Whether the node can run early, before dependencies of other nodes in the same group complete</param>
+		/// <param name="warnings">Whether to include warnings in the diagnostic output (defaults to true)</param>
+		/// <param name="credentials">Credentials required for this node to run</param>
+		/// <param name="properties">Properties for the node</param>
+		public NewNode(string name, List<string>? inputDependencies = null, List<string>? orderDependencies = null, Priority? priority = null, bool? allowRetry = null, bool? runEarly = null, bool? warnings = null, Dictionary<string, string>? credentials = null, Dictionary<string, string>? properties = null)
 		{
-			this.Name = Name;
-			this.InputDependencies = InputDependencies;
-			this.OrderDependencies = OrderDependencies;
-			this.Priority = Priority;
-			this.AllowRetry = AllowRetry;
-			this.RunEarly = RunEarly;
-			this.Warnings = Warnings;
-			this.Credentials = Credentials;
-			this.Properties = Properties;
+			Name = name;
+			InputDependencies = inputDependencies;
+			OrderDependencies = orderDependencies;
+			Priority = priority;
+			AllowRetry = allowRetry;
+			RunEarly = runEarly;
+			Warnings = warnings;
+			Credentials = credentials;
+			Properties = properties;
 		}
 	}
 
@@ -461,12 +457,12 @@ namespace Horde.Build.Models
 		/// <summary>
 		/// Constructor
 		/// </summary>
-		/// <param name="AgentType">The type of agent to execute this group</param>
-		/// <param name="Nodes">Nodes in this group</param>
-		public NewGroup(string AgentType, List<NewNode> Nodes)
+		/// <param name="agentType">The type of agent to execute this group</param>
+		/// <param name="nodes">Nodes in this group</param>
+		public NewGroup(string agentType, List<NewNode> nodes)
 		{
-			this.AgentType = AgentType;
-			this.Nodes = Nodes;
+			AgentType = agentType;
+			Nodes = nodes;
 		}
 	}
 
@@ -541,14 +537,12 @@ namespace Horde.Build.Models
 		/// <summary>
 		/// Constructor
 		/// </summary>
-		/// <param name="Name">Name of this aggregate</param>
-		/// <param name="Nodes">Nodes which must be part of the job for the aggregate to be shown</param>
-		public NewAggregate(string Name, List<string> Nodes)
+		/// <param name="name">Name of this aggregate</param>
+		/// <param name="nodes">Nodes which must be part of the job for the aggregate to be shown</param>
+		public NewAggregate(string name, List<string> nodes)
 		{
-			this.Name = Name;
-			this.Nodes = Nodes;
+			Name = name;
+			Nodes = nodes;
 		}
 	}
 }
-
-
diff --git a/Engine/Source/Programs/Horde/Horde.Build/Jobs/IGraphCollection.cs b/Engine/Source/Programs/Horde/Horde.Build/Jobs/IGraphCollection.cs
--- a/Engine/Source/Programs/Horde/Horde.Build/Jobs/IGraphCollection.cs
+++ b/Engine/Source/Programs/Horde/Horde.Build/Jobs/IGraphCollection.cs
@@ -1,20 +1,9 @@
 // Copyright Epic Games, Inc. All Rights Reserved.
 
-using EpicGames.Core;
-using Horde.Build.Api;
-using HordeCommon;
-using Horde.Build.Models;
-using Horde.Build.Services;
-using MongoDB.Bson;
-using MongoDB.Bson.Serialization.Attributes;
-using MongoDB.Driver;
-using System;
-using System.Collections.Concurrent;
 using System.Collections.Generic;
-using System.Diagnostics;
-using System.Linq;
-using System.Threading;
 using System.Threading.Tasks;
+using EpicGames.Core;
+using Horde.Build.Models;
 
 namespace Horde.Build.Collections
 {
@@ -26,34 +15,34 @@ namespace Horde.Build.Collections
 		/// <summary>
 		/// Adds a graph from a template
 		/// </summary>
-		/// <param name="Template">The template</param>
+		/// <param name="template">The template</param>
 		/// <returns>New graph</returns>
-		Task<IGraph> AddAsync(ITemplate Template);
+		Task<IGraph> AddAsync(ITemplate template);
 
 		/// <summary>
 		/// Creates a graph by appending groups and aggregates to an existing graph.
 		/// </summary>
-		/// <param name="BaseGraph">The base graph</param>
-		/// <param name="NewGroupRequests">List of group requests</param>
-		/// <param name="NewAggregateRequests">List of aggregate requests</param>
-		/// <param name="NewLabelRequests">List of label requests</param>
+		/// <param name="baseGraph">The base graph</param>
+		/// <param name="newGroupRequests">List of group requests</param>
+		/// <param name="newAggregateRequests">List of aggregate requests</param>
+		/// <param name="newLabelRequests">List of label requests</param>
 		/// <returns>The new graph definition</returns>
-		Task<IGraph> AppendAsync(IGraph? BaseGraph, List<NewGroup>? NewGroupRequests = null, List<NewAggregate>? NewAggregateRequests = null, List<NewLabel>? NewLabelRequests = null);
+		Task<IGraph> AppendAsync(IGraph? baseGraph, List<NewGroup>? newGroupRequests = null, List<NewAggregate>? newAggregateRequests = null, List<NewLabel>? newLabelRequests = null);
 
 		/// <summary>
 		/// Gets the graph for a job
 		/// </summary>
-		/// <param name="Hash">Hash of the graph to retrieve</param>
+		/// <param name="hash">Hash of the graph to retrieve</param>
 		/// <returns>The graph for this job</returns>
-		Task<IGraph> GetAsync(ContentHash Hash);
+		Task<IGraph> GetAsync(ContentHash hash);
 
 		/// <summary>
 		/// Finds all graphs stored in the collection
 		/// </summary>
-		/// <param name="Hashes">Hashes to filter by</param>
-		/// <param name="Index">Starting index of the graph to return</param>
-		/// <param name="Count">Number of results to return</param>
+		/// <param name="hashes">Hashes to filter by</param>
+		/// <param name="index">Starting index of the graph to return</param>
+		/// <param name="count">Number of results to return</param>
 		/// <returns>List of graphs</returns>
-		Task<List<IGraph>> FindAllAsync(ContentHash[]? Hashes = null, int? Index = null, int? Count = null);
+		Task<List<IGraph>> FindAllAsync(ContentHash[]? hashes = null, int? index = null, int? count = null);
 	}
 }
diff --git a/Engine/Source/Programs/Horde/Horde.Build/Jobs/IJob.cs b/Engine/Source/Programs/Horde/Horde.Build/Jobs/IJob.cs
--- a/Engine/Source/Programs/Horde/Horde.Build/Jobs/IJob.cs
+++ b/Engine/Source/Programs/Horde/Horde.Build/Jobs/IJob.cs
@@ -1,33 +1,27 @@
-﻿// Copyright Epic Games, Inc. All Rights Reserved.
+// Copyright Epic Games, Inc. All Rights Reserved.
 
+using System;
+using System.Collections.Generic;
+using System.Diagnostics.CodeAnalysis;
+using System.Linq;
 using EpicGames.Core;
 using Horde.Build.Acls;
 using Horde.Build.Api;
-using HordeCommon;
 using Horde.Build.Utilities;
+using HordeCommon;
 using MongoDB.Bson;
-using MongoDB.Bson.Serialization.Attributes;
 using MongoDB.Driver;
-using System;
-using System.Collections.Generic;
-using System.ComponentModel;
-using System.Diagnostics.CodeAnalysis;
-using System.Linq;
-using System.Runtime.CompilerServices;
-using System.Text;
-using System.Text.Json.Serialization;
-using System.Threading.Tasks;
 
 namespace Horde.Build.Models
 {
-	using ReportPlacement = HordeCommon.Rpc.ReportPlacement;
 	using JobId = ObjectId<IJob>;
 	using LeaseId = ObjectId<ILease>;
 	using LogId = ObjectId<ILogFile>;
+	using PoolId = StringId<IPool>;
+	using ReportPlacement = HordeCommon.Rpc.ReportPlacement;
 	using SessionId = ObjectId<ISession>;
 	using StreamId = StringId<IStream>;
 	using TemplateRefId = StringId<TemplateRef>;
-	using PoolId = StringId<IPool>;
 	using UserId = ObjectId<IUser>;
 
 	/// <summary>
@@ -156,18 +150,18 @@ namespace Horde.Build.Models
 		/// Determines if a jobstep has failed or is skipped. Can be used to determine whether dependent steps will be able to run.
 		/// </summary>
 		/// <returns>True if the step is failed or skipped</returns>
-		public static bool IsFailedOrSkipped(this IJobStep Step)
+		public static bool IsFailedOrSkipped(this IJobStep step)
 		{
-			return Step.State == JobStepState.Skipped || Step.Outcome == JobStepOutcome.Failure;
+			return step.State == JobStepState.Skipped || step.Outcome == JobStepOutcome.Failure;
 		}
 
 		/// <summary>
 		/// Determines if a jobstep is done by checking to see if it is completed, skipped, or aborted.
 		/// </summary>
 		/// <returns>True if the step is completed, skipped, or aborted</returns>
-		public static bool IsPending(this IJobStep Step)
+		public static bool IsPending(this IJobStep step)
 		{
-			return Step.State != JobStepState.Aborted && Step.State != JobStepState.Completed && Step.State != JobStepState.Skipped;
+			return step.State != JobStepState.Aborted && step.State != JobStepState.Completed && step.State != JobStepState.Skipped;
 		}
 	}
 
@@ -255,57 +249,57 @@ namespace Horde.Build.Models
 		/// <summary>
 		/// Attempts to get a step with the given id
 		/// </summary>
-		/// <param name="Batch">The batch to search</param>
-		/// <param name="StepId">The step id</param>
-		/// <param name="Step">On success, receives the step object</param>
+		/// <param name="batch">The batch to search</param>
+		/// <param name="stepId">The step id</param>
+		/// <param name="step">On success, receives the step object</param>
 		/// <returns>True if the step was found</returns>
-		public static bool TryGetStep(this IJobStepBatch Batch, SubResourceId StepId, [NotNullWhen(true)] out IJobStep? Step)
+		public static bool TryGetStep(this IJobStepBatch batch, SubResourceId stepId, [NotNullWhen(true)] out IJobStep? step)
 		{
-			Step = Batch.Steps.FirstOrDefault(x => x.Id == StepId);
-			return Step != null;
+			step = batch.Steps.FirstOrDefault(x => x.Id == stepId);
+			return step != null;
 		}
 
 		/// <summary>
 		/// Determines if new steps can be appended to this batch. We do not allow this after the last step has been completed, because the agent is shutting down.
 		/// </summary>
-		/// <param name="Batch">The batch to search</param>
+		/// <param name="batch">The batch to search</param>
 		/// <returns>True if new steps can be appended to this batch</returns>
-		public static bool CanBeAppendedTo(this IJobStepBatch Batch)
+		public static bool CanBeAppendedTo(this IJobStepBatch batch)
 		{
-			return Batch.State <= JobStepBatchState.Running;
+			return batch.State <= JobStepBatchState.Running;
 		}
 
 		/// <summary>
 		/// Gets the wait time for this batch
 		/// </summary>
-		/// <param name="Batch">The batch to search</param>
+		/// <param name="batch">The batch to search</param>
 		/// <returns>Wait time for the batch</returns>
-		public static TimeSpan? GetWaitTime(this IJobStepBatch Batch)
+		public static TimeSpan? GetWaitTime(this IJobStepBatch batch)
 		{
-			if (Batch.StartTimeUtc == null || Batch.ReadyTimeUtc == null)
+			if (batch.StartTimeUtc == null || batch.ReadyTimeUtc == null)
 			{
 				return null;
 			}
 			else
 			{
-				return Batch.StartTimeUtc.Value - Batch.ReadyTimeUtc.Value;
+				return batch.StartTimeUtc.Value - batch.ReadyTimeUtc.Value;
 			}
 		}
 
 		/// <summary>
 		/// Gets the initialization time for this batch
 		/// </summary>
-		/// <param name="Batch">The batch to search</param>
+		/// <param name="batch">The batch to search</param>
 		/// <returns>Initialization time for this batch</returns>
-		public static TimeSpan? GetInitTime(this IJobStepBatch Batch)
+		public static TimeSpan? GetInitTime(this IJobStepBatch batch)
 		{
-			if (Batch.StartTimeUtc != null)
+			if (batch.StartTimeUtc != null)
 			{
-				foreach (IJobStep Step in Batch.Steps)
+				foreach (IJobStep step in batch.Steps)
 				{
-					if (Step.StartTimeUtc != null)
+					if (step.StartTimeUtc != null)
 					{
-						return Step.StartTimeUtc - Batch.StartTimeUtc.Value;
+						return step.StartTimeUtc - batch.StartTimeUtc.Value;
 					}
 				}
 			}
@@ -315,29 +309,29 @@ namespace Horde.Build.Models
 		/// <summary>
 		/// Get the dependencies required for this batch to start, taking run-early nodes into account
 		/// </summary>
-		/// <param name="Batch">The batch to search</param>
-		/// <param name="Groups">List of node groups</param>
+		/// <param name="batch">The batch to search</param>
+		/// <param name="groups">List of node groups</param>
 		/// <returns>Set of nodes that must have completed for this batch to start</returns>
-		public static HashSet<INode> GetStartDependencies(this IJobStepBatch Batch, IReadOnlyList<INodeGroup> Groups)
+		public static HashSet<INode> GetStartDependencies(this IJobStepBatch batch, IReadOnlyList<INodeGroup> groups)
 		{
 			// Find all the nodes that this group will start with.
-			List<INode> Nodes = Batch.Steps.ConvertAll(x => Groups[Batch.GroupIdx].Nodes[x.NodeIdx]);
-			if (Nodes.Any(x => x.RunEarly))
+			List<INode> nodes = batch.Steps.ConvertAll(x => groups[batch.GroupIdx].Nodes[x.NodeIdx]);
+			if (nodes.Any(x => x.RunEarly))
 			{
-				Nodes.RemoveAll(x => !x.RunEarly);
+				nodes.RemoveAll(x => !x.RunEarly);
 			}
 
 			// Find all their dependencies
-			HashSet<INode> Dependencies = new HashSet<INode>();
-			foreach (INode Node in Nodes)
+			HashSet<INode> dependencies = new HashSet<INode>();
+			foreach (INode node in nodes)
 			{
-				Dependencies.UnionWith(Node.InputDependencies.Select(x => Groups[x.GroupIdx].Nodes[x.NodeIdx]));
-				Dependencies.UnionWith(Node.OrderDependencies.Select(x => Groups[x.GroupIdx].Nodes[x.NodeIdx]));
+				dependencies.UnionWith(node.InputDependencies.Select(x => groups[x.GroupIdx].Nodes[x.NodeIdx]));
+				dependencies.UnionWith(node.OrderDependencies.Select(x => groups[x.GroupIdx].Nodes[x.NodeIdx]));
 			}
 
 			// Exclude all the dependencies within the same group
-			Dependencies.ExceptWith(Groups[Batch.GroupIdx].Nodes);
-			return Dependencies;
+			dependencies.ExceptWith(groups[batch.GroupIdx].Nodes);
+			return dependencies;
 		}
 	}
 
@@ -398,41 +392,41 @@ namespace Horde.Build.Models
 		/// <summary>
 		/// Copy constructor
 		/// </summary>
-		/// <param name="Other">The timing info object to copy from</param>
-		public TimingInfo(TimingInfo Other)
+		/// <param name="other">The timing info object to copy from</param>
+		public TimingInfo(TimingInfo other)
 		{
-			TotalWaitTime = Other.TotalWaitTime;
-			TotalInitTime = Other.TotalInitTime;
-			TotalTimeToComplete = Other.TotalTimeToComplete;
+			TotalWaitTime = other.TotalWaitTime;
+			TotalInitTime = other.TotalInitTime;
+			TotalTimeToComplete = other.TotalTimeToComplete;
 
-			AverageTotalWaitTime = Other.AverageTotalWaitTime;
-			AverageTotalInitTime = Other.AverageTotalInitTime;
-			AverageTotalTimeToComplete = Other.AverageTotalTimeToComplete;
+			AverageTotalWaitTime = other.AverageTotalWaitTime;
+			AverageTotalInitTime = other.AverageTotalInitTime;
+			AverageTotalTimeToComplete = other.AverageTotalTimeToComplete;
 		}
 
 		/// <summary>
 		/// Modifies this timing to wait for another timing
 		/// </summary>
-		/// <param name="Other">The other node to wait for</param>
-		public void WaitFor(TimingInfo Other)
+		/// <param name="other">The other node to wait for</param>
+		public void WaitFor(TimingInfo other)
 		{
 			if (TotalTimeToComplete != null)
 			{
-				if (Other.TotalTimeToComplete == null || Other.TotalTimeToComplete.Value > TotalTimeToComplete.Value)
+				if (other.TotalTimeToComplete == null || other.TotalTimeToComplete.Value > TotalTimeToComplete.Value)
 				{
-					TotalInitTime = Other.TotalInitTime;
-					TotalWaitTime = Other.TotalWaitTime;
-					TotalTimeToComplete = Other.TotalTimeToComplete;
+					TotalInitTime = other.TotalInitTime;
+					TotalWaitTime = other.TotalWaitTime;
+					TotalTimeToComplete = other.TotalTimeToComplete;
 				}
 			}
 
 			if (AverageTotalTimeToComplete != null)
 			{
-				if (Other.AverageTotalTimeToComplete == null || Other.AverageTotalTimeToComplete.Value > AverageTotalTimeToComplete.Value)
+				if (other.AverageTotalTimeToComplete == null || other.AverageTotalTimeToComplete.Value > AverageTotalTimeToComplete.Value)
 				{
-					AverageTotalInitTime = Other.AverageTotalInitTime;
-					AverageTotalWaitTime = Other.AverageTotalWaitTime;
-					AverageTotalTimeToComplete = Other.AverageTotalTimeToComplete;
+					AverageTotalInitTime = other.AverageTotalInitTime;
+					AverageTotalWaitTime = other.AverageTotalWaitTime;
+					AverageTotalTimeToComplete = other.AverageTotalTimeToComplete;
 				}
 			}
 		}
@@ -440,25 +434,25 @@ namespace Horde.Build.Models
 		/// <summary>
 		/// Waits for all the given timing info objects to complete
 		/// </summary>
-		/// <param name="Others">Other timing info objects to wait for</param>
-		public void WaitForAll(IEnumerable<TimingInfo> Others)
+		/// <param name="others">Other timing info objects to wait for</param>
+		public void WaitForAll(IEnumerable<TimingInfo> others)
 		{
-			foreach(TimingInfo Other in Others)
+			foreach(TimingInfo other in others)
 			{
-				WaitFor(Other);
+				WaitFor(other);
 			}
 		}
 
 		/// <summary>
 		/// Constructs a new TimingInfo object which represents the last TimingInfo to finish
 		/// </summary>
-		/// <param name="Others">TimingInfo objects to wait for</param>
+		/// <param name="others">TimingInfo objects to wait for</param>
 		/// <returns>New TimingInfo instance</returns>
-		public static TimingInfo Max(IEnumerable<TimingInfo> Others)
+		public static TimingInfo Max(IEnumerable<TimingInfo> others)
 		{
-			TimingInfo TimingInfo = new TimingInfo();
-			TimingInfo.WaitForAll(Others);
-			return TimingInfo;
+			TimingInfo timingInfo = new TimingInfo();
+			timingInfo.WaitForAll(others);
+			return timingInfo;
 		}
 	}
 
@@ -486,7 +480,6 @@ namespace Horde.Build.Models
 	/// <summary>
 	/// Document describing a job
 	/// </summary>
-	[SuppressMessage("Compiler", "CA1056:URI parameters should not be strings")]
 	public interface IJob
 	{
 		/// <summary>
@@ -683,43 +676,43 @@ namespace Horde.Build.Models
 		/// <summary>
 		/// Attempts to get a step with a given ID
 		/// </summary>
-		/// <param name="Job">The job document</param>
-		/// <param name="StepId">The step id</param>
-		/// <param name="Step">On success, receives the step object</param>
+		/// <param name="job">The job document</param>
+		/// <param name="stepId">The step id</param>
+		/// <param name="step">On success, receives the step object</param>
 		/// <returns>True if the step was found</returns>
-		public static bool TryGetStep(this IJob Job, SubResourceId StepId, [NotNullWhen(true)] out IJobStep? Step)
+		public static bool TryGetStep(this IJob job, SubResourceId stepId, [NotNullWhen(true)] out IJobStep? step)
 		{
-			foreach (IJobStepBatch Batch in Job.Batches)
+			foreach (IJobStepBatch batch in job.Batches)
 			{
-				if (Batch.TryGetStep(StepId, out Step))
+				if (batch.TryGetStep(stepId, out step))
 				{
 					return true;
 				}
 			}
 
-			Step = null;
+			step = null;
 			return false;
 		}
 
 		/// <summary>
 		/// Gets the current job state
 		/// </summary>
-		/// <param name="Job">The job document</param>
+		/// <param name="job">The job document</param>
 		/// <returns>Job state</returns>
-		public static JobState GetState(this IJob Job)
+		public static JobState GetState(this IJob job)
 		{
 			bool bWaiting = false;
-			foreach (IJobStepBatch Batch in Job.Batches)
+			foreach (IJobStepBatch batch in job.Batches)
 			{
-				foreach (IJobStep Step in Batch.Steps)
+				foreach (IJobStep step in batch.Steps)
 				{
-					if (Step.State == JobStepState.Running)
+					if (step.State == JobStepState.Running)
 					{
 						return JobState.Running;
 					}
-					else if (Step.State == JobStepState.Ready || Step.State == JobStepState.Waiting)
+					else if (step.State == JobStepState.Ready || step.State == JobStepState.Waiting)
 					{
-						if (Batch.State == JobStepBatchState.Starting || Batch.State == JobStepBatchState.Running)
+						if (batch.State == JobStepBatchState.Starting || batch.State == JobStepBatchState.Running)
 						{
 							return JobState.Running;
 						}
@@ -736,35 +729,35 @@ namespace Horde.Build.Models
 		/// <summary>
 		/// Gets the outcome for a particular named target. May be an aggregate or node name.
 		/// </summary>
-		/// <param name="Job">The job to check</param>
+		/// <param name="job">The job to check</param>
 		/// <returns>The step outcome</returns>
-		public static (JobStepState, JobStepOutcome) GetTargetState(this IJob Job)
+		public static (JobStepState, JobStepOutcome) GetTargetState(this IJob job)
 		{
-			IReadOnlyDictionary<NodeRef, IJobStep> NodeToStep = GetStepForNodeMap(Job);
-			return GetTargetState(NodeToStep.Values);
+			IReadOnlyDictionary<NodeRef, IJobStep> nodeToStep = GetStepForNodeMap(job);
+			return GetTargetState(nodeToStep.Values);
 		}
 
 		/// <summary>
 		/// Gets the outcome for a particular named target. May be an aggregate or node name.
 		/// </summary>
-		/// <param name="Job">The job to check</param>
-		/// <param name="Graph">Graph for the job</param>
-		/// <param name="Target">Target to find an outcome for</param>
+		/// <param name="job">The job to check</param>
+		/// <param name="graph">Graph for the job</param>
+		/// <param name="target">Target to find an outcome for</param>
 		/// <returns>The step outcome</returns>
-		public static (JobStepState, JobStepOutcome)? GetTargetState(this IJob Job, IGraph Graph, string? Target)
+		public static (JobStepState, JobStepOutcome)? GetTargetState(this IJob job, IGraph graph, string? target)
 		{
-			if (Target == null)
+			if (target == null)
 			{
-				return GetTargetState(Job);
+				return GetTargetState(job);
 			}
 
-			NodeRef NodeRef;
-			if (Graph.TryFindNode(Target, out NodeRef))
+			NodeRef nodeRef;
+			if (graph.TryFindNode(target, out nodeRef))
 			{
-				IJobStep? Step;
-				if (Job.TryGetStepForNode(NodeRef, out Step))
+				IJobStep? step;
+				if (job.TryGetStepForNode(nodeRef, out step))
 				{
-					return (Step.State, Step.Outcome);
+					return (step.State, step.Outcome);
 				}
 				else
 				{
@@ -772,23 +765,23 @@ namespace Horde.Build.Models
 				}
 			}
 
-			IAggregate? Aggregate;
-			if (Graph.TryFindAggregate(Target, out Aggregate))
+			IAggregate? aggregate;
+			if (graph.TryFindAggregate(target, out aggregate))
 			{
-				IReadOnlyDictionary<NodeRef, IJobStep> StepForNode = GetStepForNodeMap(Job);
+				IReadOnlyDictionary<NodeRef, IJobStep> stepForNode = GetStepForNodeMap(job);
 
-				List<IJobStep> Steps = new List<IJobStep>();
-				foreach (NodeRef AggregateNodeRef in Aggregate.Nodes)
+				List<IJobStep> steps = new List<IJobStep>();
+				foreach (NodeRef aggregateNodeRef in aggregate.Nodes)
 				{
-					IJobStep? Step;
-					if (!StepForNode.TryGetValue(AggregateNodeRef, out Step))
+					IJobStep? step;
+					if (!stepForNode.TryGetValue(aggregateNodeRef, out step))
 					{
 						return null;
 					}
-					Steps.Add(Step);
+					steps.Add(step);
 				}
 
-				return GetTargetState(Steps);
+				return GetTargetState(steps);
 			}
 
 			return null;
@@ -797,43 +790,43 @@ namespace Horde.Build.Models
 		/// <summary>
 		/// Gets the outcome for a particular named target. May be an aggregate or node name.
 		/// </summary>
-		/// <param name="Steps">Steps to include</param>
+		/// <param name="steps">Steps to include</param>
 		/// <returns>The step outcome</returns>
-		public static (JobStepState, JobStepOutcome) GetTargetState(IEnumerable<IJobStep> Steps)
+		public static (JobStepState, JobStepOutcome) GetTargetState(IEnumerable<IJobStep> steps)
 		{
 			bool bAnySkipped = false;
 			bool bAnyWarnings = false;
 			bool bAnyFailed = false;
 			bool bAnyPending = false;
-			foreach (IJobStep Step in Steps)
+			foreach (IJobStep step in steps)
 			{
-				bAnyPending |= Step.IsPending();
-				bAnySkipped |= Step.State == JobStepState.Aborted || Step.State == JobStepState.Skipped;
-				bAnyFailed |= (Step.Outcome == JobStepOutcome.Failure);
-				bAnyWarnings |= (Step.Outcome == JobStepOutcome.Warnings);
+				bAnyPending |= step.IsPending();
+				bAnySkipped |= step.State == JobStepState.Aborted || step.State == JobStepState.Skipped;
+				bAnyFailed |= (step.Outcome == JobStepOutcome.Failure);
+				bAnyWarnings |= (step.Outcome == JobStepOutcome.Warnings);
 			}
 
-			JobStepState NewState = bAnyPending ? JobStepState.Running : JobStepState.Completed;
-			JobStepOutcome NewOutcome = bAnyFailed ? JobStepOutcome.Failure : bAnyWarnings ? JobStepOutcome.Warnings : bAnySkipped ? JobStepOutcome.Unspecified : JobStepOutcome.Success;
-			return (NewState, NewOutcome);
+			JobStepState newState = bAnyPending ? JobStepState.Running : JobStepState.Completed;
+			JobStepOutcome newOutcome = bAnyFailed ? JobStepOutcome.Failure : bAnyWarnings ? JobStepOutcome.Warnings : bAnySkipped ? JobStepOutcome.Unspecified : JobStepOutcome.Success;
+			return (newState, newOutcome);
 		}
 
 		/// <summary>
 		/// Gets the outcome for a particular named target. May be an aggregate or node name.
 		/// </summary>
-		/// <param name="Job">The job to check</param>
-		/// <param name="Graph">Graph for the job</param>
-		/// <param name="Target">Target to find an outcome for</param>
+		/// <param name="job">The job to check</param>
+		/// <param name="graph">Graph for the job</param>
+		/// <param name="target">Target to find an outcome for</param>
 		/// <returns>The step outcome</returns>
-		public static JobStepOutcome GetTargetOutcome(this IJob Job, IGraph Graph, string Target)
+		public static JobStepOutcome GetTargetOutcome(this IJob job, IGraph graph, string target)
 		{
-			NodeRef NodeRef;
-			if (Graph.TryFindNode(Target, out NodeRef))
+			NodeRef nodeRef;
+			if (graph.TryFindNode(target, out nodeRef))
 			{
-				IJobStep? Step;
-				if (Job.TryGetStepForNode(NodeRef, out Step))
+				IJobStep? step;
+				if (job.TryGetStepForNode(nodeRef, out step))
 				{
-					return Step.Outcome;
+					return step.Outcome;
 				}
 				else
 				{
@@ -841,26 +834,26 @@ namespace Horde.Build.Models
 				}
 			}
 
-			IAggregate? Aggregate;
-			if (Graph.TryFindAggregate(Target, out Aggregate))
+			IAggregate? aggregate;
+			if (graph.TryFindAggregate(target, out aggregate))
 			{
-				IReadOnlyDictionary<NodeRef, IJobStep> StepForNode = GetStepForNodeMap(Job);
+				IReadOnlyDictionary<NodeRef, IJobStep> stepForNode = GetStepForNodeMap(job);
 
-				bool Warnings = false;
-				foreach (NodeRef AggregateNodeRef in Aggregate.Nodes)
+				bool warnings = false;
+				foreach (NodeRef aggregateNodeRef in aggregate.Nodes)
 				{
-					IJobStep? Step;
-					if (!StepForNode.TryGetValue(AggregateNodeRef, out Step))
+					IJobStep? step;
+					if (!stepForNode.TryGetValue(aggregateNodeRef, out step))
 					{
 						return JobStepOutcome.Unspecified;
 					}
-					if (Step.Outcome == JobStepOutcome.Failure)
+					if (step.Outcome == JobStepOutcome.Failure)
 					{
 						return JobStepOutcome.Failure;
 					}
-					Warnings |= (Step.Outcome == JobStepOutcome.Warnings);
+					warnings |= (step.Outcome == JobStepOutcome.Warnings);
 				}
-				return Warnings ? JobStepOutcome.Warnings : JobStepOutcome.Success;
+				return warnings ? JobStepOutcome.Warnings : JobStepOutcome.Success;
 			}
 
 			return JobStepOutcome.Unspecified;
@@ -869,468 +862,472 @@ namespace Horde.Build.Models
 		/// <summary>
 		/// Gets the job step for a particular node
 		/// </summary>
-		/// <param name="Job">The job to search</param>
-		/// <param name="NodeRef">The node ref</param>
-		/// <param name="JobStep">Receives the jobstep on success</param>
+		/// <param name="job">The job to search</param>
+		/// <param name="nodeRef">The node ref</param>
+		/// <param name="jobStep">Receives the jobstep on success</param>
 		/// <returns>True if the jobstep was founds</returns>
-		public static bool TryGetStepForNode(this IJob Job, NodeRef NodeRef, [NotNullWhen(true)] out IJobStep? JobStep)
+		public static bool TryGetStepForNode(this IJob job, NodeRef nodeRef, [NotNullWhen(true)] out IJobStep? jobStep)
 		{
-			JobStep = null;
-			foreach (IJobStepBatch Batch in Job.Batches)
+			jobStep = null;
+			foreach (IJobStepBatch batch in job.Batches)
 			{
-				if (Batch.GroupIdx == NodeRef.GroupIdx)
+				if (batch.GroupIdx == nodeRef.GroupIdx)
 				{
-					foreach (IJobStep BatchStep in Batch.Steps)
+					foreach (IJobStep batchStep in batch.Steps)
 					{
-						if (BatchStep.NodeIdx == NodeRef.NodeIdx)
+						if (batchStep.NodeIdx == nodeRef.NodeIdx)
 						{
-							JobStep = BatchStep;
+							jobStep = batchStep;
 						}
 					}
 				}
 			}
-			return JobStep != null;
+			return jobStep != null;
 		}
 
 		/// <summary>
 		/// Gets a dictionary that maps <see cref="NodeRef"/> objects to their associated
 		/// <see cref="IJobStep"/> objects on a <see cref="IJob"/>.
 		/// </summary>
-		/// <param name="Job">The job document</param>
+		/// <param name="job">The job document</param>
 		/// <returns>Map of <see cref="NodeRef"/> to <see cref="IJobStep"/></returns>
-		public static IReadOnlyDictionary<NodeRef, IJobStep> GetStepForNodeMap(this IJob Job)
+		public static IReadOnlyDictionary<NodeRef, IJobStep> GetStepForNodeMap(this IJob job)
 		{
-			Dictionary<NodeRef, IJobStep> StepForNode = new Dictionary<NodeRef, IJobStep>();
-			foreach (IJobStepBatch Batch in Job.Batches)
+			Dictionary<NodeRef, IJobStep> stepForNode = new Dictionary<NodeRef, IJobStep>();
+			foreach (IJobStepBatch batch in job.Batches)
 			{
-				foreach (IJobStep BatchStep in Batch.Steps)
+				foreach (IJobStep batchStep in batch.Steps)
 				{
-					NodeRef BatchNodeRef = new NodeRef(Batch.GroupIdx, BatchStep.NodeIdx);
-					StepForNode[BatchNodeRef] = BatchStep;
+					NodeRef batchNodeRef = new NodeRef(batch.GroupIdx, batchStep.NodeIdx);
+					stepForNode[batchNodeRef] = batchStep;
 				}
 			}
-			return StepForNode;
+			return stepForNode;
 		}
 
 		/// <summary>
 		/// Gets the estimated timing info for all nodes in the job
 		/// </summary>
-		/// <param name="Job">The job document</param>
-		/// <param name="Graph">Graph for this job</param>
-		/// <param name="JobTiming">Job timing information</param>
+		/// <param name="job">The job document</param>
+		/// <param name="graph">Graph for this job</param>
+		/// <param name="jobTiming">Job timing information</param>
 		/// <returns>Map of node to expected timing info</returns>
-		public static Dictionary<INode, TimingInfo> GetTimingInfo(this IJob Job, IGraph Graph, IJobTiming JobTiming)
+		public static Dictionary<INode, TimingInfo> GetTimingInfo(this IJob job, IGraph graph, IJobTiming jobTiming)
 		{
-			TimeSpan CurrentTime = DateTime.UtcNow - Job.CreateTimeUtc;
+#pragma warning disable IDE0054 // Use compound assignment
+			TimeSpan currentTime = DateTime.UtcNow - job.CreateTimeUtc;
 
-			Dictionary<INode, TimingInfo> NodeToTimingInfo = Graph.Groups.SelectMany(x => x.Nodes).ToDictionary(x => x, x => new TimingInfo());
-			foreach (IJobStepBatch Batch in Job.Batches)
+			Dictionary<INode, TimingInfo> nodeToTimingInfo = graph.Groups.SelectMany(x => x.Nodes).ToDictionary(x => x, x => new TimingInfo());
+			foreach (IJobStepBatch batch in job.Batches)
 			{
-				INodeGroup Group = Graph.Groups[Batch.GroupIdx];
+				INodeGroup group = graph.Groups[batch.GroupIdx];
 
 				// Step through the batch, keeping track of the time that things finish.
-				TimingInfo TimingInfo = new TimingInfo();
+				TimingInfo timingInfo = new TimingInfo();
 
 				// Wait for the dependencies for the batch to start
-				HashSet<INode> DependencyNodes = Batch.GetStartDependencies(Graph.Groups);
-				TimingInfo.WaitForAll(DependencyNodes.Select(x => NodeToTimingInfo[x]));
+				HashSet<INode> dependencyNodes = batch.GetStartDependencies(graph.Groups);
+				timingInfo.WaitForAll(dependencyNodes.Select(x => nodeToTimingInfo[x]));
 
 				// If the batch has actually started, correct the expected time to use this instead
-				if (Batch.StartTimeUtc != null)
+				if (batch.StartTimeUtc != null)
 				{
-					TimingInfo.TotalTimeToComplete = Batch.StartTimeUtc - Job.CreateTimeUtc;
+					timingInfo.TotalTimeToComplete = batch.StartTimeUtc - job.CreateTimeUtc;
 				}
 
 				// Get the average times for this batch
-				TimeSpan? AverageWaitTime = GetAverageWaitTime(Graph, Batch, JobTiming);
-				TimeSpan? AverageInitTime = GetAverageInitTime(Graph, Batch, JobTiming);
+				TimeSpan? averageWaitTime = GetAverageWaitTime(graph, batch, jobTiming);
+				TimeSpan? averageInitTime = GetAverageInitTime(graph, batch, jobTiming);
 
 				// Update the wait times and initialization times along this path
-				TimingInfo.TotalWaitTime = TimingInfo.TotalWaitTime + (Batch.GetWaitTime() ?? AverageWaitTime);
-				TimingInfo.TotalInitTime = TimingInfo.TotalInitTime + (Batch.GetInitTime() ?? AverageInitTime);
+				timingInfo.TotalWaitTime = timingInfo.TotalWaitTime + (batch.GetWaitTime() ?? averageWaitTime);
+				timingInfo.TotalInitTime = timingInfo.TotalInitTime + (batch.GetInitTime() ?? averageInitTime);
 
 				// Update the average wait and initialization times too
-				TimingInfo.AverageTotalWaitTime = TimingInfo.AverageTotalWaitTime + AverageWaitTime;
-				TimingInfo.AverageTotalInitTime = TimingInfo.AverageTotalInitTime + AverageInitTime;
+				timingInfo.AverageTotalWaitTime = timingInfo.AverageTotalWaitTime + averageWaitTime;
+				timingInfo.AverageTotalInitTime = timingInfo.AverageTotalInitTime + averageInitTime;
 
 				// Step through the batch, updating the expected times as we go
-				foreach (IJobStep Step in Batch.Steps)
+				foreach (IJobStep step in batch.Steps)
 				{
-					INode Node = Group.Nodes[Step.NodeIdx];
+					INode node = group.Nodes[step.NodeIdx];
 
 					// Get the timing for this step
-					IJobStepTiming? StepTimingInfo;
-					JobTiming.TryGetStepTiming(Node.Name, out StepTimingInfo);
+					IJobStepTiming? stepTimingInfo;
+					jobTiming.TryGetStepTiming(node.Name, out stepTimingInfo);
 
 					// If the step has already started, update the actual time to reach this point
-					if(Step.StartTimeUtc != null)
+					if(step.StartTimeUtc != null)
 					{
-						TimingInfo.TotalTimeToComplete = Step.StartTimeUtc.Value - Job.CreateTimeUtc;
+						timingInfo.TotalTimeToComplete = step.StartTimeUtc.Value - job.CreateTimeUtc;
 					}
 
 					// If the step hasn't started yet, make sure the start time is later than the current time
-					if (Step.StartTimeUtc == null && CurrentTime > TimingInfo.TotalTimeToComplete)
+					if (step.StartTimeUtc == null && currentTime > timingInfo.TotalTimeToComplete)
 					{
-						TimingInfo.TotalTimeToComplete = CurrentTime;
+						timingInfo.TotalTimeToComplete = currentTime;
 					}
 
 					// Wait for all the node dependencies to complete
-					TimingInfo.WaitForAll(Graph.GetDependencies(Node).Select(x => NodeToTimingInfo[x]));
+					timingInfo.WaitForAll(graph.GetDependencies(node).Select(x => nodeToTimingInfo[x]));
 
 					// If the step has actually finished, correct the time to use that instead
-					if (Step.FinishTimeUtc != null)
+					if (step.FinishTimeUtc != null)
 					{
-						TimingInfo.TotalTimeToComplete = Step.FinishTimeUtc.Value - Job.CreateTimeUtc;
+						timingInfo.TotalTimeToComplete = step.FinishTimeUtc.Value - job.CreateTimeUtc;
 					}
 					else
 					{
-						TimingInfo.TotalTimeToComplete = TimingInfo.TotalTimeToComplete + NullableTimeSpanFromSeconds(StepTimingInfo?.AverageDuration);
+						timingInfo.TotalTimeToComplete = timingInfo.TotalTimeToComplete + NullableTimeSpanFromSeconds(stepTimingInfo?.AverageDuration);
 					}
 
 					// If the step hasn't finished yet, make sure the start time is later than the current time
-					if (Step.FinishTimeUtc == null && CurrentTime > TimingInfo.TotalTimeToComplete)
+					if (step.FinishTimeUtc == null && currentTime > timingInfo.TotalTimeToComplete)
 					{
-						TimingInfo.TotalTimeToComplete = CurrentTime;
+						timingInfo.TotalTimeToComplete = currentTime;
 					}
 
 					// Update the average time to complete
-					TimingInfo.AverageTotalTimeToComplete = TimingInfo.AverageTotalTimeToComplete + NullableTimeSpanFromSeconds(StepTimingInfo?.AverageDuration);
+					timingInfo.AverageTotalTimeToComplete = timingInfo.AverageTotalTimeToComplete + NullableTimeSpanFromSeconds(stepTimingInfo?.AverageDuration);
 
 					// Add it to the lookup
-					TimingInfo NodeTimingInfo = new TimingInfo(TimingInfo);
-					NodeTimingInfo.StepTiming = StepTimingInfo;
-					NodeToTimingInfo[Node] = NodeTimingInfo;
+					TimingInfo nodeTimingInfo = new TimingInfo(timingInfo);
+					nodeTimingInfo.StepTiming = stepTimingInfo;
+					nodeToTimingInfo[node] = nodeTimingInfo;
 				}
 			}
-			return NodeToTimingInfo;
+			return nodeToTimingInfo;
+#pragma warning restore IDE0054 // Use compound assignment
 		}
 
 		/// <summary>
 		/// Gets the average wait time for this batch
 		/// </summary>
-		/// <param name="Graph">Graph for the job</param>
-		/// <param name="Batch">The batch to get timing info for</param>
-		/// <param name="JobTiming">The job timing information</param>
+		/// <param name="graph">Graph for the job</param>
+		/// <param name="batch">The batch to get timing info for</param>
+		/// <param name="jobTiming">The job timing information</param>
 		/// <returns>Wait time for the batch</returns>
-		public static TimeSpan? GetAverageWaitTime(IGraph Graph, IJobStepBatch Batch, IJobTiming JobTiming)
+		public static TimeSpan? GetAverageWaitTime(IGraph graph, IJobStepBatch batch, IJobTiming jobTiming)
 		{
-			TimeSpan? WaitTime = null;
-			foreach (IJobStep Step in Batch.Steps)
+			TimeSpan? waitTime = null;
+			foreach (IJobStep step in batch.Steps)
 			{
-				INode Node = Graph.Groups[Batch.GroupIdx].Nodes[Step.NodeIdx];
-				if(JobTiming.TryGetStepTiming(Node.Name, out IJobStepTiming? TimingInfo))
+				INode node = graph.Groups[batch.GroupIdx].Nodes[step.NodeIdx];
+				if(jobTiming.TryGetStepTiming(node.Name, out IJobStepTiming? timingInfo))
 				{
-					if (TimingInfo.AverageWaitTime != null)
+					if (timingInfo.AverageWaitTime != null)
 					{
-						TimeSpan StepWaitTime = TimeSpan.FromSeconds(TimingInfo.AverageWaitTime.Value);
-						if (WaitTime == null || StepWaitTime > WaitTime.Value)
+						TimeSpan stepWaitTime = TimeSpan.FromSeconds(timingInfo.AverageWaitTime.Value);
+						if (waitTime == null || stepWaitTime > waitTime.Value)
 						{
-							WaitTime = StepWaitTime;
+							waitTime = stepWaitTime;
 						}
 					}
 				}
 			}
-			return WaitTime;
+			return waitTime;
 		}
 
 		/// <summary>
 		/// Gets the average initialization time for this batch
 		/// </summary>
-		/// <param name="Graph">Graph for the job</param>
-		/// <param name="Batch">The batch to get timing info for</param>
-		/// <param name="JobTiming">The job timing information</param>
+		/// <param name="graph">Graph for the job</param>
+		/// <param name="batch">The batch to get timing info for</param>
+		/// <param name="jobTiming">The job timing information</param>
 		/// <returns>Initialization time for this batch</returns>
-		public static TimeSpan? GetAverageInitTime(IGraph Graph, IJobStepBatch Batch, IJobTiming JobTiming)
+		public static TimeSpan? GetAverageInitTime(IGraph graph, IJobStepBatch batch, IJobTiming jobTiming)
 		{
-			TimeSpan? InitTime = null;
-			foreach (IJobStep Step in Batch.Steps)
+			TimeSpan? initTime = null;
+			foreach (IJobStep step in batch.Steps)
 			{
-				INode Node = Graph.Groups[Batch.GroupIdx].Nodes[Step.NodeIdx];
-				if (JobTiming.TryGetStepTiming(Node.Name, out IJobStepTiming? TimingInfo))
+				INode node = graph.Groups[batch.GroupIdx].Nodes[step.NodeIdx];
+				if (jobTiming.TryGetStepTiming(node.Name, out IJobStepTiming? timingInfo))
 				{
-					if (TimingInfo.AverageInitTime != null)
+					if (timingInfo.AverageInitTime != null)
 					{
-						TimeSpan StepInitTime = TimeSpan.FromSeconds(TimingInfo.AverageInitTime.Value);
-						if (InitTime == null || StepInitTime > InitTime.Value)
+						TimeSpan stepInitTime = TimeSpan.FromSeconds(timingInfo.AverageInitTime.Value);
+						if (initTime == null || stepInitTime > initTime.Value)
 						{
-							InitTime = StepInitTime;
+							initTime = stepInitTime;
 						}
 					}
 				}
 			}
-			return InitTime;
+			return initTime;
 		}
 
 		/// <summary>
 		/// Creates a nullable timespan from a nullable number of seconds
 		/// </summary>
-		/// <param name="Seconds">The number of seconds to construct from</param>
+		/// <param name="seconds">The number of seconds to construct from</param>
 		/// <returns>TimeSpan object</returns>
-		static TimeSpan? NullableTimeSpanFromSeconds(float? Seconds)
+		static TimeSpan? NullableTimeSpanFromSeconds(float? seconds)
 		{
-			if (Seconds == null)
+			if (seconds == null)
 			{
 				return null;
 			}
 			else
 			{
-				return TimeSpan.FromSeconds(Seconds.Value);
+				return TimeSpan.FromSeconds(seconds.Value);
 			}
 		}
 
 		/// <summary>
 		/// Attempts to get a batch with the given id
 		/// </summary>
-		/// <param name="Job">The job document</param>
-		/// <param name="BatchId">The batch id</param>
-		/// <param name="Batch">On success, receives the batch object</param>
+		/// <param name="job">The job document</param>
+		/// <param name="batchId">The batch id</param>
+		/// <param name="batch">On success, receives the batch object</param>
 		/// <returns>True if the batch was found</returns>
-		public static bool TryGetBatch(this IJob Job, SubResourceId BatchId, [NotNullWhen(true)] out IJobStepBatch? Batch)
+		public static bool TryGetBatch(this IJob job, SubResourceId batchId, [NotNullWhen(true)] out IJobStepBatch? batch)
 		{
-			Batch = Job.Batches.FirstOrDefault(x => x.Id == BatchId);
-			return Batch != null;
+			batch = job.Batches.FirstOrDefault(x => x.Id == batchId);
+			return batch != null;
 		}
 
 		/// <summary>
 		/// Attempts to get a batch with the given id
 		/// </summary>
-		/// <param name="Job">The job document</param>
-		/// <param name="BatchId">The batch id</param>
-		/// <param name="StepId">The step id</param>
-		/// <param name="Step">On success, receives the step object</param>
+		/// <param name="job">The job document</param>
+		/// <param name="batchId">The batch id</param>
+		/// <param name="stepId">The step id</param>
+		/// <param name="step">On success, receives the step object</param>
 		/// <returns>True if the batch was found</returns>
-		public static bool TryGetStep(this IJob Job, SubResourceId BatchId, SubResourceId StepId, [NotNullWhen(true)] out IJobStep? Step)
+		public static bool TryGetStep(this IJob job, SubResourceId batchId, SubResourceId stepId, [NotNullWhen(true)] out IJobStep? step)
 		{
-			IJobStepBatch? Batch;
-			if (!TryGetBatch(Job, BatchId, out Batch))
+			IJobStepBatch? batch;
+			if (!TryGetBatch(job, batchId, out batch))
 			{
-				Step = null;
+				step = null;
 				return false;
 			}
-			return Batch.TryGetStep(StepId, out Step);
+			return batch.TryGetStep(stepId, out step);
 		}
 
 		/// <summary>
 		/// Finds the set of nodes affected by a label
 		/// </summary>
-		/// <param name="Job">The job document</param>
-		/// <param name="Graph">Graph definition for the job</param>
-		/// <param name="LabelIdx">Index of the label. -1 or Graph.Labels.Count are treated as referring to the default lable.</param>
+		/// <param name="job">The job document</param>
+		/// <param name="graph">Graph definition for the job</param>
+		/// <param name="labelIdx">Index of the label. -1 or Graph.Labels.Count are treated as referring to the default lable.</param>
 		/// <returns>Set of nodes affected by the given label</returns>
-		public static HashSet<NodeRef> GetNodesForLabel(this IJob Job, IGraph Graph, int LabelIdx)
+		public static HashSet<NodeRef> GetNodesForLabel(this IJob job, IGraph graph, int labelIdx)
 		{
-			if (LabelIdx != -1 && LabelIdx != Graph.Labels.Count)
+			if (labelIdx != -1 && labelIdx != graph.Labels.Count)
 			{
 				// Return all the nodes included by the label
-				return new HashSet<NodeRef>(Graph.Labels[LabelIdx].IncludedNodes);
+				return new HashSet<NodeRef>(graph.Labels[labelIdx].IncludedNodes);
 			}
 			else
 			{
 				// Set of nodes which are not covered by an existing label, initially containing everything
-				HashSet<NodeRef> UnlabeledNodes = new HashSet<NodeRef>();
-				for (int GroupIdx = 0; GroupIdx < Graph.Groups.Count; GroupIdx++)
+				HashSet<NodeRef> unlabeledNodes = new HashSet<NodeRef>();
+				for (int groupIdx = 0; groupIdx < graph.Groups.Count; groupIdx++)
 				{
-					INodeGroup Group = Graph.Groups[GroupIdx];
-					for (int NodeIdx = 0; NodeIdx < Group.Nodes.Count; NodeIdx++)
+					INodeGroup group = graph.Groups[groupIdx];
+					for (int nodeIdx = 0; nodeIdx < group.Nodes.Count; nodeIdx++)
 					{
-						UnlabeledNodes.Add(new NodeRef(GroupIdx, NodeIdx));
+						unlabeledNodes.Add(new NodeRef(groupIdx, nodeIdx));
 					}
 				}
 
 				// Remove all the nodes that are part of an active label
-				IReadOnlyDictionary<NodeRef, IJobStep> StepForNode = Job.GetStepForNodeMap();
-				foreach (ILabel Label in Graph.Labels)
+				IReadOnlyDictionary<NodeRef, IJobStep> stepForNode = job.GetStepForNodeMap();
+				foreach (ILabel label in graph.Labels)
 				{
-					if (Label.RequiredNodes.Any(x => StepForNode.ContainsKey(x)))
+					if (label.RequiredNodes.Any(x => stepForNode.ContainsKey(x)))
 					{
-						UnlabeledNodes.ExceptWith(Label.IncludedNodes);
+						unlabeledNodes.ExceptWith(label.IncludedNodes);
 					}
 				}
-				return UnlabeledNodes;
+				return unlabeledNodes;
 			}
 		}
 
 		/// <summary>
 		/// Create a list of aggregate responses, combining the graph definitions with the state of the job
 		/// </summary>
-		/// <param name="Job">The job document</param>
-		/// <param name="Graph">Graph definition for the job</param>
-		/// <param name="Responses">List to receive all the responses</param>
+		/// <param name="job">The job document</param>
+		/// <param name="graph">Graph definition for the job</param>
+		/// <param name="responses">List to receive all the responses</param>
 		/// <returns>The default label state</returns>
-		public static GetDefaultLabelStateResponse? GetLabelStateResponses(this IJob Job, IGraph Graph, List<GetLabelStateResponse> Responses)
+		public static GetDefaultLabelStateResponse? GetLabelStateResponses(this IJob job, IGraph graph, List<GetLabelStateResponse> responses)
 		{
 			// Create a lookup from noderef to step information
-			IReadOnlyDictionary<NodeRef, IJobStep> StepForNode = Job.GetStepForNodeMap();
+			IReadOnlyDictionary<NodeRef, IJobStep> stepForNode = job.GetStepForNodeMap();
 
 			// Set of nodes which are not covered by an existing label, initially containing everything
-			HashSet<NodeRef> UnlabeledNodes = new HashSet<NodeRef>();
-			for (int GroupIdx = 0; GroupIdx < Graph.Groups.Count; GroupIdx++)
+			HashSet<NodeRef> unlabeledNodes = new HashSet<NodeRef>();
+			for (int groupIdx = 0; groupIdx < graph.Groups.Count; groupIdx++)
 			{
-				INodeGroup Group = Graph.Groups[GroupIdx];
-				for (int NodeIdx = 0; NodeIdx < Group.Nodes.Count; NodeIdx++)
+				INodeGroup group = graph.Groups[groupIdx];
+				for (int nodeIdx = 0; nodeIdx < group.Nodes.Count; nodeIdx++)
 				{
-					UnlabeledNodes.Add(new NodeRef(GroupIdx, NodeIdx));
+					unlabeledNodes.Add(new NodeRef(groupIdx, nodeIdx));
 				}
 			}
 
 			// Create the responses
-			foreach (ILabel Label in Graph.Labels)
+			foreach (ILabel label in graph.Labels)
 			{
 				// Refresh the state for this label
-				LabelState NewState = LabelState.Unspecified;
-				foreach (NodeRef RequiredNodeRef in Label.RequiredNodes)
+				LabelState newState = LabelState.Unspecified;
+				foreach (NodeRef requiredNodeRef in label.RequiredNodes)
 				{
-					if (StepForNode.ContainsKey(RequiredNodeRef))
+					if (stepForNode.ContainsKey(requiredNodeRef))
 					{
-						NewState = LabelState.Complete;
+						newState = LabelState.Complete;
 						break;
 					}
 				}
 
 				// Refresh the outcome
-				LabelOutcome NewOutcome = LabelOutcome.Success;
-				if (NewState == LabelState.Complete)
+				LabelOutcome newOutcome = LabelOutcome.Success;
+				if (newState == LabelState.Complete)
 				{
-					GetLabelState(Label.IncludedNodes, StepForNode, out NewState, out NewOutcome);
-					UnlabeledNodes.ExceptWith(Label.IncludedNodes);
+					GetLabelState(label.IncludedNodes, stepForNode, out newState, out newOutcome);
+					unlabeledNodes.ExceptWith(label.IncludedNodes);
 				}
 
 				// Create the response
-				Responses.Add(new GetLabelStateResponse(NewState, NewOutcome));
+				responses.Add(new GetLabelStateResponse(newState, newOutcome));
 			}
 
 			// Remove all the nodes that don't have a step
-			UnlabeledNodes.RemoveWhere(x => !StepForNode.ContainsKey(x));
+			unlabeledNodes.RemoveWhere(x => !stepForNode.ContainsKey(x));
 
 			// Remove successful "setup build" nodes from the list
-			if (Graph.Groups.Count > 1 && Graph.Groups[0].Nodes.Count > 0)
+			if (graph.Groups.Count > 1 && graph.Groups[0].Nodes.Count > 0)
 			{
-				INode Node = Graph.Groups[0].Nodes[0];
-				if (Node.Name == IJob.SetupNodeName)
+				INode node = graph.Groups[0].Nodes[0];
+				if (node.Name == IJob.SetupNodeName)
 				{
-					NodeRef NodeRef = new NodeRef(0, 0);
-					if (UnlabeledNodes.Contains(NodeRef))
+					NodeRef nodeRef = new NodeRef(0, 0);
+					if (unlabeledNodes.Contains(nodeRef))
 					{
-						IJobStep Step = StepForNode[NodeRef];
-						if (Step.State == JobStepState.Completed && Step.Outcome == JobStepOutcome.Success && Responses.Count > 0)
+						IJobStep step = stepForNode[nodeRef];
+						if (step.State == JobStepState.Completed && step.Outcome == JobStepOutcome.Success && responses.Count > 0)
 						{
-							UnlabeledNodes.Remove(NodeRef);
+							unlabeledNodes.Remove(nodeRef);
 						}
 					}
 				}
 			}
 
 			// Add a response for everything not included elsewhere.
-			GetLabelState(UnlabeledNodes, StepForNode, out LabelState OtherState, out LabelOutcome OtherOutcome);
-			return new GetDefaultLabelStateResponse(OtherState, OtherOutcome, UnlabeledNodes.Select(x => Graph.GetNode(x).Name).ToList());
+			GetLabelState(unlabeledNodes, stepForNode, out LabelState otherState, out LabelOutcome otherOutcome);
+			return new GetDefaultLabelStateResponse(otherState, otherOutcome, unlabeledNodes.Select(x => graph.GetNode(x).Name).ToList());
 		}
 
 		/// <summary>
 		/// Get the states of all labels for this job
 		/// </summary>
-		/// <param name="Job">The job to get states for</param>
-		/// <param name="Graph">The graph for this job</param>
+		/// <param name="job">The job to get states for</param>
+		/// <param name="graph">The graph for this job</param>
 		/// <returns>Collection of label states by label index</returns>
-		public static IReadOnlyList<(LabelState, LabelOutcome)> GetLabelStates(this IJob Job, IGraph Graph)
+		public static IReadOnlyList<(LabelState, LabelOutcome)> GetLabelStates(this IJob job, IGraph graph)
 		{
-			IReadOnlyDictionary<NodeRef, IJobStep> StepForNodeRef = Job.GetStepForNodeMap();
+			IReadOnlyDictionary<NodeRef, IJobStep> stepForNodeRef = job.GetStepForNodeMap();
 
-			List<(LabelState, LabelOutcome)> States = new List<(LabelState, LabelOutcome)>();
-			for (int Idx = 0; Idx < Graph.Labels.Count; Idx++)
+			List<(LabelState, LabelOutcome)> states = new List<(LabelState, LabelOutcome)>();
+			for (int idx = 0; idx < graph.Labels.Count; idx++)
 			{
-				ILabel Label = Graph.Labels[Idx];
+				ILabel label = graph.Labels[idx];
 
 				// Default the label to the unspecified state
-				LabelState NewState = LabelState.Unspecified;
-				LabelOutcome NewOutcome = LabelOutcome.Unspecified;
+				LabelState newState = LabelState.Unspecified;
+				LabelOutcome newOutcome = LabelOutcome.Unspecified;
 
 				// Check if the label should be included
-				if (Label.RequiredNodes.Any(x => StepForNodeRef.ContainsKey(x)))
+				if (label.RequiredNodes.Any(x => stepForNodeRef.ContainsKey(x)))
 				{
 					// Combine the state of the steps contributing towards this label
 					bool bAnySkipped = false;
 					bool bAnyWarnings = false;
 					bool bAnyFailed = false;
 					bool bAnyPending = false;
-					foreach (NodeRef IncludedNode in Label.IncludedNodes)
+					foreach (NodeRef includedNode in label.IncludedNodes)
 					{
-						IJobStep? Step;
-						if (StepForNodeRef.TryGetValue(IncludedNode, out Step))
+						IJobStep? step;
+						if (stepForNodeRef.TryGetValue(includedNode, out step))
 						{
-							bAnyPending |= Step.IsPending();
-							bAnySkipped |= Step.State == JobStepState.Aborted || Step.State == JobStepState.Skipped;
-							bAnyFailed |= (Step.Outcome == JobStepOutcome.Failure);
-							bAnyWarnings |= (Step.Outcome == JobStepOutcome.Warnings);
+							bAnyPending |= step.IsPending();
+							bAnySkipped |= step.State == JobStepState.Aborted || step.State == JobStepState.Skipped;
+							bAnyFailed |= (step.Outcome == JobStepOutcome.Failure);
+							bAnyWarnings |= (step.Outcome == JobStepOutcome.Warnings);
 						}
 					}
 
 					// Figure out the overall label state
-					NewState = bAnyPending ? LabelState.Running : LabelState.Complete;
-					NewOutcome = bAnyFailed ? LabelOutcome.Failure : bAnyWarnings ? LabelOutcome.Warnings : bAnySkipped? LabelOutcome.Unspecified : LabelOutcome.Success;
+					newState = bAnyPending ? LabelState.Running : LabelState.Complete;
+					newOutcome = bAnyFailed ? LabelOutcome.Failure : bAnyWarnings ? LabelOutcome.Warnings : bAnySkipped? LabelOutcome.Unspecified : LabelOutcome.Success;
 				}
 
-				States.Add((NewState, NewOutcome));
+				states.Add((newState, newOutcome));
 			}
-			return States;	
+			return states;	
 		}
 
 		/// <summary>
 		/// Get the states of all UGS badges for this job
 		/// </summary>
-		/// <param name="Job">The job to get states for</param>
-		/// <param name="Graph">The graph for this job</param>
+		/// <param name="job">The job to get states for</param>
+		/// <param name="graph">The graph for this job</param>
 		/// <returns>List of badge states</returns>
-		public static Dictionary<int, UgsBadgeState> GetUgsBadgeStates(this IJob Job, IGraph Graph)
+		public static Dictionary<int, UgsBadgeState> GetUgsBadgeStates(this IJob job, IGraph graph)
 		{
-			IReadOnlyList<(LabelState, LabelOutcome)> LabelStates = GetLabelStates(Job, Graph);
-			return Job.GetUgsBadgeStates(Graph, LabelStates);
+			IReadOnlyList<(LabelState, LabelOutcome)> labelStates = GetLabelStates(job, graph);
+			return job.GetUgsBadgeStates(graph, labelStates);
 		}
 
 		/// <summary>
 		/// Get the states of all UGS badges for this job
 		/// </summary>
-		/// <param name="Job">The job to get states for</param>
-		/// <param name="Graph">The graph for this job</param>
-		/// <param name="LabelStates">The existing label states to get the UGS badge states from</param>
+		/// <param name="job">The job to get states for</param>
+		/// <param name="graph">The graph for this job</param>
+		/// <param name="labelStates">The existing label states to get the UGS badge states from</param>
 		/// <returns>List of badge states</returns>
-		public static Dictionary<int, UgsBadgeState> GetUgsBadgeStates(this IJob Job, IGraph Graph, IReadOnlyList<(LabelState, LabelOutcome)> LabelStates)
+#pragma warning disable IDE0060 // Remove unused parameter
+		public static Dictionary<int, UgsBadgeState> GetUgsBadgeStates(this IJob job, IGraph graph, IReadOnlyList<(LabelState, LabelOutcome)> labelStates)
+#pragma warning restore IDE0060 // Remove unused parameter
 		{
-			Dictionary<int, UgsBadgeState> UgsBadgeStates = new Dictionary<int, UgsBadgeState>();
-			for (int LabelIdx = 0; LabelIdx < LabelStates.Count; ++LabelIdx)
+			Dictionary<int, UgsBadgeState> ugsBadgeStates = new Dictionary<int, UgsBadgeState>();
+			for (int labelIdx = 0; labelIdx < labelStates.Count; ++labelIdx)
 			{
-				if (Graph.Labels[LabelIdx].UgsName == null)
+				if (graph.Labels[labelIdx].UgsName == null)
 				{
 					continue;
 				}
 
-				(LabelState State, LabelOutcome Outcome) Label = LabelStates[LabelIdx];
-				switch (Label.State)
+				(LabelState state, LabelOutcome outcome) = labelStates[labelIdx];
+				switch (state)
 				{
 					case LabelState.Complete:
 					{
-						switch (Label.Outcome)
+						switch (outcome)
 						{
 							case LabelOutcome.Success:
 							{
-								UgsBadgeStates.Add(LabelIdx, UgsBadgeState.Success);
+								ugsBadgeStates.Add(labelIdx, UgsBadgeState.Success);
 								break;
 							}
 
 							case LabelOutcome.Warnings:
 							{
-								UgsBadgeStates.Add(LabelIdx, UgsBadgeState.Warning);
+								ugsBadgeStates.Add(labelIdx, UgsBadgeState.Warning);
 								break;
 							}
 
 							case LabelOutcome.Failure:
 							{
-								UgsBadgeStates.Add(LabelIdx, UgsBadgeState.Failure);
+								ugsBadgeStates.Add(labelIdx, UgsBadgeState.Failure);
 								break;
 							}
 
 							case LabelOutcome.Unspecified:
 							{
-								UgsBadgeStates.Add(LabelIdx, UgsBadgeState.Skipped);
+								ugsBadgeStates.Add(labelIdx, UgsBadgeState.Skipped);
 								break;
 							}
 						}
@@ -1339,70 +1336,70 @@ namespace Horde.Build.Models
 
 					case LabelState.Running:
 					{
-						UgsBadgeStates.Add(LabelIdx, UgsBadgeState.Starting);
+						ugsBadgeStates.Add(labelIdx, UgsBadgeState.Starting);
 						break;
 					}
 
 					case LabelState.Unspecified:
 					{
-						UgsBadgeStates.Add(LabelIdx, UgsBadgeState.Skipped);
+						ugsBadgeStates.Add(labelIdx, UgsBadgeState.Skipped);
 						break;
 					}
 				}
 			}
-			return UgsBadgeStates;
+			return ugsBadgeStates;
 		}
 
 		/// <summary>
 		/// Gets the state of a job, as a label that includes all steps
 		/// </summary>
-		/// <param name="Job">The job to query</param>
-		/// <param name="StepForNode">Map from node to step</param>
-		/// <param name="NewState">Receives the state of the label</param>
-		/// <param name="NewOutcome">Receives the outcome of the label</param>
-		public static void GetJobState(this IJob Job, IReadOnlyDictionary<NodeRef, IJobStep> StepForNode, out LabelState NewState, out LabelOutcome NewOutcome)
+		/// <param name="job">The job to query</param>
+		/// <param name="stepForNode">Map from node to step</param>
+		/// <param name="newState">Receives the state of the label</param>
+		/// <param name="newOutcome">Receives the outcome of the label</param>
+		public static void GetJobState(this IJob job, IReadOnlyDictionary<NodeRef, IJobStep> stepForNode, out LabelState newState, out LabelOutcome newOutcome)
 		{
-			List<NodeRef> Nodes = new List<NodeRef>();
-			foreach (IJobStepBatch Batch in Job.Batches)
+			List<NodeRef> nodes = new List<NodeRef>();
+			foreach (IJobStepBatch batch in job.Batches)
 			{
-				foreach (IJobStep Step in Batch.Steps)
+				foreach (IJobStep step in batch.Steps)
 				{
-					Nodes.Add(new NodeRef(Batch.GroupIdx, Step.NodeIdx));
+					nodes.Add(new NodeRef(batch.GroupIdx, step.NodeIdx));
 				}
 			}
-			GetLabelState(Nodes, StepForNode, out NewState, out NewOutcome);
+			GetLabelState(nodes, stepForNode, out newState, out newOutcome);
 		}
 		
 		/// <summary>
 		/// Gets the state of a label
 		/// </summary>
-		/// <param name="IncludedNodes">Nodes to include in this label</param>
-		/// <param name="StepForNode">Map from node to step</param>
-		/// <param name="NewState">Receives the state of the label</param>
-		/// <param name="NewOutcome">Receives the outcome of the label</param>
-		public static void GetLabelState(IEnumerable<NodeRef> IncludedNodes, IReadOnlyDictionary<NodeRef, IJobStep> StepForNode, out LabelState NewState, out LabelOutcome NewOutcome)
+		/// <param name="includedNodes">Nodes to include in this label</param>
+		/// <param name="stepForNode">Map from node to step</param>
+		/// <param name="newState">Receives the state of the label</param>
+		/// <param name="newOutcome">Receives the outcome of the label</param>
+		public static void GetLabelState(IEnumerable<NodeRef> includedNodes, IReadOnlyDictionary<NodeRef, IJobStep> stepForNode, out LabelState newState, out LabelOutcome newOutcome)
 		{
-			NewState = LabelState.Complete;
-			NewOutcome = LabelOutcome.Success;
-			foreach (NodeRef IncludedNodeRef in IncludedNodes)
+			newState = LabelState.Complete;
+			newOutcome = LabelOutcome.Success;
+			foreach (NodeRef includedNodeRef in includedNodes)
 			{
-				IJobStep? IncludedStep;
-				if (StepForNode.TryGetValue(IncludedNodeRef, out IncludedStep))
+				IJobStep? includedStep;
+				if (stepForNode.TryGetValue(includedNodeRef, out includedStep))
 				{
 					// Update the state
-					if (IncludedStep.State != JobStepState.Completed && IncludedStep.State != JobStepState.Skipped && IncludedStep.State != JobStepState.Aborted)
+					if (includedStep.State != JobStepState.Completed && includedStep.State != JobStepState.Skipped && includedStep.State != JobStepState.Aborted)
 					{
-						NewState = LabelState.Running;
+						newState = LabelState.Running;
 					}
 
 					// Update the outcome
-					if (IncludedStep.State == JobStepState.Skipped || IncludedStep.State == JobStepState.Aborted || IncludedStep.Outcome == JobStepOutcome.Failure)
+					if (includedStep.State == JobStepState.Skipped || includedStep.State == JobStepState.Aborted || includedStep.Outcome == JobStepOutcome.Failure)
 					{
-						NewOutcome = LabelOutcome.Failure;
+						newOutcome = LabelOutcome.Failure;
 					}
-					else if (IncludedStep.Outcome == JobStepOutcome.Warnings && NewOutcome == LabelOutcome.Success)
+					else if (includedStep.Outcome == JobStepOutcome.Warnings && newOutcome == LabelOutcome.Success)
 					{
-						NewOutcome = LabelOutcome.Warnings;
+						newOutcome = LabelOutcome.Warnings;
 					}
 				}
 			}
@@ -1411,18 +1408,18 @@ namespace Horde.Build.Models
 		/// <summary>
 		/// Creates an RPC response object
 		/// </summary>
-		/// <param name="Job">The job document</param>
+		/// <param name="job">The job document</param>
 		/// <returns></returns>
-		public static HordeCommon.Rpc.GetJobResponse ToRpcResponse(this IJob Job)
+		public static HordeCommon.Rpc.GetJobResponse ToRpcResponse(this IJob job)
 		{
-			HordeCommon.Rpc.GetJobResponse Response = new HordeCommon.Rpc.GetJobResponse();
-			Response.StreamId = Job.StreamId.ToString();
-			Response.Change = Job.Change;
-			Response.CodeChange = Job.CodeChange;
-			Response.PreflightChange = Job.PreflightChange;
-			Response.ClonedPreflightChange = Job.ClonedPreflightChange;
-			Response.Arguments.Add(Job.Arguments);
-			return Response;
+			HordeCommon.Rpc.GetJobResponse response = new HordeCommon.Rpc.GetJobResponse();
+			response.StreamId = job.StreamId.ToString();
+			response.Change = job.Change;
+			response.CodeChange = job.CodeChange;
+			response.PreflightChange = job.PreflightChange;
+			response.ClonedPreflightChange = job.ClonedPreflightChange;
+			response.Arguments.Add(job.Arguments);
+			return response;
 		}
 	}
 
diff --git a/Engine/Source/Programs/Horde/Horde.Build/Jobs/IJobCollection.cs b/Engine/Source/Programs/Horde/Horde.Build/Jobs/IJobCollection.cs
--- a/Engine/Source/Programs/Horde/Horde.Build/Jobs/IJobCollection.cs
+++ b/Engine/Source/Programs/Horde/Horde.Build/Jobs/IJobCollection.cs
@@ -1,21 +1,14 @@
 // Copyright Epic Games, Inc. All Rights Reserved.
 
+using System;
+using System.Collections.Generic;
+using System.Threading.Tasks;
 using EpicGames.Core;
 using Horde.Build.Api;
-using HordeCommon;
 using Horde.Build.Models;
-using Horde.Build.Services;
 using Horde.Build.Utilities;
-using Microsoft.Extensions.Logging;
+using HordeCommon;
 using MongoDB.Bson;
-using MongoDB.Bson.Serialization.Attributes;
-using MongoDB.Driver;
-using System;
-using System.Collections.Generic;
-using System.Diagnostics.CodeAnalysis;
-using System.IO;
-using System.Linq;
-using System.Threading.Tasks;
 
 namespace Horde.Build.Collections
 {
@@ -31,173 +24,172 @@ namespace Horde.Build.Collections
 	/// <summary>
 	/// Interface for a collection of job documents
 	/// </summary>
-	[SuppressMessage("Compiler", "CA1054:URI properties should not be strings")]
 	public interface IJobCollection
 	{
 		/// <summary>
 		/// Creates a new job
 		/// </summary>
-		/// <param name="JobId">A requested job id</param>
-		/// <param name="StreamId">Unique id of the stream that this job belongs to</param>
-		/// <param name="TemplateRefId">Name of the template ref</param>
-		/// <param name="TemplateHash">Template for this job</param>
-		/// <param name="Graph">The graph for the new job</param>
-		/// <param name="Name">Name of the job</param>
-		/// <param name="Change">The change to build</param>
-		/// <param name="CodeChange">The corresponding code changelist number</param>
-		/// <param name="PreflightChange">Optional changelist to preflight</param>
-		/// <param name="ClonedPreflightChange">Optional cloned preflight changelist</param>
-		/// <param name="StartedByUserId">User that started the job</param>
-		/// <param name="Priority">Priority of the job</param>
-		/// <param name="AutoSubmit">Whether to automatically submit the preflighted change on completion</param>
-		/// <param name="UpdateIssues">Whether to update issues based on the outcome of this job</param>
-		/// <param name="PromoteIssuesByDefault">Whether to promote issues based on the outcome of this job by default</param>
-		/// <param name="JobTriggers">List of downstream job triggers</param>
-		/// <param name="ShowUgsBadges">Whether to show badges in UGS for this job</param>
-		/// <param name="ShowUgsAlerts">Whether to show alerts in UGS for this job</param>
-		/// <param name="NotificationChannel">Notification channel for this job</param>
-		/// <param name="NotificationChannelFilter">Notification channel filter for this job</param>
-		/// <param name="Arguments">Arguments for the job</param>
+		/// <param name="jobId">A requested job id</param>
+		/// <param name="streamId">Unique id of the stream that this job belongs to</param>
+		/// <param name="templateRefId">Name of the template ref</param>
+		/// <param name="templateHash">Template for this job</param>
+		/// <param name="graph">The graph for the new job</param>
+		/// <param name="name">Name of the job</param>
+		/// <param name="change">The change to build</param>
+		/// <param name="codeChange">The corresponding code changelist number</param>
+		/// <param name="preflightChange">Optional changelist to preflight</param>
+		/// <param name="clonedPreflightChange">Optional cloned preflight changelist</param>
+		/// <param name="startedByUserId">User that started the job</param>
+		/// <param name="priority">Priority of the job</param>
+		/// <param name="autoSubmit">Whether to automatically submit the preflighted change on completion</param>
+		/// <param name="updateIssues">Whether to update issues based on the outcome of this job</param>
+		/// <param name="promoteIssuesByDefault">Whether to promote issues based on the outcome of this job by default</param>
+		/// <param name="jobTriggers">List of downstream job triggers</param>
+		/// <param name="showUgsBadges">Whether to show badges in UGS for this job</param>
+		/// <param name="showUgsAlerts">Whether to show alerts in UGS for this job</param>
+		/// <param name="notificationChannel">Notification channel for this job</param>
+		/// <param name="notificationChannelFilter">Notification channel filter for this job</param>
+		/// <param name="arguments">Arguments for the job</param>
 		/// <returns>The new job document</returns>
-		Task<IJob> AddAsync(JobId JobId, StreamId StreamId, TemplateRefId TemplateRefId, ContentHash TemplateHash, IGraph Graph, string Name, int Change, int CodeChange, int? PreflightChange, int? ClonedPreflightChange, UserId? StartedByUserId, Priority? Priority, bool? AutoSubmit, bool? UpdateIssues, bool? PromoteIssuesByDefault, List<ChainedJobTemplate>? JobTriggers, bool ShowUgsBadges, bool ShowUgsAlerts, string? NotificationChannel, string? NotificationChannelFilter, List<string>? Arguments);
+		Task<IJob> AddAsync(JobId jobId, StreamId streamId, TemplateRefId templateRefId, ContentHash templateHash, IGraph graph, string name, int change, int codeChange, int? preflightChange, int? clonedPreflightChange, UserId? startedByUserId, Priority? priority, bool? autoSubmit, bool? updateIssues, bool? promoteIssuesByDefault, List<ChainedJobTemplate>? jobTriggers, bool showUgsBadges, bool showUgsAlerts, string? notificationChannel, string? notificationChannelFilter, List<string>? arguments);
 
 		/// <summary>
 		/// Gets a job with the given unique id
 		/// </summary>
-		/// <param name="JobId">Job id to search for</param>
+		/// <param name="jobId">Job id to search for</param>
 		/// <returns>Information about the given job</returns>
-		Task<IJob?> GetAsync(JobId JobId);
+		Task<IJob?> GetAsync(JobId jobId);
 
 		/// <summary>
 		/// Deletes a job
 		/// </summary>
-		/// <param name="Job">The job to remove</param>
-		Task<bool> RemoveAsync(IJob Job);
+		/// <param name="job">The job to remove</param>
+		Task<bool> RemoveAsync(IJob job);
 
 		/// <summary>
 		/// Delete all the jobs for a stream
 		/// </summary>
-		/// <param name="StreamId">Unique id of the stream</param>
+		/// <param name="streamId">Unique id of the stream</param>
 		/// <returns>Async task</returns>
-		Task RemoveStreamAsync(StreamId StreamId);
+		Task RemoveStreamAsync(StreamId streamId);
 
 		/// <summary>
 		/// Gets a job's permissions info by ID
 		/// </summary>
-		/// <param name="JobId">Unique id of the job</param>
+		/// <param name="jobId">Unique id of the job</param>
 		/// <returns>The job document</returns>
-		Task<IJobPermissions?> GetPermissionsAsync(JobId JobId);
+		Task<IJobPermissions?> GetPermissionsAsync(JobId jobId);
 
 		/// <summary>
 		/// Searches for jobs matching the given criteria
 		/// </summary>
-		/// <param name="JobIds">List of job ids to return</param>
-		/// <param name="StreamId">The stream containing the job</param>
-		/// <param name="Name">Name of the job</param>
-		/// <param name="Templates">Templates to look for</param>
-		/// <param name="MinChange">The minimum changelist number</param>
-		/// <param name="MaxChange">The maximum changelist number</param>
-		/// <param name="PreflightChange">Preflight change to find</param>
-		/// <param name="PreflightOnly">Whether to only include preflights</param>
-		/// <param name="StartedByUser">User id for which to include jobs</param>
-		/// <param name="PreflightStartedByUser">User for which to include preflight jobs</param>
-		/// <param name="MinCreateTime">The minimum creation time</param>
-		/// <param name="MaxCreateTime">The maximum creation time</param>
-		/// <param name="ModifiedBefore">Filter the results by modified time</param>
-		/// <param name="ModifiedAfter">Filter the results by modified time</param>
-		/// <param name="Index">Index of the first result to return</param>
-		/// <param name="Count">Number of results to return</param>
-		/// <param name="ConsistentRead">If the database read should be made to the replica server</param>
-		/// <param name="IndexHint">Name of index to be specified as a hint to the database query planner</param>
-		/// <param name="ExcludeUserJobs">Whether to exclude user jobs from the find</param>
+		/// <param name="jobIds">List of job ids to return</param>
+		/// <param name="streamId">The stream containing the job</param>
+		/// <param name="name">Name of the job</param>
+		/// <param name="templates">Templates to look for</param>
+		/// <param name="minChange">The minimum changelist number</param>
+		/// <param name="maxChange">The maximum changelist number</param>
+		/// <param name="preflightChange">Preflight change to find</param>
+		/// <param name="preflightOnly">Whether to only include preflights</param>
+		/// <param name="startedByUser">User id for which to include jobs</param>
+		/// <param name="preflightStartedByUser">User for which to include preflight jobs</param>
+		/// <param name="minCreateTime">The minimum creation time</param>
+		/// <param name="maxCreateTime">The maximum creation time</param>
+		/// <param name="modifiedBefore">Filter the results by modified time</param>
+		/// <param name="modifiedAfter">Filter the results by modified time</param>
+		/// <param name="index">Index of the first result to return</param>
+		/// <param name="count">Number of results to return</param>
+		/// <param name="consistentRead">If the database read should be made to the replica server</param>
+		/// <param name="indexHint">Name of index to be specified as a hint to the database query planner</param>
+		/// <param name="excludeUserJobs">Whether to exclude user jobs from the find</param>
 		/// <returns>List of jobs matching the given criteria</returns>
-		Task<List<IJob>> FindAsync(JobId[]? JobIds = null, StreamId? StreamId = null, string? Name = null, TemplateRefId[]? Templates = null, int? MinChange = null, int? MaxChange = null, int? PreflightChange = null, bool? PreflightOnly = null, UserId? PreflightStartedByUser = null, UserId? StartedByUser = null, DateTimeOffset? MinCreateTime = null, DateTimeOffset? MaxCreateTime = null, DateTimeOffset? ModifiedBefore = null, DateTimeOffset? ModifiedAfter = null, int? Index = null, int? Count = null, bool ConsistentRead = true, string? IndexHint = null, bool? ExcludeUserJobs = null);
+		Task<List<IJob>> FindAsync(JobId[]? jobIds = null, StreamId? streamId = null, string? name = null, TemplateRefId[]? templates = null, int? minChange = null, int? maxChange = null, int? preflightChange = null, bool? preflightOnly = null, UserId? preflightStartedByUser = null, UserId? startedByUser = null, DateTimeOffset? minCreateTime = null, DateTimeOffset? maxCreateTime = null, DateTimeOffset? modifiedBefore = null, DateTimeOffset? modifiedAfter = null, int? index = null, int? count = null, bool consistentRead = true, string? indexHint = null, bool? excludeUserJobs = null);
 
 		/// <summary>
 		/// Searches for jobs in a specified stream and templates
 		/// </summary>
-		/// <param name="StreamId">The stream containing the job</param>
-		/// <param name="Templates">Templates to look for</param>
-		/// <param name="PreflightStartedByUser">User for which to include preflight jobs</param>
-		/// <param name="MaxCreateTime">The maximum creation time</param>
-		/// <param name="ModifiedAfter">Filter the results by modified time</param>
-		/// <param name="Index">Index of the first result to return</param>
-		/// <param name="Count">Number of results to return</param>
-		/// <param name="ConsistentRead">If the database read should be made to the replica server</param>
-		Task<List<IJob>> FindLatestByStreamWithTemplatesAsync(StreamId StreamId, TemplateRefId[] Templates, UserId? PreflightStartedByUser = null, DateTimeOffset? MaxCreateTime = null, DateTimeOffset? ModifiedAfter = null, int? Index = null, int? Count = null, bool ConsistentRead = false);
+		/// <param name="streamId">The stream containing the job</param>
+		/// <param name="templates">Templates to look for</param>
+		/// <param name="preflightStartedByUser">User for which to include preflight jobs</param>
+		/// <param name="maxCreateTime">The maximum creation time</param>
+		/// <param name="modifiedAfter">Filter the results by modified time</param>
+		/// <param name="index">Index of the first result to return</param>
+		/// <param name="count">Number of results to return</param>
+		/// <param name="consistentRead">If the database read should be made to the replica server</param>
+		Task<List<IJob>> FindLatestByStreamWithTemplatesAsync(StreamId streamId, TemplateRefId[] templates, UserId? preflightStartedByUser = null, DateTimeOffset? maxCreateTime = null, DateTimeOffset? modifiedAfter = null, int? index = null, int? count = null, bool consistentRead = false);
 
 		/// <summary>
 		/// Updates a new job
 		/// </summary>
-		/// <param name="Job">The job document to update</param>
-		/// <param name="Graph">Graph for the job</param>
-		/// <param name="Name">Name of the job</param>
-		/// <param name="Priority">Priority of the job</param>
-		/// <param name="AutoSubmit">Automatically submit the job on completion</param>
-		/// <param name="AutoSubmitChange">Changelist that was automatically submitted</param>
-		/// <param name="AutoSubmitMessage"></param>
-		/// <param name="AbortedByUserId">Name of the user that aborted the job</param>
-		/// <param name="NotificationTriggerId">Id for a notification trigger</param>
-		/// <param name="Reports">New reports</param>
-		/// <param name="Arguments">New arguments for the job</param>
-		/// <param name="LabelIdxToTriggerId">New trigger ID for a label in the job</param>
-		/// <param name="JobTrigger">New downstream job id</param>
-		Task<IJob?> TryUpdateJobAsync(IJob Job, IGraph Graph, string? Name = null, Priority? Priority = null, bool? AutoSubmit = null, int? AutoSubmitChange = null, string? AutoSubmitMessage = null, UserId? AbortedByUserId = null, ObjectId? NotificationTriggerId = null, List<Report>? Reports = null, List<string>? Arguments = null, KeyValuePair<int, ObjectId>? LabelIdxToTriggerId = null, KeyValuePair<TemplateRefId, JobId>? JobTrigger = null);
+		/// <param name="job">The job document to update</param>
+		/// <param name="graph">Graph for the job</param>
+		/// <param name="name">Name of the job</param>
+		/// <param name="priority">Priority of the job</param>
+		/// <param name="autoSubmit">Automatically submit the job on completion</param>
+		/// <param name="autoSubmitChange">Changelist that was automatically submitted</param>
+		/// <param name="autoSubmitMessage"></param>
+		/// <param name="abortedByUserId">Name of the user that aborted the job</param>
+		/// <param name="notificationTriggerId">Id for a notification trigger</param>
+		/// <param name="reports">New reports</param>
+		/// <param name="arguments">New arguments for the job</param>
+		/// <param name="labelIdxToTriggerId">New trigger ID for a label in the job</param>
+		/// <param name="jobTrigger">New downstream job id</param>
+		Task<IJob?> TryUpdateJobAsync(IJob job, IGraph graph, string? name = null, Priority? priority = null, bool? autoSubmit = null, int? autoSubmitChange = null, string? autoSubmitMessage = null, UserId? abortedByUserId = null, ObjectId? notificationTriggerId = null, List<Report>? reports = null, List<string>? arguments = null, KeyValuePair<int, ObjectId>? labelIdxToTriggerId = null, KeyValuePair<TemplateRefId, JobId>? jobTrigger = null);
 
 		/// <summary>
 		/// Updates the state of a batch
 		/// </summary>
-		/// <param name="Job">Job to update</param>
-		/// <param name="Graph">The graph for this job</param>
-		/// <param name="BatchId">Unique id of the batch to update</param>
-		/// <param name="NewLogId">The new log file id</param>
-		/// <param name="NewState">New state of the jobstep</param>
-		/// <param name="NewError">Error code for the batch</param>
+		/// <param name="job">Job to update</param>
+		/// <param name="graph">The graph for this job</param>
+		/// <param name="batchId">Unique id of the batch to update</param>
+		/// <param name="newLogId">The new log file id</param>
+		/// <param name="newState">New state of the jobstep</param>
+		/// <param name="newError">Error code for the batch</param>
 		/// <returns>True if the job was updated, false if it was deleted</returns>
-		Task<IJob?> TryUpdateBatchAsync(IJob Job, IGraph Graph, SubResourceId BatchId, LogId? NewLogId, JobStepBatchState? NewState, JobStepBatchError? NewError);
+		Task<IJob?> TryUpdateBatchAsync(IJob job, IGraph graph, SubResourceId batchId, LogId? newLogId, JobStepBatchState? newState, JobStepBatchError? newError);
 
 		/// <summary>
 		/// Update a jobstep state
 		/// </summary>
-		/// <param name="Job">Job to update</param>
-		/// <param name="Graph">The graph for this job</param>
-		/// <param name="BatchId">Unique id of the batch containing the step</param>
-		/// <param name="StepId">Unique id of the step to update</param>
-		/// <param name="NewState">New state of the jobstep</param>
-		/// <param name="NewOutcome">New outcome of the jobstep</param>
-		/// <param name="NewAbortRequested">New state of request abort</param>
-		/// <param name="NewAbortByUserId">New name of user that requested the abort</param>
-		/// <param name="NewLogId">New log id for the jobstep</param>
-		/// <param name="NewNotificationTriggerId">New id for a notification trigger</param>
-		/// <param name="NewRetryByUserId">Whether the step should be retried</param>
-		/// <param name="NewPriority">New priority for this step</param>
-		/// <param name="NewReports">New report documents</param>
-		/// <param name="NewProperties">Property changes. Any properties with a null value will be removed.</param>
+		/// <param name="job">Job to update</param>
+		/// <param name="graph">The graph for this job</param>
+		/// <param name="batchId">Unique id of the batch containing the step</param>
+		/// <param name="stepId">Unique id of the step to update</param>
+		/// <param name="newState">New state of the jobstep</param>
+		/// <param name="newOutcome">New outcome of the jobstep</param>
+		/// <param name="newAbortRequested">New state of request abort</param>
+		/// <param name="newAbortByUserId">New name of user that requested the abort</param>
+		/// <param name="newLogId">New log id for the jobstep</param>
+		/// <param name="newNotificationTriggerId">New id for a notification trigger</param>
+		/// <param name="newRetryByUserId">Whether the step should be retried</param>
+		/// <param name="newPriority">New priority for this step</param>
+		/// <param name="newReports">New report documents</param>
+		/// <param name="newProperties">Property changes. Any properties with a null value will be removed.</param>
 		/// <returns>True if the job was updated, false if it was deleted in the meantime</returns>
-		Task<IJob?> TryUpdateStepAsync(IJob Job, IGraph Graph, SubResourceId BatchId, SubResourceId StepId, JobStepState NewState = default, JobStepOutcome NewOutcome = default, bool? NewAbortRequested = null, UserId? NewAbortByUserId = null, LogId? NewLogId = null, ObjectId? NewNotificationTriggerId = null, UserId? NewRetryByUserId = null, Priority? NewPriority = null, List<Report>? NewReports = null, Dictionary<string, string?>? NewProperties = null);
+		Task<IJob?> TryUpdateStepAsync(IJob job, IGraph graph, SubResourceId batchId, SubResourceId stepId, JobStepState newState = default, JobStepOutcome newOutcome = default, bool? newAbortRequested = null, UserId? newAbortByUserId = null, LogId? newLogId = null, ObjectId? newNotificationTriggerId = null, UserId? newRetryByUserId = null, Priority? newPriority = null, List<Report>? newReports = null, Dictionary<string, string?>? newProperties = null);
 
 		/// <summary>
 		/// Attempts to update the node groups to be executed for a job. Fails if another write happens in the meantime.
 		/// </summary>
-		/// <param name="Job">The job to update</param>
-		/// <param name="NewGraph">New graph for this job</param>
+		/// <param name="job">The job to update</param>
+		/// <param name="newGraph">New graph for this job</param>
 		/// <returns>True if the groups were updated to the given list. False if another write happened first.</returns>
-		Task<IJob?> TryUpdateGraphAsync(IJob Job, IGraph NewGraph);
+		Task<IJob?> TryUpdateGraphAsync(IJob job, IGraph newGraph);
 
 		/// <summary>
 		/// Removes a job from the dispatch queue. Ignores the state of any batches still remaining to execute. Should only be used to correct for inconsistent state.
 		/// </summary>
-		/// <param name="Job">Job</param>
+		/// <param name="job">Job</param>
 		/// <returns></returns>
-		Task<IJob?> TryRemoveFromDispatchQueueAsync(IJob Job);
+		Task<IJob?> TryRemoveFromDispatchQueueAsync(IJob job);
 
 		/// <summary>
 		/// Adds an issue to a job
 		/// </summary>
-		/// <param name="JobId">The job id</param>
-		/// <param name="IssueId">The issue to add</param>
+		/// <param name="jobId">The job id</param>
+		/// <param name="issueId">The issue to add</param>
 		/// <returns>Async task</returns>
-		Task AddIssueToJobAsync(JobId JobId, int IssueId);
+		Task AddIssueToJobAsync(JobId jobId, int issueId);
 
 		/// <summary>
 		/// Gets a queue of jobs to consider for execution
@@ -208,52 +200,52 @@ namespace Horde.Build.Collections
 		/// <summary>
 		/// Marks a job as skipped
 		/// </summary>
-		/// <param name="Job">The job to update</param>
-		/// <param name="Graph">Graph for the job</param>
-		/// <param name="Reason">Reason for this batch being failed</param>
+		/// <param name="job">The job to update</param>
+		/// <param name="graph">Graph for the job</param>
+		/// <param name="reason">Reason for this batch being failed</param>
 		/// <returns>Updated version of the job</returns>
-		Task<IJob?> SkipAllBatchesAsync(IJob? Job, IGraph Graph, JobStepBatchError Reason);
+		Task<IJob?> SkipAllBatchesAsync(IJob? job, IGraph graph, JobStepBatchError reason);
 
 		/// <summary>
 		/// Marks a batch as skipped
 		/// </summary>
-		/// <param name="Job">The job to update</param>
-		/// <param name="BatchId">The batch to mark as skipped</param>
-		/// <param name="Graph">Graph for the job</param>
-		/// <param name="Reason">Reason for this batch being failed</param>
+		/// <param name="job">The job to update</param>
+		/// <param name="batchId">The batch to mark as skipped</param>
+		/// <param name="graph">Graph for the job</param>
+		/// <param name="reason">Reason for this batch being failed</param>
 		/// <returns>Updated version of the job</returns>
-		Task<IJob?> SkipBatchAsync(IJob? Job, SubResourceId BatchId, IGraph Graph, JobStepBatchError Reason);
+		Task<IJob?> SkipBatchAsync(IJob? job, SubResourceId batchId, IGraph graph, JobStepBatchError reason);
 
 		/// <summary>
 		/// Abort an agent's lease, and update the payload accordingly
 		/// </summary>
-		/// <param name="Job">The job to update</param>
-		/// <param name="BatchIdx">Index of the batch to cancel</param>
-		/// <param name="Graph">Graph for the job</param>
-		/// <param name="Reason">Reason for this batch being failed</param>
+		/// <param name="job">The job to update</param>
+		/// <param name="batchIdx">Index of the batch to cancel</param>
+		/// <param name="graph">Graph for the job</param>
+		/// <param name="reason">Reason for this batch being failed</param>
 		/// <returns>True if the job is updated</returns>
-		Task<IJob?> TryFailBatchAsync(IJob Job, int BatchIdx, IGraph Graph, JobStepBatchError Reason);
+		Task<IJob?> TryFailBatchAsync(IJob job, int batchIdx, IGraph graph, JobStepBatchError reason);
 
 		/// <summary>
 		/// Attempt to assign a lease to execute a batch
 		/// </summary>
-		/// <param name="Job">The job containing the batch</param>
-		/// <param name="BatchIdx">Index of the batch</param>
-		/// <param name="PoolId">The pool id</param>
-		/// <param name="AgentId">New agent to execute the batch</param>
-		/// <param name="SessionId">Session of the agent that is to execute the batch</param>
-		/// <param name="LeaseId">The lease unique id</param>
-		/// <param name="LogId">Unique id of the log for the batch</param>
+		/// <param name="job">The job containing the batch</param>
+		/// <param name="batchIdx">Index of the batch</param>
+		/// <param name="poolId">The pool id</param>
+		/// <param name="agentId">New agent to execute the batch</param>
+		/// <param name="sessionId">Session of the agent that is to execute the batch</param>
+		/// <param name="leaseId">The lease unique id</param>
+		/// <param name="logId">Unique id of the log for the batch</param>
 		/// <returns>True if the batch is updated</returns>
-		Task<IJob?> TryAssignLeaseAsync(IJob Job, int BatchIdx, PoolId PoolId, AgentId AgentId, SessionId SessionId, LeaseId LeaseId, LogId LogId);
+		Task<IJob?> TryAssignLeaseAsync(IJob job, int batchIdx, PoolId poolId, AgentId agentId, SessionId sessionId, LeaseId leaseId, LogId logId);
 
 		/// <summary>
 		/// Cancel a lease reservation on a batch (before it has started)
 		/// </summary>
-		/// <param name="Job">The job containing the lease</param>
-		/// <param name="BatchIdx">Index of the batch to cancel</param>
+		/// <param name="job">The job containing the lease</param>
+		/// <param name="batchIdx">Index of the batch to cancel</param>
 		/// <returns>True if the job is updated</returns>
-		Task<IJob?> TryCancelLeaseAsync(IJob Job, int BatchIdx);
+		Task<IJob?> TryCancelLeaseAsync(IJob job, int batchIdx);
 
 		/// <summary>
 		/// Upgrade all documents in the collection
