Horde: Fix coding conventions within Horde.Build.

diff --git a/Engine/Source/Programs/Horde/Horde.Build/Jobs/IJobStepRef.cs b/Engine/Source/Programs/Horde/Horde.Build/Jobs/IJobStepRef.cs
--- a/Engine/Source/Programs/Horde/Horde.Build/Jobs/IJobStepRef.cs
+++ b/Engine/Source/Programs/Horde/Horde.Build/Jobs/IJobStepRef.cs
@@ -1,22 +1,18 @@
 // Copyright Epic Games, Inc. All Rights Reserved.
 
-using HordeCommon;
+using System;
 using Horde.Build.Utilities;
-using MongoDB.Bson;
+using HordeCommon;
 using MongoDB.Bson.Serialization;
 using MongoDB.Bson.Serialization.Attributes;
-using System;
-using System.Collections.Generic;
-using System.Linq;
-using System.Threading.Tasks;
 
 namespace Horde.Build.Models
 {
 	using JobId = ObjectId<IJob>;
 	using LogId = ObjectId<ILogFile>;
+	using PoolId = StringId<IPool>;
 	using StreamId = StringId<IStream>;
 	using TemplateRefId = StringId<TemplateRef>;
-	using PoolId = StringId<IPool>;
 
 	/// <summary>
 	/// Unique id struct for JobStepRef objects. Includes a job id, batch id, and step id to uniquely identify the step.
@@ -42,25 +38,25 @@ namespace Horde.Build.Models
 		/// <summary>
 		/// Constructor
 		/// </summary>
-		/// <param name="JobId">The job id</param>
-		/// <param name="BatchId">The batch id within the job</param>
-		/// <param name="StepId">The step id</param>
-		public JobStepRefId(JobId JobId, SubResourceId BatchId, SubResourceId StepId)
+		/// <param name="jobId">The job id</param>
+		/// <param name="batchId">The batch id within the job</param>
+		/// <param name="stepId">The step id</param>
+		public JobStepRefId(JobId jobId, SubResourceId batchId, SubResourceId stepId)
 		{
-			this.JobId = JobId;
-			this.BatchId = BatchId;
-			this.StepId = StepId;
+			JobId = jobId;
+			BatchId = batchId;
+			StepId = stepId;
 		}
 
 		/// <summary>
 		/// Parse a job step id from a string
 		/// </summary>
-		/// <param name="Text">Text to parse</param>
+		/// <param name="text">Text to parse</param>
 		/// <returns>The parsed id</returns>
-		public static JobStepRefId Parse(string Text)
+		public static JobStepRefId Parse(string text)
 		{
-			string[] Components = Text.Split(':');
-			return new JobStepRefId(JobId.Parse(Components[0]), SubResourceId.Parse(Components[1]), SubResourceId.Parse(Components[2]));
+			string[] components = text.Split(':');
+			return new JobStepRefId(JobId.Parse(components[0]), SubResourceId.Parse(components[1]), SubResourceId.Parse(components[2]));
 		}
 
 		/// <summary>
@@ -79,33 +75,30 @@ namespace Horde.Build.Models
 	public sealed class JobStepRefIdSerializer : IBsonSerializer<JobStepRefId>
 	{
 		/// <inheritdoc/>
-		public Type ValueType
-		{
-			get { return typeof(JobStepRefId); }
-		}
+		public Type ValueType => typeof(JobStepRefId);
 
 		/// <inheritdoc/>
-		void IBsonSerializer.Serialize(BsonSerializationContext Context, BsonSerializationArgs Args, object Value)
+		void IBsonSerializer.Serialize(BsonSerializationContext context, BsonSerializationArgs args, object value)
 		{
-			Serialize(Context, Args, (JobStepRefId)Value);
+			Serialize(context, args, (JobStepRefId)value);
 		}
 
 		/// <inheritdoc/>
-		object IBsonSerializer.Deserialize(BsonDeserializationContext Context, BsonDeserializationArgs Args)
+		object IBsonSerializer.Deserialize(BsonDeserializationContext context, BsonDeserializationArgs args)
 		{
-			return ((IBsonSerializer<JobStepRefId>)this).Deserialize(Context, Args);
+			return ((IBsonSerializer<JobStepRefId>)this).Deserialize(context, args);
 		}
 
 		/// <inheritdoc/>
-		public JobStepRefId Deserialize(BsonDeserializationContext Context, BsonDeserializationArgs Args)
+		public JobStepRefId Deserialize(BsonDeserializationContext context, BsonDeserializationArgs args)
 		{
-			return JobStepRefId.Parse(Context.Reader.ReadString());
+			return JobStepRefId.Parse(context.Reader.ReadString());
 		}
 
 		/// <inheritdoc/>
-		public void Serialize(BsonSerializationContext Context, BsonSerializationArgs Args, JobStepRefId Value)
+		public void Serialize(BsonSerializationContext context, BsonSerializationArgs args, JobStepRefId value)
 		{
-			Context.Writer.WriteString(((JobStepRefId)Value).ToString());
+			context.Writer.WriteString(((JobStepRefId)value).ToString());
 		}
 	}
 
diff --git a/Engine/Source/Programs/Horde/Horde.Build/Jobs/IJobStepRefCollection.cs b/Engine/Source/Programs/Horde/Horde.Build/Jobs/IJobStepRefCollection.cs
--- a/Engine/Source/Programs/Horde/Horde.Build/Jobs/IJobStepRefCollection.cs
+++ b/Engine/Source/Programs/Horde/Horde.Build/Jobs/IJobStepRefCollection.cs
@@ -1,16 +1,11 @@
 // Copyright Epic Games, Inc. All Rights Reserved.
 
-using HordeCommon;
-using Horde.Build.Models;
-using Horde.Build.Services;
-using MongoDB.Bson;
-using MongoDB.Bson.Serialization.Attributes;
-using MongoDB.Driver;
 using System;
 using System.Collections.Generic;
-using System.Linq;
 using System.Threading.Tasks;
+using Horde.Build.Models;
 using Horde.Build.Utilities;
+using HordeCommon;
 
 namespace Horde.Build.Collections
 {
@@ -27,90 +22,90 @@ namespace Horde.Build.Collections
 		/// <summary>
 		/// Constructor
 		/// </summary>
-		/// <param name="Id">Unique id for the step being referenced</param>
-		/// <param name="JobName">Name of the job</param>
-		/// <param name="StepName">Name of the step</param>
-		/// <param name="StreamId">Unique id for the stream containing the job</param>
-		/// <param name="TemplateId"></param>
-		/// <param name="Change">The change number being built</param>
-		/// <param name="LogId">The log file id</param>
-		/// <param name="PoolId">The pool id</param>
-		/// <param name="AgentId">The agent id</param>
-		/// <param name="Outcome">Outcome of this step, if known</param>
-		/// <param name="LastSuccess">The last change that completed with success</param>
-		/// <param name="LastWarning">The last change that completed with a warning (or success)</param>
-		/// <param name="WaitTime">Time taken for the batch containing this step to start</param>
-		/// <param name="InitTime">Time taken for the batch containing this step to initializer</param>
-		/// <param name="StartTimeUtc">Start time</param>
-		/// <param name="FinishTimeUtc">Finish time for the step, if known</param>
-		Task<IJobStepRef> InsertOrReplaceAsync(JobStepRefId Id, string JobName, string StepName, StreamId StreamId, TemplateRefId TemplateId, int Change, LogId? LogId, PoolId? PoolId, AgentId? AgentId, JobStepOutcome? Outcome, int? LastSuccess, int? LastWarning, float WaitTime, float InitTime, DateTime StartTimeUtc, DateTime? FinishTimeUtc);
+		/// <param name="id">Unique id for the step being referenced</param>
+		/// <param name="jobName">Name of the job</param>
+		/// <param name="stepName">Name of the step</param>
+		/// <param name="streamId">Unique id for the stream containing the job</param>
+		/// <param name="templateId"></param>
+		/// <param name="change">The change number being built</param>
+		/// <param name="logId">The log file id</param>
+		/// <param name="poolId">The pool id</param>
+		/// <param name="agentId">The agent id</param>
+		/// <param name="outcome">Outcome of this step, if known</param>
+		/// <param name="lastSuccess">The last change that completed with success</param>
+		/// <param name="lastWarning">The last change that completed with a warning (or success)</param>
+		/// <param name="waitTime">Time taken for the batch containing this step to start</param>
+		/// <param name="initTime">Time taken for the batch containing this step to initializer</param>
+		/// <param name="startTimeUtc">Start time</param>
+		/// <param name="finishTimeUtc">Finish time for the step, if known</param>
+		Task<IJobStepRef> InsertOrReplaceAsync(JobStepRefId id, string jobName, string stepName, StreamId streamId, TemplateRefId templateId, int change, LogId? logId, PoolId? poolId, AgentId? agentId, JobStepOutcome? outcome, int? lastSuccess, int? lastWarning, float waitTime, float initTime, DateTime startTimeUtc, DateTime? finishTimeUtc);
 
 		/// <summary>
 		/// Gets the history of a given node
 		/// </summary>
-		/// <param name="StreamId">Unique id for a stream</param>
-		/// <param name="TemplateId"></param>
-		/// <param name="NodeName">Name of the node</param>
-		/// <param name="Change">The current change</param>
-		/// <param name="IncludeFailed">Whether to include failed nodes</param>
-		/// <param name="Count">Number of results to return</param>
+		/// <param name="streamId">Unique id for a stream</param>
+		/// <param name="templateId"></param>
+		/// <param name="nodeName">Name of the node</param>
+		/// <param name="change">The current change</param>
+		/// <param name="includeFailed">Whether to include failed nodes</param>
+		/// <param name="count">Number of results to return</param>
 		/// <returns>List of step references</returns>
-		Task<List<IJobStepRef>> GetStepsForNodeAsync(StreamId StreamId, TemplateRefId TemplateId, string NodeName, int? Change, bool IncludeFailed, int Count);
+		Task<List<IJobStepRef>> GetStepsForNodeAsync(StreamId streamId, TemplateRefId templateId, string nodeName, int? change, bool includeFailed, int count);
 
 		/// <summary>
 		/// Gets the previous job that ran a given step
 		/// </summary>
-		/// <param name="StreamId">Id of the stream to search</param>
-		/// <param name="TemplateId">The template id</param>
-		/// <param name="NodeName">Name of the step to find</param>
-		/// <param name="Change">The current changelist number</param>
+		/// <param name="streamId">Id of the stream to search</param>
+		/// <param name="templateId">The template id</param>
+		/// <param name="nodeName">Name of the step to find</param>
+		/// <param name="change">The current changelist number</param>
 		/// <returns>The previous job, or null.</returns>
-		Task<IJobStepRef?> GetPrevStepForNodeAsync(StreamId StreamId, TemplateRefId TemplateId, string NodeName, int Change);
+		Task<IJobStepRef?> GetPrevStepForNodeAsync(StreamId streamId, TemplateRefId templateId, string nodeName, int change);
 
 		/// <summary>
 		/// Gets the next job that ran a given step
 		/// </summary>
-		/// <param name="StreamId">Id of the stream to search</param>
-		/// <param name="TemplateId">The template id</param>
-		/// <param name="NodeName">Name of the step to find</param>
-		/// <param name="Change">The current changelist number</param>
+		/// <param name="streamId">Id of the stream to search</param>
+		/// <param name="templateId">The template id</param>
+		/// <param name="nodeName">Name of the step to find</param>
+		/// <param name="change">The current changelist number</param>
 		/// <returns>The previous job, or null.</returns>
-		Task<IJobStepRef?> GetNextStepForNodeAsync(StreamId StreamId, TemplateRefId TemplateId, string NodeName, int Change);
+		Task<IJobStepRef?> GetNextStepForNodeAsync(StreamId streamId, TemplateRefId templateId, string nodeName, int change);
 	}
 
 	static class JobStepRefCollectionExtensions
 	{
-		public static async Task UpdateAsync(this IJobStepRefCollection JobStepRefs, IJob Job, IJobStepBatch Batch, IJobStep Step, IGraph Graph)
+		public static async Task UpdateAsync(this IJobStepRefCollection jobStepRefs, IJob job, IJobStepBatch batch, IJobStep step, IGraph graph)
 		{
-			if (Job.PreflightChange == 0)
+			if (job.PreflightChange == 0)
 			{
-				float WaitTime = (float)(Batch.GetWaitTime() ?? TimeSpan.Zero).TotalSeconds;
-				float InitTime = (float)(Batch.GetInitTime() ?? TimeSpan.Zero).TotalSeconds;
+				float waitTime = (float)(batch.GetWaitTime() ?? TimeSpan.Zero).TotalSeconds;
+				float initTime = (float)(batch.GetInitTime() ?? TimeSpan.Zero).TotalSeconds;
 
-				string NodeName = Graph.Groups[Batch.GroupIdx].Nodes[Step.NodeIdx].Name;
-				JobStepOutcome? Outcome = Step.IsPending() ? (JobStepOutcome?)null : Step.Outcome;
+				string nodeName = graph.Groups[batch.GroupIdx].Nodes[step.NodeIdx].Name;
+				JobStepOutcome? outcome = step.IsPending() ? (JobStepOutcome?)null : step.Outcome;
 
-				int? LastSuccess = null;
-				int? LastWarning = null;
-				if (Outcome != JobStepOutcome.Success)
+				int? lastSuccess = null;
+				int? lastWarning = null;
+				if (outcome != JobStepOutcome.Success)
 				{
-					IJobStepRef? PrevStep = await JobStepRefs.GetPrevStepForNodeAsync(Job.StreamId, Job.TemplateId, NodeName, Job.Change);
-					if (PrevStep != null)
+					IJobStepRef? prevStep = await jobStepRefs.GetPrevStepForNodeAsync(job.StreamId, job.TemplateId, nodeName, job.Change);
+					if (prevStep != null)
 					{
-						LastSuccess = PrevStep.LastSuccess;
-						if (Outcome != JobStepOutcome.Warnings)
+						lastSuccess = prevStep.LastSuccess;
+						if (outcome != JobStepOutcome.Warnings)
 						{
-							LastWarning = PrevStep.LastWarning;
+							lastWarning = prevStep.LastWarning;
 						}
 					}
 				}
 
-				if (Job.AbortedByUserId != null && Outcome == null)
+				if (job.AbortedByUserId != null && outcome == null)
 				{
-					Outcome = JobStepOutcome.Unspecified;
+					outcome = JobStepOutcome.Unspecified;
 				}
 
-				await JobStepRefs.InsertOrReplaceAsync(new JobStepRefId(Job.Id, Batch.Id, Step.Id), Job.Name, NodeName, Job.StreamId, Job.TemplateId, Job.Change, Step.LogId, Batch.PoolId, Batch.AgentId, Outcome, LastSuccess, LastWarning, WaitTime, InitTime, Step.StartTimeUtc ?? DateTime.UtcNow, Step.FinishTimeUtc);
+				await jobStepRefs.InsertOrReplaceAsync(new JobStepRefId(job.Id, batch.Id, step.Id), job.Name, nodeName, job.StreamId, job.TemplateId, job.Change, step.LogId, batch.PoolId, batch.AgentId, outcome, lastSuccess, lastWarning, waitTime, initTime, step.StartTimeUtc ?? DateTime.UtcNow, step.FinishTimeUtc);
 			}
 		}
 	}
diff --git a/Engine/Source/Programs/Horde/Horde.Build/Jobs/IJobTiming.cs b/Engine/Source/Programs/Horde/Horde.Build/Jobs/IJobTiming.cs
--- a/Engine/Source/Programs/Horde/Horde.Build/Jobs/IJobTiming.cs
+++ b/Engine/Source/Programs/Horde/Horde.Build/Jobs/IJobTiming.cs
@@ -1,10 +1,6 @@
 // Copyright Epic Games, Inc. All Rights Reserved.
 
-using System;
-using System.Collections.Generic;
 using System.Diagnostics.CodeAnalysis;
-using System.Linq;
-using System.Threading.Tasks;
 
 namespace Horde.Build.Models
 {
@@ -37,9 +33,9 @@ namespace Horde.Build.Models
 		/// <summary>
 		/// Gets timing information for a particular step
 		/// </summary>
-		/// <param name="Name">Name of the node being executed</param>
-		/// <param name="Timing">Receives the timing information for the given step</param>
+		/// <param name="name">Name of the node being executed</param>
+		/// <param name="timing">Receives the timing information for the given step</param>
 		/// <returns>True if the timing was found</returns>
-		public bool TryGetStepTiming(string Name, [NotNullWhen(true)] out IJobStepTiming? Timing);
+		public bool TryGetStepTiming(string name, [NotNullWhen(true)] out IJobStepTiming? timing);
 	}
 }
diff --git a/Engine/Source/Programs/Horde/Horde.Build/Jobs/IJobTimingCollection.cs b/Engine/Source/Programs/Horde/Horde.Build/Jobs/IJobTimingCollection.cs
--- a/Engine/Source/Programs/Horde/Horde.Build/Jobs/IJobTimingCollection.cs
+++ b/Engine/Source/Programs/Horde/Horde.Build/Jobs/IJobTimingCollection.cs
@@ -1,18 +1,10 @@
 // Copyright Epic Games, Inc. All Rights Reserved.
 
-using Amazon.Runtime.Internal.Util;
-using Amazon.S3.Model.Internal.MarshallTransformations;
+using System.Collections.Generic;
+using System.Threading.Tasks;
 using Horde.Build.Models;
-using Horde.Build.Services;
 using Horde.Build.Utilities;
-using MongoDB.Bson;
 using MongoDB.Bson.Serialization.Attributes;
-using MongoDB.Driver;
-using System;
-using System.Collections.Generic;
-using System.Diagnostics.CodeAnalysis;
-using System.Linq;
-using System.Threading.Tasks;
 
 namespace Horde.Build.Collections
 {
@@ -52,22 +44,22 @@ namespace Horde.Build.Collections
 		[BsonConstructor]
 		private JobStepTimingData()
 		{
-			this.Name = null!;
+			Name = null!;
 		}
 
 		/// <summary>
 		/// Constructor
 		/// </summary>
-		/// <param name="Name">Name of the node</param>
-		/// <param name="AverageWaitTime">Wait time before executing the group containing this node</param>
-		/// <param name="AverageInitTime">Time taken for the group containing this node to initialize</param>
-		/// <param name="AverageDuration">Time spent executing this node</param>
-		public JobStepTimingData(string Name, float? AverageWaitTime, float? AverageInitTime, float? AverageDuration)
+		/// <param name="name">Name of the node</param>
+		/// <param name="averageWaitTime">Wait time before executing the group containing this node</param>
+		/// <param name="averageInitTime">Time taken for the group containing this node to initialize</param>
+		/// <param name="averageDuration">Time spent executing this node</param>
+		public JobStepTimingData(string name, float? averageWaitTime, float? averageInitTime, float? averageDuration)
 		{
-			this.Name = Name;
-			this.AverageWaitTime = AverageWaitTime;
-			this.AverageInitTime = AverageInitTime;
-			this.AverageDuration = AverageDuration;
+			Name = name;
+			AverageWaitTime = averageWaitTime;
+			AverageInitTime = averageInitTime;
+			AverageDuration = averageDuration;
 		}
 	}
 
@@ -79,24 +71,24 @@ namespace Horde.Build.Collections
 		/// <summary>
 		/// Add timing information for the given job
 		/// </summary>
-		/// <param name="JobId">The job id</param>
-		/// <param name="Steps">List of timing info for each step</param>
+		/// <param name="jobId">The job id</param>
+		/// <param name="steps">List of timing info for each step</param>
 		/// <returns>New timing document</returns>
-		Task<IJobTiming?> TryAddAsync(JobId JobId, List<JobStepTimingData> Steps);
+		Task<IJobTiming?> TryAddAsync(JobId jobId, List<JobStepTimingData> steps);
 
 		/// <summary>
 		/// Attempts to get the timing information for a particular job
 		/// </summary>
-		/// <param name="JobId">The unique job id</param>
+		/// <param name="jobId">The unique job id</param>
 		/// <returns>Timing info for the requested jbo</returns>
-		Task<IJobTiming?> TryGetAsync(JobId JobId);
+		Task<IJobTiming?> TryGetAsync(JobId jobId);
 
 		/// <summary>
 		/// Adds timing information for the particular job
 		/// </summary>
-		/// <param name="JobTiming">The current timing info for the job</param>
-		/// <param name="Steps">List of steps to add</param>
+		/// <param name="jobTiming">The current timing info for the job</param>
+		/// <param name="steps">List of steps to add</param>
 		/// <returns>New timing document. Null if the document could not be updated.</returns>
-		Task<IJobTiming?> TryAddStepsAsync(IJobTiming JobTiming, List<JobStepTimingData> Steps);
+		Task<IJobTiming?> TryAddStepsAsync(IJobTiming jobTiming, List<JobStepTimingData> steps);
 	}
 }
diff --git a/Engine/Source/Programs/Horde/Horde.Build/Jobs/ITemplate.cs b/Engine/Source/Programs/Horde/Horde.Build/Jobs/ITemplate.cs
--- a/Engine/Source/Programs/Horde/Horde.Build/Jobs/ITemplate.cs
+++ b/Engine/Source/Programs/Horde/Horde.Build/Jobs/ITemplate.cs
@@ -1,17 +1,11 @@
 // Copyright Epic Games, Inc. All Rights Reserved.
 
+using System;
+using System.Collections.Generic;
 using EpicGames.Core;
 using Horde.Build.Api;
 using HordeCommon;
-using Horde.Build.Utilities;
-using MongoDB.Bson;
 using MongoDB.Bson.Serialization.Attributes;
-using System;
-using System.Collections.Generic;
-using System.Linq;
-using System.Runtime.CompilerServices;
-using System.Text.Json.Serialization;
-using System.Threading.Tasks;
 
 namespace Horde.Build.Models
 {
@@ -25,8 +19,8 @@ namespace Horde.Build.Models
 		/// <summary>
 		/// Gets the default arguments for this parameter and its children
 		/// </summary>
-		/// <param name="DefaultArguments">List of default arguments</param>
-		public abstract void GetDefaultArguments(List<string> DefaultArguments);
+		/// <param name="defaultArguments">List of default arguments</param>
+		public abstract void GetDefaultArguments(List<string> defaultArguments);
 
 		/// <summary>
 		/// Convert this parameter to data for serialization
@@ -60,32 +54,32 @@ namespace Horde.Build.Models
 		/// </summary>
 		private GroupParameter()
 		{
-			this.Label = null!;
-			this.Children = null!;
+			Label = null!;
+			Children = null!;
 		}
 
 		/// <summary>
 		/// Constructor
 		/// </summary>
-		/// <param name="Label">Name of the group</param>
-		/// <param name="Style">How to display this group</param>
-		/// <param name="Children">List of child parameters</param>
-		public GroupParameter(string Label, GroupParameterStyle Style, List<Parameter> Children)
+		/// <param name="label">Name of the group</param>
+		/// <param name="style">How to display this group</param>
+		/// <param name="children">List of child parameters</param>
+		public GroupParameter(string label, GroupParameterStyle style, List<Parameter> children)
 		{
-			this.Label = Label;
-			this.Style = Style;
-			this.Children = Children;
+			Label = label;
+			Style = style;
+			Children = children;
 		}
 
 		/// <summary>
 		/// Gets the default arguments for this parameter and its children
 		/// </summary>
-		/// <param name="DefaultArguments">List of default arguments</param>
-		public override void GetDefaultArguments(List<string> DefaultArguments)
+		/// <param name="defaultArguments">List of default arguments</param>
+		public override void GetDefaultArguments(List<string> defaultArguments)
 		{
-			foreach (Parameter Child in Children)
+			foreach (Parameter child in Children)
 			{
-				Child.GetDefaultArguments(DefaultArguments);
+				child.GetDefaultArguments(defaultArguments);
 			}
 		}
 
@@ -154,31 +148,31 @@ namespace Horde.Build.Models
 		/// <summary>
 		/// Constructor
 		/// </summary>
-		/// <param name="Label">Parameter to pass this value to the BuildGraph script with</param>
-		/// <param name="Argument">Label to show next to the parameter</param>
-		/// <param name="Default">Default value for this argument</param>
-		/// <param name="Hint">Hint text to display</param>
-		/// <param name="Validation">Regex used to validate entries</param>
-		/// <param name="Description">Message displayed for invalid values</param>
-		/// <param name="ToolTip">Tool tip text to display</param>
-		public TextParameter(string Label, string Argument, string Default, string? Hint, string? Validation, string? Description, string? ToolTip)
+		/// <param name="label">Parameter to pass this value to the BuildGraph script with</param>
+		/// <param name="argument">Label to show next to the parameter</param>
+		/// <param name="defaultValue">Default value for this argument</param>
+		/// <param name="hint">Hint text to display</param>
+		/// <param name="validation">Regex used to validate entries</param>
+		/// <param name="description">Message displayed for invalid values</param>
+		/// <param name="toolTip">Tool tip text to display</param>
+		public TextParameter(string label, string argument, string defaultValue, string? hint, string? validation, string? description, string? toolTip)
 		{
-			this.Label = Label;
-			this.Argument = Argument;
-			this.Default = Default;
-			this.Hint = Hint;
-			this.Validation = Validation;
-			this.Description = Description;
-			this.ToolTip = ToolTip;
+			Label = label;
+			Argument = argument;
+			Default = defaultValue;
+			Hint = hint;
+			Validation = validation;
+			Description = description;
+			ToolTip = toolTip;
 		}
 
 		/// <summary>
 		/// Gets the default arguments for this parameter and its children
 		/// </summary>
-		/// <param name="DefaultArguments">List of default arguments</param>
-		public override void GetDefaultArguments(List<string> DefaultArguments)
+		/// <param name="defaultArguments">List of default arguments</param>
+		public override void GetDefaultArguments(List<string> defaultArguments)
 		{
-			DefaultArguments.Add(Argument + Default);
+			defaultArguments.Add(Argument + Default);
 		}
 
 		/// <summary>
@@ -228,18 +222,18 @@ namespace Horde.Build.Models
 		/// <summary>
 		/// Constructor
 		/// </summary>
-		/// <param name="Group">The group to put this parameter in</param>
-		/// <param name="Text">Text to display for this option</param>
-		/// <param name="ArgumentIfEnabled">Argument to add if this option is enabled</param>
-		/// <param name="ArgumentIfDisabled">Argument to add if this option is disabled</param>
-		/// <param name="Default">Whether this item is selected by default</param>
-		public ListParameterItem(string? Group, string Text, string? ArgumentIfEnabled, string? ArgumentIfDisabled, bool Default)
+		/// <param name="group">The group to put this parameter in</param>
+		/// <param name="text">Text to display for this option</param>
+		/// <param name="argumentIfEnabled">Argument to add if this option is enabled</param>
+		/// <param name="argumentIfDisabled">Argument to add if this option is disabled</param>
+		/// <param name="defaultValue">Whether this item is selected by default</param>
+		public ListParameterItem(string? group, string text, string? argumentIfEnabled, string? argumentIfDisabled, bool defaultValue)
 		{
-			this.Group = Group;
-			this.Text = Text;
-			this.ArgumentIfEnabled = ArgumentIfEnabled;
-			this.ArgumentIfDisabled = ArgumentIfDisabled;
-			this.Default = Default;
+			Group = group;
+			Text = text;
+			ArgumentIfEnabled = argumentIfEnabled;
+			ArgumentIfDisabled = argumentIfDisabled;
+			Default = defaultValue;
 		}
 
 		/// <summary>
@@ -282,45 +276,45 @@ namespace Horde.Build.Models
 		/// </summary>
 		private ListParameter()
 		{
-			this.Label = null!;
-			this.Items = null!;
+			Label = null!;
+			Items = null!;
 		}
 
 		/// <summary>
 		/// List of possible values
 		/// </summary>
-		/// <param name="Label">Label to show next to this parameter</param>
-		/// <param name="Style">Style of list parameter to use</param>
-		/// <param name="Entries">Entries for this list</param>
-		/// <param name="ToolTip">Tool tip text to display</param>
-		public ListParameter(string Label, ListParameterStyle Style, List<ListParameterItem> Entries, string? ToolTip)
+		/// <param name="label">Label to show next to this parameter</param>
+		/// <param name="style">Style of list parameter to use</param>
+		/// <param name="entries">Entries for this list</param>
+		/// <param name="toolTip">Tool tip text to display</param>
+		public ListParameter(string label, ListParameterStyle style, List<ListParameterItem> entries, string? toolTip)
 		{
-			this.Label = Label;
-			this.Style = Style;
-			this.Items = Entries;
-			this.ToolTip = ToolTip;
+			Label = label;
+			Style = style;
+			Items = entries;
+			ToolTip = toolTip;
 		}
 
 		/// <summary>
 		/// Gets the default arguments for this parameter and its children
 		/// </summary>
-		/// <param name="DefaultArguments">List of default arguments</param>
-		public override void GetDefaultArguments(List<string> DefaultArguments)
+		/// <param name="defaultArguments">List of default arguments</param>
+		public override void GetDefaultArguments(List<string> defaultArguments)
 		{
-			foreach(ListParameterItem Item in Items)
+			foreach(ListParameterItem item in Items)
 			{
-				if (Item.Default)
+				if (item.Default)
 				{
-					if (Item.ArgumentIfEnabled != null)
+					if (item.ArgumentIfEnabled != null)
 					{
-						DefaultArguments.Add(Item.ArgumentIfEnabled);
+						defaultArguments.Add(item.ArgumentIfEnabled);
 					}
 				}
 				else
 				{
-					if (Item.ArgumentIfDisabled != null)
+					if (item.ArgumentIfDisabled != null)
 					{
-						DefaultArguments.Add(Item.ArgumentIfDisabled);
+						defaultArguments.Add(item.ArgumentIfDisabled);
 					}
 				}
 			}
@@ -380,30 +374,30 @@ namespace Horde.Build.Models
 		/// <summary>
 		/// Constructor
 		/// </summary>
-		/// <param name="Label">Label to display for this parameter</param>
-		/// <param name="ArgumentIfEnabled">Value if enabled</param>
-		/// <param name="ArgumentIfDisabled">Value if disabled</param>
-		/// <param name="Default">Default value for this argument</param>
-		/// <param name="ToolTip">Tool tip text to display</param>
-		public BoolParameter(string Label, string? ArgumentIfEnabled, string? ArgumentIfDisabled, bool Default, string? ToolTip)
+		/// <param name="label">Label to display for this parameter</param>
+		/// <param name="argumentIfEnabled">Value if enabled</param>
+		/// <param name="argumentIfDisabled">Value if disabled</param>
+		/// <param name="defaultValue">Default value for this argument</param>
+		/// <param name="toolTip">Tool tip text to display</param>
+		public BoolParameter(string label, string? argumentIfEnabled, string? argumentIfDisabled, bool defaultValue, string? toolTip)
 		{
-			this.Label = Label;
-			this.ArgumentIfEnabled = ArgumentIfEnabled;
-			this.ArgumentIfDisabled = ArgumentIfDisabled;
-			this.Default = Default;
-			this.ToolTip = ToolTip;
+			Label = label;
+			ArgumentIfEnabled = argumentIfEnabled;
+			ArgumentIfDisabled = argumentIfDisabled;
+			Default = defaultValue;
+			ToolTip = toolTip;
 		}
 
 		/// <summary>
 		/// Gets the default arguments for this parameter and its children
 		/// </summary>
-		/// <param name="DefaultArguments">List of default arguments</param>
-		public override void GetDefaultArguments(List<string> DefaultArguments)
+		/// <param name="defaultArguments">List of default arguments</param>
+		public override void GetDefaultArguments(List<string> defaultArguments)
 		{
-			string? DefaultArgument = Default ? ArgumentIfEnabled : ArgumentIfDisabled;
-			if (!String.IsNullOrEmpty(DefaultArgument))
+			string? defaultArgument = Default ? ArgumentIfEnabled : ArgumentIfDisabled;
+			if (!String.IsNullOrEmpty(defaultArgument))
 			{
-				DefaultArguments.Add(DefaultArgument);
+				defaultArguments.Add(defaultArgument);
 			}
 		}
 		/// <summary>
@@ -486,14 +480,14 @@ namespace Horde.Build.Models
 		/// Gets the arguments for default options in this template. Does not include the standard template arguments.
 		/// </summary>
 		/// <returns>List of default arguments</returns>
-		public static List<string> GetDefaultArguments(this ITemplate Template)
+		public static List<string> GetDefaultArguments(this ITemplate template)
 		{
-			List<string> DefaultArguments = new List<string>(Template.Arguments);
-			foreach (Parameter Parameter in Template.Parameters)
+			List<string> defaultArguments = new List<string>(template.Arguments);
+			foreach (Parameter parameter in template.Parameters)
 			{
-				Parameter.GetDefaultArguments(DefaultArguments);
+				parameter.GetDefaultArguments(defaultArguments);
 			}
-			return DefaultArguments;
+			return defaultArguments;
 		}
 	}
 }
diff --git a/Engine/Source/Programs/Horde/Horde.Build/Jobs/ITemplateCollection.cs b/Engine/Source/Programs/Horde/Horde.Build/Jobs/ITemplateCollection.cs
--- a/Engine/Source/Programs/Horde/Horde.Build/Jobs/ITemplateCollection.cs
+++ b/Engine/Source/Programs/Horde/Horde.Build/Jobs/ITemplateCollection.cs
@@ -1,19 +1,10 @@
 // Copyright Epic Games, Inc. All Rights Reserved.
 
-using EpicGames.Core;
-using HordeCommon;
-using Horde.Build.Api;
-using Horde.Build.Models;
-using Horde.Build.Services;
-using Horde.Build.Utilities;
-using MongoDB.Bson;
-using MongoDB.Bson.Serialization;
-using MongoDB.Bson.Serialization.Attributes;
-using MongoDB.Driver;
-using System;
 using System.Collections.Generic;
-using System.Linq;
 using System.Threading.Tasks;
+using EpicGames.Core;
+using Horde.Build.Models;
+using HordeCommon;
 
 namespace Horde.Build.Collections
 {
@@ -25,17 +16,17 @@ namespace Horde.Build.Collections
 		/// <summary>
 		/// Public constructor
 		/// </summary>
-		/// <param name="Name">Name of the template</param>
-		/// <param name="Priority">Priority of this template</param>
+		/// <param name="name">Name of the template</param>
+		/// <param name="priority">Priority of this template</param>
 		/// <param name="bAllowPreflights">Whether to allow preflights of this job</param>
 		/// <param name="bUpdateIssues"> Whether to update issues for all jobs using this template</param>
 		/// <param name="bPromoteIssuesByDefault">Whether to promote issues by default for all jobs using this template</param>
-		/// <param name="InitialAgentType">The agent type to parse the buildgraph script</param>
-		/// <param name="SubmitNewChange">Path to a file within the stream to submit to generate a new changelist for jobs</param>
-		/// <param name="SubmitDescription">Description for new changes submitted to the stream</param>
-		/// <param name="Arguments">List of arguments which are always specified</param>
-		/// <param name="Parameters">List of template parameters</param>
-		Task<ITemplate> AddAsync(string Name, Priority? Priority = null, bool bAllowPreflights = true, bool bUpdateIssues = false, bool bPromoteIssuesByDefault = false, string? InitialAgentType = null, string? SubmitNewChange = null, string? SubmitDescription = null, List<string>? Arguments = null, List<Parameter>? Parameters = null);
+		/// <param name="initialAgentType">The agent type to parse the buildgraph script</param>
+		/// <param name="submitNewChange">Path to a file within the stream to submit to generate a new changelist for jobs</param>
+		/// <param name="submitDescription">Description for new changes submitted to the stream</param>
+		/// <param name="arguments">List of arguments which are always specified</param>
+		/// <param name="parameters">List of template parameters</param>
+		Task<ITemplate> AddAsync(string name, Priority? priority = null, bool bAllowPreflights = true, bool bUpdateIssues = false, bool bPromoteIssuesByDefault = false, string? initialAgentType = null, string? submitNewChange = null, string? submitDescription = null, List<string>? arguments = null, List<Parameter>? parameters = null);
 
 		/// <summary>
 		/// Gets all the available templates
@@ -46,8 +37,8 @@ namespace Horde.Build.Collections
 		/// <summary>
 		/// Gets a template by ID
 		/// </summary>
-		/// <param name="TemplateId">Unique id of the template</param>
+		/// <param name="templateId">Unique id of the template</param>
 		/// <returns>The template document</returns>
-		Task<ITemplate?> GetAsync(ContentHash TemplateId);
+		Task<ITemplate?> GetAsync(ContentHash templateId);
 	}
 }
diff --git a/Engine/Source/Programs/Horde/Horde.Build/Jobs/ITestData.cs b/Engine/Source/Programs/Horde/Horde.Build/Jobs/ITestData.cs
--- a/Engine/Source/Programs/Horde/Horde.Build/Jobs/ITestData.cs
+++ b/Engine/Source/Programs/Horde/Horde.Build/Jobs/ITestData.cs
@@ -2,10 +2,6 @@
 
 using Horde.Build.Utilities;
 using MongoDB.Bson;
-using System;
-using System.Collections.Generic;
-using System.Linq;
-using System.Threading.Tasks;
 
 namespace Horde.Build.Models
 {
diff --git a/Engine/Source/Programs/Horde/Horde.Build/Jobs/ITestDataCollection.cs b/Engine/Source/Programs/Horde/Horde.Build/Jobs/ITestDataCollection.cs
--- a/Engine/Source/Programs/Horde/Horde.Build/Jobs/ITestDataCollection.cs
+++ b/Engine/Source/Programs/Horde/Horde.Build/Jobs/ITestDataCollection.cs
@@ -1,23 +1,15 @@
 // Copyright Epic Games, Inc. All Rights Reserved.
 
+using System.Collections.Generic;
+using System.Threading.Tasks;
 using Horde.Build.Models;
-using Horde.Build.Services;
 using Horde.Build.Utilities;
 using MongoDB.Bson;
-using MongoDB.Bson.Serialization.Attributes;
-using MongoDB.Driver;
-using System;
-using System.Collections.Generic;
-using System.Diagnostics.CodeAnalysis;
-using System.Linq;
-using System.Threading.Tasks;
 
 namespace Horde.Build.Collections
 {
 	using JobId = ObjectId<IJob>;
-	using ProjectId = StringId<IProject>;
 	using StreamId = StringId<IStream>;
-	using TemplateRefId = StringId<TemplateRef>;
 
 	/// <summary>
 	/// Collection of test data documents
@@ -27,39 +19,39 @@ namespace Horde.Build.Collections
 		/// <summary>
 		/// Creates a new test data document
 		/// </summary>
-		/// <param name="Job">The job containing the step</param>
-		/// <param name="Step">The step producing the data</param>
-		/// <param name="Key">Key identifying the test</param>
-		/// <param name="Value">The data to store</param>
+		/// <param name="job">The job containing the step</param>
+		/// <param name="step">The step producing the data</param>
+		/// <param name="key">Key identifying the test</param>
+		/// <param name="value">The data to store</param>
 		/// <returns>The new stream document</returns>
-		Task<ITestData> AddAsync(IJob Job, IJobStep Step, string Key, BsonDocument Value);
+		Task<ITestData> AddAsync(IJob job, IJobStep step, string key, BsonDocument value);
 
 		/// <summary>
 		/// Gets a stream by ID
 		/// </summary>
-		/// <param name="Id">Unique id of the stream</param>
+		/// <param name="id">Unique id of the stream</param>
 		/// <returns>The stream document</returns>
-		Task<ITestData?> GetAsync(ObjectId Id);
+		Task<ITestData?> GetAsync(ObjectId id);
 
 		/// <summary>
 		/// Searches for test data that matches a set of criteria
 		/// </summary>
-		/// <param name="StreamId">The stream id</param>
-		/// <param name="MinChange">The minimum changelist number to return (inclusive)</param>
-		/// <param name="MaxChange">The maximum changelist number to return (inclusive)</param>
-		/// <param name="JobId">The job id</param>
-		/// <param name="StepId">The unique step id</param>
-		/// <param name="Key">Key identifying the result to return</param>
-		/// <param name="Index">Offset within the results to return</param>
-		/// <param name="Count">Number of results to return</param>
+		/// <param name="streamId">The stream id</param>
+		/// <param name="minChange">The minimum changelist number to return (inclusive)</param>
+		/// <param name="maxChange">The maximum changelist number to return (inclusive)</param>
+		/// <param name="jobId">The job id</param>
+		/// <param name="stepId">The unique step id</param>
+		/// <param name="key">Key identifying the result to return</param>
+		/// <param name="index">Offset within the results to return</param>
+		/// <param name="count">Number of results to return</param>
 		/// <returns>The stream document</returns>
-		Task<List<ITestData>> FindAsync(StreamId? StreamId, int? MinChange, int? MaxChange, JobId? JobId, SubResourceId? StepId, string? Key = null, int Index = 0, int Count = 10);
+		Task<List<ITestData>> FindAsync(StreamId? streamId, int? minChange, int? maxChange, JobId? jobId, SubResourceId? stepId, string? key = null, int index = 0, int count = 10);
 
 		/// <summary>
 		/// Delete the test data
 		/// </summary>
-		/// <param name="Id">Unique id of the test data</param>
+		/// <param name="id">Unique id of the test data</param>
 		/// <returns>Async task</returns>
-		Task DeleteAsync(ObjectId Id);
+		Task DeleteAsync(ObjectId id);
 	}
 }
diff --git a/Engine/Source/Programs/Horde/Horde.Build/Jobs/Impl/ArtifactCollection.cs b/Engine/Source/Programs/Horde/Horde.Build/Jobs/Impl/ArtifactCollection.cs
--- a/Engine/Source/Programs/Horde/Horde.Build/Jobs/Impl/ArtifactCollection.cs
+++ b/Engine/Source/Programs/Horde/Horde.Build/Jobs/Impl/ArtifactCollection.cs
@@ -1,29 +1,16 @@
 // Copyright Epic Games, Inc. All Rights Reserved.
 
-using Amazon.Extensions.NETCore.Setup;
-using Amazon.Runtime;
-using Amazon.S3;
-using Amazon.S3.Model;
-using Horde.Build.Api;
+using System;
+using System.Collections.Generic;
+using System.Linq;
+using System.Threading.Tasks;
 using Horde.Build.Jobs;
 using Horde.Build.Models;
 using Horde.Build.Storage;
 using Horde.Build.Utilities;
-using Microsoft.Extensions.Configuration;
-using Microsoft.Extensions.Hosting;
-using Microsoft.Extensions.Logging;
-using Microsoft.Extensions.Options;
 using MongoDB.Bson;
 using MongoDB.Bson.Serialization.Attributes;
 using MongoDB.Driver;
-using System;
-using System.Collections.Generic;
-using System.IO;
-using System.Linq;
-using System.Security.Claims;
-using System.Text;
-using System.Threading;
-using System.Threading.Tasks;
 
 namespace Horde.Build.Services
 {
@@ -58,157 +45,154 @@ namespace Horde.Build.Services
 			[BsonConstructor]
 			private Artifact()
 			{
-				this.Name = null!;
-				this.MimeType = null!;
+				Name = null!;
+				MimeType = null!;
 			}
 
-			public Artifact(JobId JobId, SubResourceId? StepId, string Name, long Length, string MimeType)
+			public Artifact(JobId jobId, SubResourceId? stepId, string name, long length, string mimeType)
 			{
-				this.Id = ObjectId.GenerateNewId();
-				this.JobId = JobId;
-				this.StepId = StepId;
-				this.Name = Name;
-				this.Length = Length;
-				this.MimeType = MimeType;
+				Id = ObjectId.GenerateNewId();
+				JobId = jobId;
+				StepId = stepId;
+				Name = name;
+				Length = length;
+				MimeType = mimeType;
 			}
 		}
 
-		private readonly IStorageBackend StorageBackend;
-		private readonly ILogger<ArtifactCollection> Logger;
-		private readonly IMongoCollection<Artifact> Artifacts;
+		private readonly IStorageBackend _storageBackend;
+		private readonly IMongoCollection<Artifact> _artifacts;
 
 		/// <summary>
 		/// Constructor
 		/// </summary>
-		/// <param name="DatabaseService">The database service</param>
-		/// <param name="StorageBackend">The storage backend</param>
-		/// <param name="Logger">Log interface</param>
-		public ArtifactCollection(DatabaseService DatabaseService, IStorageBackend<ArtifactCollection> StorageBackend, ILogger<ArtifactCollection> Logger)
+		/// <param name="databaseService">The database service</param>
+		/// <param name="storageBackend">The storage backend</param>
+		public ArtifactCollection(DatabaseService databaseService, IStorageBackend<ArtifactCollection> storageBackend)
 		{
-			this.StorageBackend = StorageBackend;
-			this.Logger = Logger;
+			_storageBackend = storageBackend;
 
 			// Initialize Artifacts table
-			Artifacts = DatabaseService.GetCollection<Artifact>("Artifacts");
-			if (!DatabaseService.ReadOnlyMode)
+			_artifacts = databaseService.GetCollection<Artifact>("Artifacts");
+			if (!databaseService.ReadOnlyMode)
 			{
-				Artifacts.Indexes.CreateOne(new CreateIndexModel<Artifact>(Builders<Artifact>.IndexKeys.Ascending(x => x.JobId)));
+				_artifacts.Indexes.CreateOne(new CreateIndexModel<Artifact>(Builders<Artifact>.IndexKeys.Ascending(x => x.JobId)));
 			}
 		}
 
 		/// <summary>
 		/// Creates a new artifact
 		/// </summary>
-		/// <param name="JobId">Unique id of the job that owns this artifact</param>
-		/// <param name="StepId">Optional Step id</param>
-		/// <param name="Name">Name of artifact</param>
-		/// <param name="MimeType">Type of artifact</param>
-		/// <param name="Data">The data to write</param>
+		/// <param name="jobId">Unique id of the job that owns this artifact</param>
+		/// <param name="stepId">Optional Step id</param>
+		/// <param name="name">Name of artifact</param>
+		/// <param name="mimeType">Type of artifact</param>
+		/// <param name="data">The data to write</param>
 		/// <returns>The new log file document</returns>
-		public async Task<IArtifact> CreateArtifactAsync(JobId JobId, SubResourceId? StepId, string Name, string MimeType, System.IO.Stream Data)
+		public async Task<IArtifact> CreateArtifactAsync(JobId jobId, SubResourceId? stepId, string name, string mimeType, System.IO.Stream data)
 		{
 			// upload first
-			string ArtifactName = ValidateName(Name);
-			await StorageBackend.WriteAsync(GetPath(JobId, StepId, ArtifactName), Data);
+			string artifactName = ValidateName(name);
+			await _storageBackend.WriteAsync(GetPath(jobId, stepId, artifactName), data);
 
 			// then create entry
-			Artifact NewArtifact = new Artifact(JobId, StepId, ArtifactName, Data.Length, MimeType);
-			await Artifacts.InsertOneAsync(NewArtifact);
-			return NewArtifact;
+			Artifact newArtifact = new Artifact(jobId, stepId, artifactName, data.Length, mimeType);
+			await _artifacts.InsertOneAsync(newArtifact);
+			return newArtifact;
 		}
 
 		/// <summary>
 		/// Gets all the available artifacts for a job
 		/// </summary>
-		/// <param name="JobId">Unique id of the job to query</param>
-		/// <param name="StepId">Unique id of the Step to query</param>
-		/// <param name="Name">Name of the artifact</param>
+		/// <param name="jobId">Unique id of the job to query</param>
+		/// <param name="stepId">Unique id of the Step to query</param>
+		/// <param name="name">Name of the artifact</param>
 		/// <returns>List of artifact documents</returns>
-		public async Task<List<IArtifact>> GetArtifactsAsync(JobId? JobId, SubResourceId? StepId, string? Name)
+		public async Task<List<IArtifact>> GetArtifactsAsync(JobId? jobId, SubResourceId? stepId, string? name)
 		{
-			FilterDefinitionBuilder<Artifact> Builder = Builders<Artifact>.Filter;
+			FilterDefinitionBuilder<Artifact> builder = Builders<Artifact>.Filter;
 
-			FilterDefinition<Artifact> Filter = FilterDefinition<Artifact>.Empty;
-			if (JobId != null)
+			FilterDefinition<Artifact> filter = FilterDefinition<Artifact>.Empty;
+			if (jobId != null)
 			{
-				Filter &= Builder.Eq(x => x.JobId, JobId.Value);
+				filter &= builder.Eq(x => x.JobId, jobId.Value);
 			}
-			if (StepId != null)
+			if (stepId != null)
 			{
-				Filter &= Builder.Eq(x => x.StepId, StepId.Value);
+				filter &= builder.Eq(x => x.StepId, stepId.Value);
 			}
-			if (Name != null)
+			if (name != null)
 			{
-				Filter &= Builder.Eq(x => x.Name, Name);
+				filter &= builder.Eq(x => x.Name, name);
 			}
 
-			return await Artifacts.Find(Filter).ToListAsync<Artifact, IArtifact>();
+			return await _artifacts.Find(filter).ToListAsync<Artifact, IArtifact>();
 		}
 
 		/// <summary>
 		/// Gets a specific list of artifacts based on id
 		/// </summary>
-		/// <param name="ArtifactIds">The list of artifact Ids</param>
+		/// <param name="artifactIds">The list of artifact Ids</param>
 		/// <returns>List of artifact documents</returns>
-		public async Task<List<IArtifact>> GetArtifactsAsync(IEnumerable<ObjectId> ArtifactIds)
+		public async Task<List<IArtifact>> GetArtifactsAsync(IEnumerable<ObjectId> artifactIds)
 		{
-			FilterDefinitionBuilder<Artifact> Builder = Builders<Artifact>.Filter;
+			FilterDefinitionBuilder<Artifact> builder = Builders<Artifact>.Filter;
 
-			FilterDefinition<Artifact> Filter = FilterDefinition<Artifact>.Empty;
-			Filter &= Builder.In(x => x.Id, ArtifactIds);
+			FilterDefinition<Artifact> filter = FilterDefinition<Artifact>.Empty;
+			filter &= builder.In(x => x.Id, artifactIds);
 			
-			return await Artifacts.Find(Filter).ToListAsync<Artifact, IArtifact>();
+			return await _artifacts.Find(filter).ToListAsync<Artifact, IArtifact>();
 		}
 
 		/// <summary>
 		/// Gets an artifact by ID
 		/// </summary>
-		/// <param name="ArtifactId">Unique id of the artifact</param>
+		/// <param name="artifactId">Unique id of the artifact</param>
 		/// <returns>The artifact document</returns>
-		public async Task<IArtifact?> GetArtifactAsync(ObjectId ArtifactId)
+		public async Task<IArtifact?> GetArtifactAsync(ObjectId artifactId)
 		{
-			return await Artifacts.Find<Artifact>(x => x.Id == ArtifactId).FirstOrDefaultAsync();
+			return await _artifacts.Find<Artifact>(x => x.Id == artifactId).FirstOrDefaultAsync();
 		}
 
 		/// <summary>
 		/// Attempts to update a single artifact document, if the update counter is valid
 		/// </summary>
-		/// <param name="Current">The current artifact document</param>
-		/// <param name="Update">The update definition</param>
+		/// <param name="current">The current artifact document</param>
+		/// <param name="update">The update definition</param>
 		/// <returns>True if the document was updated, false if another writer updated the document first</returns>
-		private async Task<bool> TryUpdateArtifactAsync(Artifact Current, UpdateDefinition<Artifact> Update)
+		private async Task<bool> TryUpdateArtifactAsync(Artifact current, UpdateDefinition<Artifact> update)
 		{
-			UpdateResult Result = await Artifacts.UpdateOneAsync<Artifact>(x => x.Id == Current.Id && x.UpdateIndex == Current.UpdateIndex, Update.Set(x => x.UpdateIndex, Current.UpdateIndex + 1));
-			return Result.ModifiedCount == 1;
+			UpdateResult result = await _artifacts.UpdateOneAsync<Artifact>(x => x.Id == current.Id && x.UpdateIndex == current.UpdateIndex, update.Set(x => x.UpdateIndex, current.UpdateIndex + 1));
+			return result.ModifiedCount == 1;
 		}
 
 		/// <summary>
 		/// Updates an artifact
 		/// </summary>
-		/// <param name="Artifact">The artifact</param>
-		/// <param name="NewMimeType">New mime type</param>
-		/// <param name="NewData">New data</param>
+		/// <param name="artifact">The artifact</param>
+		/// <param name="newMimeType">New mime type</param>
+		/// <param name="newData">New data</param>
 		/// <returns>Async task</returns>
-		public async Task<bool> UpdateArtifactAsync(IArtifact? Artifact, string NewMimeType, System.IO.Stream NewData)
+		public async Task<bool> UpdateArtifactAsync(IArtifact? artifact, string newMimeType, System.IO.Stream newData)
 		{
-			while (Artifact != null)
+			while (artifact != null)
 			{
-				UpdateDefinitionBuilder<Artifact> UpdateBuilder = Builders<Artifact>.Update;
+				UpdateDefinitionBuilder<Artifact> updateBuilder = Builders<Artifact>.Update;
 
-				List<UpdateDefinition<Artifact>> Updates = new List<UpdateDefinition<Artifact>>();
-				Updates.Add(UpdateBuilder.Set(x => x.MimeType, NewMimeType));
-				Updates.Add(UpdateBuilder.Set(x => x.Length, NewData.Length));
+				List<UpdateDefinition<Artifact>> updates = new List<UpdateDefinition<Artifact>>();
+				updates.Add(updateBuilder.Set(x => x.MimeType, newMimeType));
+				updates.Add(updateBuilder.Set(x => x.Length, newData.Length));
 
 				// re-upload the data to external
-				string ArtifactName = ValidateName(Artifact.Name);
-				await StorageBackend.WriteAsync(GetPath(Artifact.JobId, Artifact.StepId, ArtifactName), NewData);
+				string artifactName = ValidateName(artifact.Name);
+				await _storageBackend.WriteAsync(GetPath(artifact.JobId, artifact.StepId, artifactName), newData);
 
-				if (await TryUpdateArtifactAsync((Artifact)Artifact, UpdateBuilder.Combine(Updates)))
+				if (await TryUpdateArtifactAsync((Artifact)artifact, updateBuilder.Combine(updates)))
 				{
 					return true;
 				}
 
-				Artifact = await GetArtifactAsync(Artifact.Id);
+				artifact = await GetArtifactAsync(artifact.Id);
 			}
 			return false;
 		}
@@ -216,78 +200,78 @@ namespace Horde.Build.Services
 		/// <summary>
 		/// gets artifact data
 		/// </summary>
-		/// <param name="Artifact">The artifact</param>
+		/// <param name="artifact">The artifact</param>
 		/// <returns>The chunk data</returns>
-		public async Task<System.IO.Stream> OpenArtifactReadStreamAsync(IArtifact Artifact)
+		public async Task<System.IO.Stream> OpenArtifactReadStreamAsync(IArtifact artifact)
 		{
-			System.IO.Stream? Stream = await StorageBackend.ReadAsync(GetPath(Artifact.JobId, Artifact.StepId, Artifact.Name));
-			if (Stream == null)
+			System.IO.Stream? stream = await _storageBackend.ReadAsync(GetPath(artifact.JobId, artifact.StepId, artifact.Name));
+			if (stream == null)
 			{
-				throw new Exception($"Unable to get artifact {Artifact.Id}");
+				throw new Exception($"Unable to get artifact {artifact.Id}");
 			}
-			return Stream;
+			return stream;
 		}
 
 		/// <summary>
 		/// Get the path for an artifact
 		/// </summary>
-		/// <param name="JobId"></param>
-		/// <param name="StepId"></param>
-		/// <param name="Name"></param>
+		/// <param name="jobId"></param>
+		/// <param name="stepId"></param>
+		/// <param name="name"></param>
 		/// <returns></returns>
-		private static string GetPath(JobId JobId, SubResourceId? StepId, string Name)
+		private static string GetPath(JobId jobId, SubResourceId? stepId, string name)
 		{
-			if (StepId == null)
+			if (stepId == null)
 			{
-				return $"{JobId}/{Name}";
+				return $"{jobId}/{name}";
 			}
 			else
 			{
-				return $"{JobId}/{StepId.Value}/{Name}";
+				return $"{jobId}/{stepId.Value}/{name}";
 			}
 		}
 
 		/// <summary>
 		/// Checks that the name given is valid for an artifact
 		/// </summary>
-		/// <param name="Name"></param>
-		private static string ValidateName(string Name)
+		/// <param name="name"></param>
+		private static string ValidateName(string name)
 		{
-			string NewName = Name.Replace('\\', '/');
-			if (NewName.Length == 0)
+			string newName = name.Replace('\\', '/');
+			if (newName.Length == 0)
 			{
 				throw new ArgumentException("Artifact name is empty");
 			}
-			else if (NewName.StartsWith('/'))
+			else if (newName.StartsWith('/'))
 			{
-				throw new ArgumentException($"Artifact has an absolute path ({NewName})");
+				throw new ArgumentException($"Artifact has an absolute path ({newName})");
 			}
-			else if (NewName.EndsWith('/'))
+			else if (newName.EndsWith('/'))
 			{
-				throw new ArgumentException($"Artifact does not have a file name ({NewName})");
+				throw new ArgumentException($"Artifact does not have a file name ({newName})");
 			}
-			else if (NewName.Contains("//", StringComparison.Ordinal))
+			else if (newName.Contains("//", StringComparison.Ordinal))
 			{
-				throw new ArgumentException($"Artifact name contains an invalid directory ({NewName})");
+				throw new ArgumentException($"Artifact name contains an invalid directory ({newName})");
 			}
 
-			string InvalidChars = ":<>|\"";
-			for (int Idx = 0; Idx < NewName.Length; Idx++)
+			string invalidChars = ":<>|\"";
+			for (int idx = 0; idx < newName.Length; idx++)
 			{
-				int MinIdx = Idx;
-				for (; Idx < NewName.Length && NewName[Idx] != '/'; Idx++)
+				int minIdx = idx;
+				for (; idx < newName.Length && newName[idx] != '/'; idx++)
 				{
-					if (InvalidChars.Contains(NewName[Idx], StringComparison.Ordinal))
+					if (invalidChars.Contains(newName[idx], StringComparison.Ordinal))
 					{
-						throw new ArgumentException($"Invalid character in artifact name ({NewName})");
+						throw new ArgumentException($"Invalid character in artifact name ({newName})");
 					}
 				}
-				if ((Idx == MinIdx + 1 && NewName[MinIdx] == '.') || (Idx == MinIdx + 2 && NewName[MinIdx] == '.' && NewName[MinIdx + 1] == '.'))
+				if ((idx == minIdx + 1 && newName[minIdx] == '.') || (idx == minIdx + 2 && newName[minIdx] == '.' && newName[minIdx + 1] == '.'))
 				{
-					throw new ArgumentException($"Artifact may not contain symbolic directory names ({NewName})");
+					throw new ArgumentException($"Artifact may not contain symbolic directory names ({newName})");
 				}
 			}
-			return NewName;
+			return newName;
 		}
 	}
 }
\ No newline at end of file
diff --git a/Engine/Source/Programs/Horde/Horde.Build/Jobs/Impl/GraphCollection.cs b/Engine/Source/Programs/Horde/Horde.Build/Jobs/Impl/GraphCollection.cs
--- a/Engine/Source/Programs/Horde/Horde.Build/Jobs/Impl/GraphCollection.cs
+++ b/Engine/Source/Programs/Horde/Horde.Build/Jobs/Impl/GraphCollection.cs
@@ -1,13 +1,5 @@
 // Copyright Epic Games, Inc. All Rights Reserved.
 
-using EpicGames.Core;
-using Horde.Build.Api;
-using HordeCommon;
-using Horde.Build.Models;
-using Horde.Build.Services;
-using MongoDB.Bson;
-using MongoDB.Bson.Serialization.Attributes;
-using MongoDB.Driver;
 using System;
 using System.Collections.Concurrent;
 using System.Collections.Generic;
@@ -15,6 +7,13 @@ using System.Diagnostics;
 using System.Linq;
 using System.Threading;
 using System.Threading.Tasks;
+using EpicGames.Core;
+using Horde.Build.Models;
+using Horde.Build.Services;
+using HordeCommon;
+using MongoDB.Bson;
+using MongoDB.Bson.Serialization.Attributes;
+using MongoDB.Driver;
 
 namespace Horde.Build.Collections.Impl
 {
@@ -60,17 +59,17 @@ namespace Horde.Build.Collections.Impl
 				OrderDependencies = null!;
 			}
 
-			public Node(string Name, NodeRef[] InputDependencies, NodeRef[] OrderDependencies, Priority Priority, bool AllowRetry, bool RunEarly, bool Warnings, Dictionary<string, string>? Credentials, Dictionary<string, string>? Properties)
+			public Node(string name, NodeRef[] inputDependencies, NodeRef[] orderDependencies, Priority priority, bool allowRetry, bool runEarly, bool warnings, Dictionary<string, string>? credentials, Dictionary<string, string>? properties)
 			{
-				this.Name = Name;
-				this.InputDependencies = InputDependencies;
-				this.OrderDependencies = OrderDependencies;
-				this.Priority = Priority;
-				this.AllowRetry = AllowRetry;
-				this.RunEarly = RunEarly;
-				this.Warnings = Warnings;
-				this.Credentials = Credentials;
-				this.Properties = Properties;
+				Name = name;
+				InputDependencies = inputDependencies;
+				OrderDependencies = orderDependencies;
+				Priority = priority;
+				AllowRetry = allowRetry;
+				RunEarly = runEarly;
+				Warnings = warnings;
+				Credentials = credentials;
+				Properties = properties;
 			}
 		}
 
@@ -94,12 +93,12 @@ namespace Horde.Build.Collections.Impl
 			/// <summary>
 			/// Constructor
 			/// </summary>
-			/// <param name="AgentType">The type of agent to execute this group</param>
-			/// <param name="Nodes">Nodes to execute</param>
-			public NodeGroup(string AgentType, List<Node> Nodes)
+			/// <param name="agentType">The type of agent to execute this group</param>
+			/// <param name="nodes">Nodes to execute</param>
+			public NodeGroup(string agentType, List<Node> nodes)
 			{
-				this.AgentType = AgentType;
-				this.Nodes = Nodes;
+				AgentType = agentType;
+				Nodes = nodes;
 			}
 		}
 
@@ -116,10 +115,10 @@ namespace Horde.Build.Collections.Impl
 				Nodes = new List<NodeRef>();
 			}
 
-			public Aggregate(string Name, List<NodeRef> Nodes)
+			public Aggregate(string name, List<NodeRef> nodes)
 			{
-				this.Name = Name;
-				this.Nodes = Nodes;
+				Name = name;
+				Nodes = nodes;
 			}
 		}
 
@@ -158,15 +157,15 @@ namespace Horde.Build.Collections.Impl
 				IncludedNodes = new List<NodeRef>();
 			}
 
-			public Label(string? DashboardName, string? DashboardCategory, string? UgsName, string? UgsProject, LabelChange Change, List<NodeRef> RequiredNodes, List<NodeRef> IncludedNodes)
+			public Label(string? dashboardName, string? dashboardCategory, string? ugsName, string? ugsProject, LabelChange change, List<NodeRef> requiredNodes, List<NodeRef> includedNodes)
 			{
-				this.DashboardName = DashboardName;
-				this.DashboardCategory = DashboardCategory;
-				this.UgsName = UgsName;
-				this.UgsProject = UgsProject;
-				this.Change = Change;
-				this.RequiredNodes = RequiredNodes;
-				this.IncludedNodes = IncludedNodes;
+				DashboardName = dashboardName;
+				DashboardCategory = dashboardCategory;
+				UgsName = ugsName;
+				UgsProject = ugsProject;
+				Change = change;
+				RequiredNodes = requiredNodes;
+				IncludedNodes = includedNodes;
 			}
 		}
 
@@ -183,7 +182,7 @@ namespace Horde.Build.Collections.Impl
 			public List<Label> Labels { get; private set; } = new List<Label>();
 
 			[BsonIgnore]
-			IReadOnlyDictionary<string, NodeRef>? CachedNodeNameToRef;
+			IReadOnlyDictionary<string, NodeRef>? _cachedNodeNameToRef;
 
 			IReadOnlyList<INodeGroup> IGraph.Groups => Groups;
 			IReadOnlyList<IAggregate> IGraph.Aggregates => Aggregates;
@@ -194,75 +193,75 @@ namespace Horde.Build.Collections.Impl
 			{
 			}
 
-			public GraphDocument(List<NodeGroup> Groups, List<Aggregate> Aggregates, List<Label> Labels)
+			public GraphDocument(List<NodeGroup> groups, List<Aggregate> aggregates, List<Label> labels)
 			{
-				this.Groups = Groups;
-				this.Aggregates = Aggregates;
-				this.Labels = Labels;
-				this.Id = ContentHash.SHA1(BsonExtensionMethods.ToBson(this));
+				Groups = groups;
+				Aggregates = aggregates;
+				Labels = labels;
+				Id = ContentHash.SHA1(BsonExtensionMethods.ToBson(this));
 			}
 
-			public GraphDocument(GraphDocument BaseGraph, List<NewGroup>? NewGroupRequests, List<NewAggregate>? NewAggregateRequests, List<NewLabel>? NewLabelRequests)
+			public GraphDocument(GraphDocument baseGraph, List<NewGroup>? newGroupRequests, List<NewAggregate>? newAggregateRequests, List<NewLabel>? newLabelRequests)
 			{
-				Dictionary<string, NodeRef> NodeNameToRef = new Dictionary<string, NodeRef>(BaseGraph.GetNodeNameToRef(), StringComparer.OrdinalIgnoreCase);
+				Dictionary<string, NodeRef> nodeNameToRef = new Dictionary<string, NodeRef>(baseGraph.GetNodeNameToRef(), StringComparer.OrdinalIgnoreCase);
 
 				// Update the new list of groups
-				List<NodeGroup> NewGroups = new List<NodeGroup>(BaseGraph.Groups);
-				if (NewGroupRequests != null)
+				List<NodeGroup> newGroups = new List<NodeGroup>(baseGraph.Groups);
+				if (newGroupRequests != null)
 				{
-					foreach (NewGroup NewGroupRequest in NewGroupRequests)
+					foreach (NewGroup newGroupRequest in newGroupRequests)
 					{
-						List<Node> Nodes = new List<Node>();
-						foreach (NewNode NewNodeRequest in NewGroupRequest.Nodes)
+						List<Node> nodes = new List<Node>();
+						foreach (NewNode newNodeRequest in newGroupRequest.Nodes)
 						{
-							int NodeIdx = Nodes.Count;
+							int nodeIdx = nodes.Count;
 
-							Priority Priority = NewNodeRequest.Priority ?? Priority.Normal;
-							bool bAllowRetry = NewNodeRequest.AllowRetry ?? true;
-							bool bRunEarly = NewNodeRequest.RunEarly ?? false;
-							bool bWarnings = NewNodeRequest.Warnings ?? true;
+							Priority priority = newNodeRequest.Priority ?? Priority.Normal;
+							bool bAllowRetry = newNodeRequest.AllowRetry ?? true;
+							bool bRunEarly = newNodeRequest.RunEarly ?? false;
+							bool bWarnings = newNodeRequest.Warnings ?? true;
 
-							NodeRef[] InputDependencies = (NewNodeRequest.InputDependencies == null) ? Array.Empty<NodeRef>() : NewNodeRequest.InputDependencies.Select(x => NodeNameToRef[x]).ToArray();
-							NodeRef[] OrderDependencies = (NewNodeRequest.OrderDependencies == null) ? Array.Empty<NodeRef>() : NewNodeRequest.OrderDependencies.Select(x => NodeNameToRef[x]).ToArray();
-							OrderDependencies = OrderDependencies.Union(InputDependencies).ToArray();
-							Nodes.Add(new Node(NewNodeRequest.Name, InputDependencies, OrderDependencies, Priority, bAllowRetry, bRunEarly, bWarnings, NewNodeRequest.Credentials, NewNodeRequest.Properties));
+							NodeRef[] inputDependencies = (newNodeRequest.InputDependencies == null) ? Array.Empty<NodeRef>() : newNodeRequest.InputDependencies.Select(x => nodeNameToRef[x]).ToArray();
+							NodeRef[] orderDependencies = (newNodeRequest.OrderDependencies == null) ? Array.Empty<NodeRef>() : newNodeRequest.OrderDependencies.Select(x => nodeNameToRef[x]).ToArray();
+							orderDependencies = orderDependencies.Union(inputDependencies).ToArray();
+							nodes.Add(new Node(newNodeRequest.Name, inputDependencies, orderDependencies, priority, bAllowRetry, bRunEarly, bWarnings, newNodeRequest.Credentials, newNodeRequest.Properties));
 
-							NodeNameToRef.Add(NewNodeRequest.Name, new NodeRef(NewGroups.Count, NodeIdx));
+							nodeNameToRef.Add(newNodeRequest.Name, new NodeRef(newGroups.Count, nodeIdx));
 						}
-						NewGroups.Add(new NodeGroup(NewGroupRequest.AgentType, Nodes));
+						newGroups.Add(new NodeGroup(newGroupRequest.AgentType, nodes));
 					}
 				}
 
 				// Update the list of aggregates
-				List<Aggregate> NewAggregates = new List<Aggregate>(BaseGraph.Aggregates);
-				if (NewAggregateRequests != null)
+				List<Aggregate> newAggregates = new List<Aggregate>(baseGraph.Aggregates);
+				if (newAggregateRequests != null)
 				{
-					foreach (NewAggregate NewAggregateRequest in NewAggregateRequests)
+					foreach (NewAggregate newAggregateRequest in newAggregateRequests)
 					{
-						List<NodeRef> Nodes = NewAggregateRequest.Nodes.ConvertAll(x => NodeNameToRef[x]);
-						NewAggregates.Add(new Aggregate(NewAggregateRequest.Name, Nodes));
+						List<NodeRef> nodes = newAggregateRequest.Nodes.ConvertAll(x => nodeNameToRef[x]);
+						newAggregates.Add(new Aggregate(newAggregateRequest.Name, nodes));
 					}
 				}
 
 				// Update the list of labels
-				List<Label> NewLabels = new List<Label>(BaseGraph.Labels);
-				if (NewLabelRequests != null)
+				List<Label> newLabels = new List<Label>(baseGraph.Labels);
+				if (newLabelRequests != null)
 				{
-					foreach (NewLabel NewLabelRequest in NewLabelRequests)
+					foreach (NewLabel newLabelRequest in newLabelRequests)
 					{
-						List<NodeRef> RequiredNodes = NewLabelRequest.RequiredNodes.ConvertAll(x => NodeNameToRef[x]);
-						List<NodeRef> IncludedNodes = NewLabelRequest.IncludedNodes.ConvertAll(x => NodeNameToRef[x]);
-						NewLabels.Add(new Label(NewLabelRequest.DashboardName, NewLabelRequest.DashboardCategory, NewLabelRequest.UgsName, NewLabelRequest.UgsProject, NewLabelRequest.Change, RequiredNodes, IncludedNodes));
+						List<NodeRef> requiredNodes = newLabelRequest.RequiredNodes.ConvertAll(x => nodeNameToRef[x]);
+						List<NodeRef> includedNodes = newLabelRequest.IncludedNodes.ConvertAll(x => nodeNameToRef[x]);
+						newLabels.Add(new Label(newLabelRequest.DashboardName, newLabelRequest.DashboardCategory, newLabelRequest.UgsName, newLabelRequest.UgsProject, newLabelRequest.Change, requiredNodes, includedNodes));
 					}
 				}
 
 				// Create the new arrays
-				Groups = NewGroups;
-				Aggregates = NewAggregates;
-				Labels = NewLabels;
+				Groups = newGroups;
+				Aggregates = newAggregates;
+				Labels = newLabels;
 
 				// Create the new graph, and save the generated node lookup into it
-				CachedNodeNameToRef = NodeNameToRef;
+				_cachedNodeNameToRef = nodeNameToRef;
 
 				// Compute the hash
 				Id = ContentHash.SHA1(BsonExtensionMethods.ToBson(this));
@@ -270,21 +269,21 @@ namespace Horde.Build.Collections.Impl
 
 			public IReadOnlyDictionary<string, NodeRef> GetNodeNameToRef()
 			{
-				if (CachedNodeNameToRef == null)
+				if (_cachedNodeNameToRef == null)
 				{
-					Dictionary<string, NodeRef> NodeNameToRef = new Dictionary<string, NodeRef>(StringComparer.OrdinalIgnoreCase);
-					for (int GroupIdx = 0; GroupIdx < Groups.Count; GroupIdx++)
+					Dictionary<string, NodeRef> nodeNameToRef = new Dictionary<string, NodeRef>(StringComparer.OrdinalIgnoreCase);
+					for (int groupIdx = 0; groupIdx < Groups.Count; groupIdx++)
 					{
-						List<Node> Nodes = Groups[GroupIdx].Nodes;
-						for (int NodeIdx = 0; NodeIdx < Nodes.Count; NodeIdx++)
+						List<Node> nodes = Groups[groupIdx].Nodes;
+						for (int nodeIdx = 0; nodeIdx < nodes.Count; nodeIdx++)
 						{
-							Node Node = Nodes[NodeIdx];
-							NodeNameToRef[Node.Name] = new NodeRef(GroupIdx, NodeIdx);
+							Node node = nodes[nodeIdx];
+							nodeNameToRef[node.Name] = new NodeRef(groupIdx, nodeIdx);
 						}
 					}
-					CachedNodeNameToRef = NodeNameToRef;
+					_cachedNodeNameToRef = nodeNameToRef;
 				}
-				return CachedNodeNameToRef;
+				return _cachedNodeNameToRef;
 			}
 		}
 
@@ -296,15 +295,12 @@ namespace Horde.Build.Collections.Impl
 			/// <summary>
 			/// Time at which the graph was last accessed
 			/// </summary>
-			public long LastAccessTime
-			{
-				get { return LastAccessTimePrivate; }
-			}
+			public long LastAccessTime => _lastAccessTimePrivate;
 
 			/// <summary>
 			/// Backing value for <see cref="LastAccessTime"/>
 			/// </summary>
-			private long LastAccessTimePrivate;
+			private long _lastAccessTimePrivate;
 
 			/// <summary>
 			/// The graph instance
@@ -314,11 +310,11 @@ namespace Horde.Build.Collections.Impl
 			/// <summary>
 			/// Constructor
 			/// </summary>
-			/// <param name="Graph">The graph to store</param>
-			public CachedGraph(IGraph Graph)
+			/// <param name="graph">The graph to store</param>
+			public CachedGraph(IGraph graph)
 			{
-				this.LastAccessTimePrivate = Stopwatch.GetTimestamp();
-				this.Graph = Graph;
+				_lastAccessTimePrivate = Stopwatch.GetTimestamp();
+				Graph = graph;
 			}
 
 			/// <summary>
@@ -328,9 +324,9 @@ namespace Horde.Build.Collections.Impl
 			{
 				for (; ; )
 				{
-					long Time = Stopwatch.GetTimestamp();
-					long LastAccessTimeCopy = LastAccessTimePrivate;
-					if (Time < LastAccessTimeCopy || Interlocked.CompareExchange(ref LastAccessTimePrivate, Time, LastAccessTimeCopy) == LastAccessTimeCopy)
+					long time = Stopwatch.GetTimestamp();
+					long lastAccessTimeCopy = _lastAccessTimePrivate;
+					if (time < lastAccessTimeCopy || Interlocked.CompareExchange(ref _lastAccessTimePrivate, time, lastAccessTimeCopy) == lastAccessTimeCopy)
 					{
 						break;
 					}
@@ -341,7 +337,7 @@ namespace Horde.Build.Collections.Impl
 		/// <summary>
 		/// The jobs collection
 		/// </summary>
-		IMongoCollection<GraphDocument> Graphs;
+		readonly IMongoCollection<GraphDocument> _graphs;
 
 		/// <summary>
 		/// Maximum number of graphs to keep in the cache
@@ -351,33 +347,33 @@ namespace Horde.Build.Collections.Impl
 		/// <summary>
 		/// Cache of recently accessed graphs
 		/// </summary>
-		ConcurrentDictionary<ContentHash, CachedGraph> CachedGraphs = new ConcurrentDictionary<ContentHash, CachedGraph>();
+		readonly ConcurrentDictionary<ContentHash, CachedGraph> _cachedGraphs = new ConcurrentDictionary<ContentHash, CachedGraph>();
 
 		/// <summary>
 		/// Constructor
 		/// </summary>
-		/// <param name="DatabaseService">The database service singleton</param>
-		public GraphCollection(DatabaseService DatabaseService)
+		/// <param name="databaseService">The database service singleton</param>
+		public GraphCollection(DatabaseService databaseService)
 		{
-			Graphs = DatabaseService.GetCollection<GraphDocument>("Graphs");
+			_graphs = databaseService.GetCollection<GraphDocument>("Graphs");
 		}
 
 		/// <summary>
 		/// Adds a new graph document
 		/// </summary>
-		/// <param name="Graph">The graph to add</param>
+		/// <param name="graph">The graph to add</param>
 		/// <returns>Async task</returns>
-		async Task AddAsync(GraphDocument Graph)
+		async Task AddAsync(GraphDocument graph)
 		{
-			if (!await Graphs.Find(x => x.Id == Graph.Id).AnyAsync())
+			if (!await _graphs.Find(x => x.Id == graph.Id).AnyAsync())
 			{
 				try
 				{
-					await Graphs.InsertOneAsync(Graph);
+					await _graphs.InsertOneAsync(graph);
 				}
-				catch (MongoWriteException Ex)
+				catch (MongoWriteException ex)
 				{
-					if (Ex.WriteError.Category != ServerErrorCategory.DuplicateKey)
+					if (ex.WriteError.Category != ServerErrorCategory.DuplicateKey)
 					{
 						throw;
 					}
@@ -386,84 +382,84 @@ namespace Horde.Build.Collections.Impl
 		}
 
 		/// <inheritdoc/>
-		public async Task<IGraph> AddAsync(ITemplate Template)
+		public async Task<IGraph> AddAsync(ITemplate template)
 		{
-			Node Node = new Node(IJob.SetupNodeName, Array.Empty<NodeRef>(), Array.Empty<NodeRef>(), Priority.High, true, false, true, null, null);
-			NodeGroup Group = new NodeGroup(Template.InitialAgentType ?? "Win64", new List<Node> { Node });
+			Node node = new Node(IJob.SetupNodeName, Array.Empty<NodeRef>(), Array.Empty<NodeRef>(), Priority.High, true, false, true, null, null);
+			NodeGroup group = new NodeGroup(template.InitialAgentType ?? "Win64", new List<Node> { node });
 
-			GraphDocument Graph = new GraphDocument(new List<NodeGroup> { Group }, new List<Aggregate>(), new List<Label>());
-			await AddAsync(Graph);
-			return Graph;
+			GraphDocument graph = new GraphDocument(new List<NodeGroup> { group }, new List<Aggregate>(), new List<Label>());
+			await AddAsync(graph);
+			return graph;
 		}
 
 		/// <inheritdoc/>
-		public async Task<IGraph> AppendAsync(IGraph? BaseGraph, List<NewGroup>? NewGroupRequests, List<NewAggregate>? NewAggregateRequests, List<NewLabel>? NewLabelRequests)
+		public async Task<IGraph> AppendAsync(IGraph? baseGraph, List<NewGroup>? newGroupRequests, List<NewAggregate>? newAggregateRequests, List<NewLabel>? newLabelRequests)
 		{
-			GraphDocument Graph = new GraphDocument((GraphDocument?)BaseGraph ?? GraphDocument.Empty, NewGroupRequests, NewAggregateRequests, NewLabelRequests);
-			await AddAsync(Graph);
-			return Graph;
+			GraphDocument graph = new GraphDocument((GraphDocument?)baseGraph ?? GraphDocument.Empty, newGroupRequests, newAggregateRequests, newLabelRequests);
+			await AddAsync(graph);
+			return graph;
 		}
 
 		/// <inheritdoc/>
-		public async Task<IGraph> GetAsync(ContentHash? Hash)
+		public async Task<IGraph> GetAsync(ContentHash? hash)
 		{
 			// Special case for an empty graph request
-			if (Hash == null || Hash == ContentHash.Empty || Hash == GraphDocument.Empty.Id)
+			if (hash == null || hash == ContentHash.Empty || hash == GraphDocument.Empty.Id)
 			{
 				return GraphDocument.Empty;
 			}
 
 			// Try to read the graph from the cache
-			CachedGraph? CachedGraph;
-			if (CachedGraphs.TryGetValue(Hash, out CachedGraph))
+			CachedGraph? cachedGraph;
+			if (_cachedGraphs.TryGetValue(hash, out cachedGraph))
 			{
 				// Update the last access time
-				CachedGraph.Touch();
+				cachedGraph.Touch();
 			}
 			else
 			{
 				// Trim the cache
-				while (CachedGraphs.Count > MaxGraphs)
+				while (_cachedGraphs.Count > MaxGraphs)
 				{
-					ContentHash? RemoveHash = CachedGraphs.OrderBy(x => x.Value.LastAccessTime).Select(x => x.Key).FirstOrDefault();
-					if (RemoveHash == null || RemoveHash == ContentHash.Empty)
+					ContentHash? removeHash = _cachedGraphs.OrderBy(x => x.Value.LastAccessTime).Select(x => x.Key).FirstOrDefault();
+					if (removeHash == null || removeHash == ContentHash.Empty)
 					{
 						break;
 					}
-					CachedGraphs.TryRemove(RemoveHash, out _);
+					_cachedGraphs.TryRemove(removeHash, out _);
 				}
 
 				// Create the new entry
-				CachedGraph = new CachedGraph(await Graphs.Find<GraphDocument>(x => x.Id == Hash).FirstAsync());
-				CachedGraphs.TryAdd(Hash, CachedGraph);
+				cachedGraph = new CachedGraph(await _graphs.Find<GraphDocument>(x => x.Id == hash).FirstAsync());
+				_cachedGraphs.TryAdd(hash, cachedGraph);
 			}
-			return CachedGraph.Graph;
+			return cachedGraph.Graph;
 		}
 
 		/// <inheritdoc/>
-		public async Task<List<IGraph>> FindAllAsync(ContentHash[]? Hashes, int? Index, int? Count)
+		public async Task<List<IGraph>> FindAllAsync(ContentHash[]? hashes, int? index, int? count)
 		{
-			FilterDefinitionBuilder<GraphDocument> FilterBuilder = Builders<GraphDocument>.Filter;
+			FilterDefinitionBuilder<GraphDocument> filterBuilder = Builders<GraphDocument>.Filter;
 
-			FilterDefinition<GraphDocument> Filter = FilterBuilder.Empty;
-			if (Hashes != null)
+			FilterDefinition<GraphDocument> filter = filterBuilder.Empty;
+			if (hashes != null)
 			{
-				Filter &= FilterBuilder.In(x => x.Id, Hashes);
+				filter &= filterBuilder.In(x => x.Id, hashes);
 			}
 
-			List<GraphDocument> Results;
-			IFindFluent<GraphDocument, GraphDocument> Search = Graphs.Find(Filter);
-			if (Index != null)
+			List<GraphDocument> results;
+			IFindFluent<GraphDocument, GraphDocument> search = _graphs.Find(filter);
+			if (index != null)
 			{
-				Search = Search.Skip(Index.Value);
+				search = search.Skip(index.Value);
 			}
-			if (Count != null)
+			if (count != null)
 			{
-				Search = Search.Limit(Count.Value);
+				search = search.Limit(count.Value);
 			}
 
-			Results = await Search.ToListAsync();
-			return Results.ConvertAll<IGraph>(x => x);
+			results = await search.ToListAsync();
+			return results.ConvertAll<IGraph>(x => x);
 		}
 	}
 }
