Horde: Fix coding conventions within Horde.Build.

diff --git a/Engine/Source/Programs/Horde/Horde.Build/Jobs/Impl/JobCollection.cs b/Engine/Source/Programs/Horde/Horde.Build/Jobs/Impl/JobCollection.cs
--- a/Engine/Source/Programs/Horde/Horde.Build/Jobs/Impl/JobCollection.cs
+++ b/Engine/Source/Programs/Horde/Horde.Build/Jobs/Impl/JobCollection.cs
@@ -1,36 +1,36 @@
-﻿// Copyright Epic Games, Inc. All Rights Reserved.
+// Copyright Epic Games, Inc. All Rights Reserved.
 
+using System;
+using System.Collections.Generic;
+using System.Diagnostics.CodeAnalysis;
+using System.IO;
+using System.Linq;
+using System.Threading.Tasks;
 using EpicGames.Core;
 using Horde.Build.Acls;
 using Horde.Build.Api;
-using HordeCommon;
 using Horde.Build.Models;
 using Horde.Build.Services;
 using Horde.Build.Utilities;
+using HordeCommon;
 using Microsoft.Extensions.Logging;
 using MongoDB.Bson;
+using MongoDB.Bson.IO;
+using MongoDB.Bson.Serialization;
 using MongoDB.Bson.Serialization.Attributes;
 using MongoDB.Driver;
-using System;
-using System.Collections.Generic;
-using System.Diagnostics.CodeAnalysis;
-using System.IO;
-using System.Linq;
-using System.Threading.Tasks;
-using OpenTracing.Util;
 using OpenTracing;
-using MongoDB.Bson.Serialization;
-using MongoDB.Bson.IO;
+using OpenTracing.Util;
 
 namespace Horde.Build.Collections.Impl
 {
 	using JobId = ObjectId<IJob>;
 	using LeaseId = ObjectId<ILease>;
 	using LogId = ObjectId<ILogFile>;
+	using PoolId = StringId<IPool>;
 	using SessionId = ObjectId<ISession>;
 	using StreamId = StringId<IStream>;
 	using TemplateRefId = StringId<TemplateRef>;
-	using PoolId = StringId<IPool>;
 	using UserId = ObjectId<IUser>;
 
 	/// <summary>
@@ -79,14 +79,14 @@ namespace Horde.Build.Collections.Impl
 			public UserId? RetriedByUserId { get; set; }
 
 			[BsonElement("RetryByUser")]
-			public string? RetriedByUser_DEPRECATED { get; set; }
+			public string? RetriedByUserDeprecated { get; set; }
 
 			public bool AbortRequested { get; set; } = false;
 
 			public UserId? AbortedByUserId { get; set; }
 			
 			[BsonElement("AbortByUser")]
-			public string? AbortedByUser_DEPRECATED { get; set; }
+			public string? AbortedByUserDeprecated { get; set; }
 
 			[BsonIgnoreIfNull]
 			public List<Report>? Reports { get; set; }
@@ -104,10 +104,10 @@ namespace Horde.Build.Collections.Impl
 			{
 			}
 
-			public JobStepDocument(SubResourceId Id, int NodeIdx)
+			public JobStepDocument(SubResourceId id, int nodeIdx)
 			{
-				this.Id = Id;
-				this.NodeIdx = NodeIdx;
+				Id = id;
+				NodeIdx = nodeIdx;
 			}
 		}
 
@@ -162,10 +162,10 @@ namespace Horde.Build.Collections.Impl
 			{
 			}
 
-			public JobStepBatchDocument(SubResourceId Id, int GroupIdx)
+			public JobStepBatchDocument(SubResourceId id, int groupIdx)
 			{
-				this.Id = Id;
-				this.GroupIdx = GroupIdx;
+				Id = id;
+				GroupIdx = groupIdx;
 			}
 		}
 
@@ -181,17 +181,17 @@ namespace Horde.Build.Collections.Impl
 				Target = String.Empty;
 			}
 
-			public ChainedJobDocument(ChainedJobTemplate Trigger)
+			public ChainedJobDocument(ChainedJobTemplate trigger)
 			{
-				this.Target = Trigger.Trigger;
-				this.TemplateRefId = Trigger.TemplateRefId;
+				Target = trigger.Trigger;
+				TemplateRefId = trigger.TemplateRefId;
 			}
 		}
 
 		class LabelNotificationDocument
 		{
-			public int LabelIdx;
-			public ObjectId TriggerId;
+			public int _labelIdx;
+			public ObjectId _triggerId;
 		}
 
 		class JobDocument : IJob
@@ -208,13 +208,13 @@ namespace Horde.Build.Collections.Impl
 			public UserId? StartedByUserId { get; set; }
 
 			[BsonIgnoreIfNull, BsonElement("StartedByUser")]
-			public string? StartedByUser_DEPRECATED { get; set; }
+			public string? StartedByUserDeprecated { get; set; }
 
 			[BsonIgnoreIfNull]
 			public UserId? AbortedByUserId { get; set; }
 
 			[BsonIgnoreIfNull, BsonElement("AbortedByUser")]
-			public string? AbortedByUser_DEPRECATED { get; set; }
+			public string? AbortedByUserDeprecated { get; set; }
 
 			[BsonRequired]
 			public string Name { get; set; }
@@ -254,7 +254,7 @@ namespace Horde.Build.Collections.Impl
 			public bool ShowUgsAlerts { get; set; }
 			public string? NotificationChannel { get; set; }
 			public string? NotificationChannelFilter { get; set; }
-			public List<LabelNotificationDocument> LabelNotifications = new List<LabelNotificationDocument>();
+			public List<LabelNotificationDocument> _labelNotifications = new List<LabelNotificationDocument>();
 			public List<ChainedJobDocument> ChainedJobs { get; set; } = new List<ChainedJobDocument>();
 
 			[BsonIgnoreIfNull]
@@ -278,7 +278,7 @@ namespace Horde.Build.Collections.Impl
 			IReadOnlyList<IReport>? IJob.Reports => Reports;
 			IReadOnlyList<string> IJob.Arguments => Arguments;
 			IReadOnlyList<int> IJob.Issues => ReferencedByIssues;
-			IReadOnlyDictionary<int, ObjectId> IJob.LabelIdxToTriggerId => LabelNotifications.ToDictionary(x => x.LabelIdx, x => x.TriggerId);
+			IReadOnlyDictionary<int, ObjectId> IJob.LabelIdxToTriggerId => _labelNotifications.ToDictionary(x => x._labelIdx, x => x._triggerId);
 			IReadOnlyList<IChainedJob> IJob.ChainedJobs => ChainedJobs;
 
 			[BsonConstructor]
@@ -288,32 +288,32 @@ namespace Horde.Build.Collections.Impl
 				GraphHash = null!;
 			}
 
-			public JobDocument(JobId Id, StreamId StreamId, TemplateRefId TemplateId, ContentHash TemplateHash, ContentHash GraphHash, string Name, int Change, int CodeChange, int PreflightChange, int ClonedPreflightChange, UserId? StartedByUserId, Priority? Priority, bool? AutoSubmit, bool? UpdateIssues, bool? PromoteIssuesByDefault, DateTime CreateTimeUtc, List<ChainedJobDocument> ChainedJobs, bool ShowUgsBadges, bool ShowUgsAlerts, string? NotificationChannel, string? NotificationChannelFilter, List<string>? Arguments)
-			{
-				this.Id = Id;
-				this.StreamId = StreamId;
-				this.TemplateId = TemplateId;
-				this.TemplateHash = TemplateHash;
-				this.GraphHash = GraphHash;
-				this.Name = Name;
-				this.Change = Change;
-				this.CodeChange = CodeChange;
-				this.PreflightChange = PreflightChange;
-				this.ClonedPreflightChange = ClonedPreflightChange;
-				this.StartedByUserId = StartedByUserId;
-				this.Priority = Priority ?? HordeCommon.Priority.Normal;
-				this.AutoSubmit = AutoSubmit ?? false;
-				this.UpdateIssues = UpdateIssues ?? (StartedByUserId == null && PreflightChange == 0);
-				this.PromoteIssuesByDefault = PromoteIssuesByDefault ?? false;
-				this.CreateTimeUtc = CreateTimeUtc;
-				this.ChainedJobs = ChainedJobs;
-				this.ShowUgsBadges = ShowUgsBadges;
-				this.ShowUgsAlerts = ShowUgsAlerts;
-				this.NotificationChannel = NotificationChannel;
-				this.NotificationChannelFilter = NotificationChannelFilter;
-				this.Arguments = Arguments ?? this.Arguments;
-				this.NextSubResourceId = SubResourceId.Random();
-				this.UpdateTimeUtc = CreateTimeUtc;
+			public JobDocument(JobId id, StreamId streamId, TemplateRefId templateId, ContentHash templateHash, ContentHash graphHash, string name, int change, int codeChange, int preflightChange, int clonedPreflightChange, UserId? startedByUserId, Priority? priority, bool? autoSubmit, bool? updateIssues, bool? promoteIssuesByDefault, DateTime createTimeUtc, List<ChainedJobDocument> chainedJobs, bool showUgsBadges, bool showUgsAlerts, string? notificationChannel, string? notificationChannelFilter, List<string>? arguments)
+			{
+				Id = id;
+				StreamId = streamId;
+				TemplateId = templateId;
+				TemplateHash = templateHash;
+				GraphHash = graphHash;
+				Name = name;
+				Change = change;
+				CodeChange = codeChange;
+				PreflightChange = preflightChange;
+				ClonedPreflightChange = clonedPreflightChange;
+				StartedByUserId = startedByUserId;
+				Priority = priority ?? HordeCommon.Priority.Normal;
+				AutoSubmit = autoSubmit ?? false;
+				UpdateIssues = updateIssues ?? (startedByUserId == null && preflightChange == 0);
+				PromoteIssuesByDefault = promoteIssuesByDefault ?? false;
+				CreateTimeUtc = createTimeUtc;
+				ChainedJobs = chainedJobs;
+				ShowUgsBadges = showUgsBadges;
+				ShowUgsAlerts = showUgsAlerts;
+				NotificationChannel = notificationChannel;
+				NotificationChannelFilter = notificationChannelFilter;
+				Arguments = arguments ?? Arguments;
+				NextSubResourceId = SubResourceId.Random();
+				UpdateTimeUtc = createTimeUtc;
 			}
 		}
 
@@ -329,17 +329,17 @@ namespace Horde.Build.Collections.Impl
 			public readonly string TemplateId;
 			public readonly string SchedulePriority;
 
-			public DatabaseIndexes(IMongoCollection<JobDocument> Jobs, bool DatabaseReadOnlyMode) : base(DatabaseReadOnlyMode)
+			public DatabaseIndexes(IMongoCollection<JobDocument> jobs, bool databaseReadOnlyMode) : base(databaseReadOnlyMode)
 			{
-				StreamId = CreateOrGetIndex(Jobs, "StreamId_1", Builders<JobDocument>.IndexKeys.Ascending(x => x.StreamId));
-				Change = CreateOrGetIndex(Jobs, "Change_1", Builders<JobDocument>.IndexKeys.Ascending(x => x.Change));
-				PreflightChange = CreateOrGetIndex(Jobs, "PreflightChange_1", Builders<JobDocument>.IndexKeys.Ascending(x => x.PreflightChange));
-				CreateTimeUtc = CreateOrGetIndex(Jobs, "CreateTimeUtc_-1", Builders<JobDocument>.IndexKeys.Descending(x => x.CreateTimeUtc));
-				UpdateTimeUtc = CreateOrGetIndex(Jobs, "UpdateTimeUtc_-1", Builders<JobDocument>.IndexKeys.Descending(x => x.UpdateTimeUtc));
-				Name = CreateOrGetIndex(Jobs, "Name_1", Builders<JobDocument>.IndexKeys.Ascending(x => x.Name));
-				StartedByUserId = CreateOrGetIndex(Jobs, "StartedByUserId_1", Builders<JobDocument>.IndexKeys.Ascending(x => x.StartedByUserId));
-				TemplateId = CreateOrGetIndex(Jobs, "TemplateId_1", Builders<JobDocument>.IndexKeys.Ascending(x => x.TemplateId));
-				SchedulePriority = CreateOrGetIndex(Jobs, "SchedulePriority_-1", Builders<JobDocument>.IndexKeys.Descending(x => x.SchedulePriority));
+				StreamId = CreateOrGetIndex(jobs, "StreamId_1", Builders<JobDocument>.IndexKeys.Ascending(x => x.StreamId));
+				Change = CreateOrGetIndex(jobs, "Change_1", Builders<JobDocument>.IndexKeys.Ascending(x => x.Change));
+				PreflightChange = CreateOrGetIndex(jobs, "PreflightChange_1", Builders<JobDocument>.IndexKeys.Ascending(x => x.PreflightChange));
+				CreateTimeUtc = CreateOrGetIndex(jobs, "CreateTimeUtc_-1", Builders<JobDocument>.IndexKeys.Descending(x => x.CreateTimeUtc));
+				UpdateTimeUtc = CreateOrGetIndex(jobs, "UpdateTimeUtc_-1", Builders<JobDocument>.IndexKeys.Descending(x => x.UpdateTimeUtc));
+				Name = CreateOrGetIndex(jobs, "Name_1", Builders<JobDocument>.IndexKeys.Ascending(x => x.Name));
+				StartedByUserId = CreateOrGetIndex(jobs, "StartedByUserId_1", Builders<JobDocument>.IndexKeys.Ascending(x => x.StartedByUserId));
+				TemplateId = CreateOrGetIndex(jobs, "TemplateId_1", Builders<JobDocument>.IndexKeys.Ascending(x => x.TemplateId));
+				SchedulePriority = CreateOrGetIndex(jobs, "SchedulePriority_-1", Builders<JobDocument>.IndexKeys.Descending(x => x.SchedulePriority));
 			}
 		}
 
@@ -363,409 +363,404 @@ namespace Horde.Build.Collections.Impl
 		/// <summary>
 		/// The jobs collection
 		/// </summary>
-		IMongoCollection<JobDocument> Jobs;
+		readonly IMongoCollection<JobDocument> _jobs;
 
-		/// <summary>
-		/// The user collection
-		/// </summary>
-		IUserCollection UserCollection;
-		
 		/// <summary>
 		/// Clock representing wall-clock time
 		/// </summary>
-		IClock Clock;
+		readonly IClock _clock;
 
 		/// <summary>
 		/// Logger for output
 		/// </summary>
-		ILogger<JobCollection> Logger;
+		readonly ILogger<JobCollection> _logger;
 
 		/// <summary>
 		/// Creation and lookup of indexes
 		/// </summary>
-		DatabaseIndexes Indexes;
+		readonly DatabaseIndexes _indexes;
 
 		/// <summary>
 		/// Constructor
 		/// </summary>
-		/// <param name="DatabaseService">The database service singleton</param>
-		/// <param name="UserCollection"></param>
-		/// <param name="Clock"></param>
-		/// <param name="Logger">The logger instance</param>
-		public JobCollection(DatabaseService DatabaseService, IUserCollection UserCollection, IClock Clock, ILogger<JobCollection> Logger)
+		/// <param name="databaseService">The database service singleton</param>
+		/// <param name="clock"></param>
+		/// <param name="logger">The logger instance</param>
+		public JobCollection(DatabaseService databaseService, IClock clock, ILogger<JobCollection> logger)
 		{
-			this.UserCollection = UserCollection;
-			this.Clock = Clock;
-			this.Logger = Logger;
+			_clock = clock;
+			_logger = logger;
 
-			Jobs = DatabaseService.GetCollection<JobDocument>("Jobs");
-			Indexes = new DatabaseIndexes(Jobs, DatabaseService.ReadOnlyMode);
+			_jobs = databaseService.GetCollection<JobDocument>("Jobs");
+			_indexes = new DatabaseIndexes(_jobs, databaseService.ReadOnlyMode);
 		}
 
-		static Task PostLoadAsync(JobDocument Job)
+		static Task PostLoadAsync(JobDocument job)
 		{
-			if (Job.GraphHash == ContentHash.Empty)
+			if (job.GraphHash == ContentHash.Empty)
 			{
-				Job.Batches.Clear();
+				job.Batches.Clear();
 			}
 			return Task.CompletedTask;
 		}
 
-		static JobDocument Clone(JobDocument Job)
+		static JobDocument Clone(JobDocument job)
 		{
-			using (MemoryStream Stream = new MemoryStream())
+			using (MemoryStream stream = new MemoryStream())
 			{
-				using (BsonBinaryWriter Writer = new BsonBinaryWriter(Stream))
+				using (BsonBinaryWriter writer = new BsonBinaryWriter(stream))
 				{
-					BsonSerializer.Serialize(Writer, Job);
+					BsonSerializer.Serialize(writer, job);
 				}
-				return BsonSerializer.Deserialize<JobDocument>(Stream.ToArray());
+				return BsonSerializer.Deserialize<JobDocument>(stream.ToArray());
 			}
 		}
 
 		/// <inheritdoc/>
-		[SuppressMessage("Compiler", "CA1054:URI parameters should not be strings")]
-		public async Task<IJob> AddAsync(JobId JobId, StreamId StreamId, TemplateRefId TemplateRefId, ContentHash TemplateHash, IGraph Graph, string Name, int Change, int CodeChange, int? PreflightChange, int? ClonedPreflightChange, UserId? StartedByUserId, Priority? Priority, bool? AutoSubmit, bool? UpdateIssues, bool? PromoteIssuesByDefault, List<ChainedJobTemplate>? ChainedJobs, bool ShowUgsBadges, bool ShowUgsAlerts, string? NotificationChannel, string? NotificationChannelFilter, List<string>? Arguments)
+		public async Task<IJob> AddAsync(JobId jobId, StreamId streamId, TemplateRefId templateRefId, ContentHash templateHash, IGraph graph, string name, int change, int codeChange, int? preflightChange, int? clonedPreflightChange, UserId? startedByUserId, Priority? priority, bool? autoSubmit, bool? updateIssues, bool? promoteIssuesByDefault, List<ChainedJobTemplate>? chainedJobs, bool showUgsBadges, bool showUgsAlerts, string? notificationChannel, string? notificationChannelFilter, List<string>? arguments)
 		{
-			List<ChainedJobDocument> JobTriggers = new List<ChainedJobDocument>();
-			if (ChainedJobs == null)
+			List<ChainedJobDocument> jobTriggers = new List<ChainedJobDocument>();
+			if (chainedJobs == null)
 			{
-				JobTriggers = new List<ChainedJobDocument>();
+				jobTriggers = new List<ChainedJobDocument>();
 			}
 			else
 			{
-				JobTriggers = ChainedJobs.ConvertAll(x => new ChainedJobDocument(x));
+				jobTriggers = chainedJobs.ConvertAll(x => new ChainedJobDocument(x));
 			}
 
-			JobDocument NewJob = new JobDocument(JobId, StreamId, TemplateRefId, TemplateHash, Graph.Id, Name, Change, CodeChange, PreflightChange ?? 0, ClonedPreflightChange ?? 0, StartedByUserId, Priority, AutoSubmit, UpdateIssues, PromoteIssuesByDefault, DateTime.UtcNow, JobTriggers, ShowUgsBadges, ShowUgsAlerts, NotificationChannel, NotificationChannelFilter, Arguments);
-			CreateBatches(NewJob, Graph, Logger);
+			JobDocument newJob = new JobDocument(jobId, streamId, templateRefId, templateHash, graph.Id, name, change, codeChange, preflightChange ?? 0, clonedPreflightChange ?? 0, startedByUserId, priority, autoSubmit, updateIssues, promoteIssuesByDefault, DateTime.UtcNow, jobTriggers, showUgsBadges, showUgsAlerts, notificationChannel, notificationChannelFilter, arguments);
+			CreateBatches(newJob, graph, _logger);
 
-			await Jobs.InsertOneAsync(NewJob);
+			await _jobs.InsertOneAsync(newJob);
 
-			return NewJob;
+			return newJob;
 		}
 
 		/// <inheritdoc/>
-		public async Task<IJob?> GetAsync(JobId JobId)
+		public async Task<IJob?> GetAsync(JobId jobId)
 		{
-			JobDocument? Job = await Jobs.Find<JobDocument>(x => x.Id == JobId).FirstOrDefaultAsync();
-			if (Job != null)
+			JobDocument? job = await _jobs.Find<JobDocument>(x => x.Id == jobId).FirstOrDefaultAsync();
+			if (job != null)
 			{
-				await PostLoadAsync(Job);
+				await PostLoadAsync(job);
 			}
-			return Job;
+			return job;
 		}
 
 		/// <inheritdoc/>
-		public async Task<bool> RemoveAsync(IJob Job)
+		public async Task<bool> RemoveAsync(IJob job)
 		{
-			DeleteResult Result = await Jobs.DeleteOneAsync(x => x.Id == Job.Id && x.UpdateIndex == Job.UpdateIndex);
-			return Result.DeletedCount > 0;
+			DeleteResult result = await _jobs.DeleteOneAsync(x => x.Id == job.Id && x.UpdateIndex == job.UpdateIndex);
+			return result.DeletedCount > 0;
 		}
 
 		/// <inheritdoc/>
-		public async Task RemoveStreamAsync(StreamId StreamId)
+		public async Task RemoveStreamAsync(StreamId streamId)
 		{
-			await Jobs.DeleteManyAsync(x => x.StreamId == StreamId);
+			await _jobs.DeleteManyAsync(x => x.StreamId == streamId);
 		}
 
 		/// <inheritdoc/>
-		public async Task<IJobPermissions?> GetPermissionsAsync(JobId JobId)
+		public async Task<IJobPermissions?> GetPermissionsAsync(JobId jobId)
 		{
-			return await Jobs.Find<JobDocument>(x => x.Id == JobId).Project<JobPermissions>(JobPermissions.Projection).FirstOrDefaultAsync();
+			return await _jobs.Find<JobDocument>(x => x.Id == jobId).Project<JobPermissions>(JobPermissions.Projection).FirstOrDefaultAsync();
 		}
 
 		/// <inheritdoc/>
-		public async Task<List<IJob>> FindAsync(JobId[]? JobIds, StreamId? StreamId, string? Name, TemplateRefId[]? Templates, int? MinChange, int? MaxChange, int? PreflightChange, bool? PreflightOnly, UserId ? PreflightStartedByUser, UserId? StartedByUser, DateTimeOffset? MinCreateTime, DateTimeOffset? MaxCreateTime, DateTimeOffset? ModifiedBefore, DateTimeOffset? ModifiedAfter, int? Index, int? Count, bool ConsistentRead, string? IndexHint, bool? ExcludeUserJobs)
+		public async Task<List<IJob>> FindAsync(JobId[]? jobIds, StreamId? streamId, string? name, TemplateRefId[]? templates, int? minChange, int? maxChange, int? preflightChange, bool? preflightOnly, UserId ? preflightStartedByUser, UserId? startedByUser, DateTimeOffset? minCreateTime, DateTimeOffset? maxCreateTime, DateTimeOffset? modifiedBefore, DateTimeOffset? modifiedAfter, int? index, int? count, bool consistentRead, string? indexHint, bool? excludeUserJobs)
 		{
-			FilterDefinitionBuilder<JobDocument> FilterBuilder = Builders<JobDocument>.Filter;
+			FilterDefinitionBuilder<JobDocument> filterBuilder = Builders<JobDocument>.Filter;
 
-			FilterDefinition<JobDocument> Filter = FilterBuilder.Empty;
-			if (JobIds != null && JobIds.Length > 0)
+			FilterDefinition<JobDocument> filter = filterBuilder.Empty;
+			if (jobIds != null && jobIds.Length > 0)
 			{
-				Filter &= FilterBuilder.In(x => x.Id, JobIds);
+				filter &= filterBuilder.In(x => x.Id, jobIds);
 			}
-			if (StreamId != null)
+			if (streamId != null)
 			{
-				Filter &= FilterBuilder.Eq(x => x.StreamId, StreamId.Value);
+				filter &= filterBuilder.Eq(x => x.StreamId, streamId.Value);
 			}
-			if (Name != null)
+			if (name != null)
 			{
-				if (Name.StartsWith("$", StringComparison.InvariantCulture))
+				if (name.StartsWith("$", StringComparison.InvariantCulture))
 				{
-					BsonRegularExpression Regex = new BsonRegularExpression(Name.Substring(1), "i");
-					Filter &= FilterBuilder.Regex(x => x.Name, Regex);
+					BsonRegularExpression regex = new BsonRegularExpression(name.Substring(1), "i");
+					filter &= filterBuilder.Regex(x => x.Name, regex);
 				}
 				else
 				{
-					Filter &= FilterBuilder.Eq(x => x.Name, Name);
+					filter &= filterBuilder.Eq(x => x.Name, name);
 				}				
 			}
-			if (Templates != null)
+			if (templates != null)
 			{
-				Filter &= FilterBuilder.In(x => x.TemplateId, Templates);
+				filter &= filterBuilder.In(x => x.TemplateId, templates);
 			}
-			if (MinChange != null)
+			if (minChange != null)
 			{
-				Filter &= FilterBuilder.Gte(x => x.Change, MinChange);
+				filter &= filterBuilder.Gte(x => x.Change, minChange);
 			}
-			if (MaxChange != null)
+			if (maxChange != null)
 			{
-				Filter &= FilterBuilder.Lte(x => x.Change, MaxChange);
+				filter &= filterBuilder.Lte(x => x.Change, maxChange);
 			}
-			if (PreflightChange != null)
+			if (preflightChange != null)
 			{
-				Filter &= FilterBuilder.Eq(x => x.PreflightChange, PreflightChange);
+				filter &= filterBuilder.Eq(x => x.PreflightChange, preflightChange);
 			}
-			if (PreflightOnly != null && PreflightOnly.Value)
+			if (preflightOnly != null && preflightOnly.Value)
 			{
-				Filter &= FilterBuilder.Ne(x => x.PreflightChange, 0);
+				filter &= filterBuilder.Ne(x => x.PreflightChange, 0);
 			}
-			if (ExcludeUserJobs != null && ExcludeUserJobs.Value)
+			if (excludeUserJobs != null && excludeUserJobs.Value)
 			{
-				Filter &= FilterBuilder.Eq(x => x.StartedByUserId, null);
+				filter &= filterBuilder.Eq(x => x.StartedByUserId, null);
 			}
 			else
 			{
-				if (PreflightStartedByUser != null)
+				if (preflightStartedByUser != null)
 				{
-					Filter &= FilterBuilder.Or(FilterBuilder.Eq(x => x.PreflightChange, 0), FilterBuilder.Eq(x => x.StartedByUserId, PreflightStartedByUser));
+					filter &= filterBuilder.Or(filterBuilder.Eq(x => x.PreflightChange, 0), filterBuilder.Eq(x => x.StartedByUserId, preflightStartedByUser));
 				}
-				if (StartedByUser != null)
+				if (startedByUser != null)
 				{
-					Filter &= FilterBuilder.Eq(x => x.StartedByUserId, StartedByUser);
+					filter &= filterBuilder.Eq(x => x.StartedByUserId, startedByUser);
 				}
 			}
-			if (MinCreateTime != null)
+			if (minCreateTime != null)
 			{
-				Filter &= FilterBuilder.Gte(x => x.CreateTimeUtc!, MinCreateTime.Value.UtcDateTime);
+				filter &= filterBuilder.Gte(x => x.CreateTimeUtc!, minCreateTime.Value.UtcDateTime);
 			}
-			if (MaxCreateTime != null)
+			if (maxCreateTime != null)
 			{
-				Filter &= FilterBuilder.Lte(x => x.CreateTimeUtc!, MaxCreateTime.Value.UtcDateTime);
+				filter &= filterBuilder.Lte(x => x.CreateTimeUtc!, maxCreateTime.Value.UtcDateTime);
 			}
-			if (ModifiedBefore != null)
+			if (modifiedBefore != null)
 			{
-				Filter &= FilterBuilder.Lte(x => x.UpdateTimeUtc!, ModifiedBefore.Value.UtcDateTime);
+				filter &= filterBuilder.Lte(x => x.UpdateTimeUtc!, modifiedBefore.Value.UtcDateTime);
 			}
-			if (ModifiedAfter != null)
+			if (modifiedAfter != null)
 			{
-				Filter &= FilterBuilder.Gte(x => x.UpdateTimeUtc!, ModifiedAfter.Value.UtcDateTime);
+				filter &= filterBuilder.Gte(x => x.UpdateTimeUtc!, modifiedAfter.Value.UtcDateTime);
 			}
 
-			List<JobDocument> Results;
-			using (IScope Scope = GlobalTracer.Instance.BuildSpan("Jobs.Find").StartActive())
+			List<JobDocument> results;
+			using (IScope scope = GlobalTracer.Instance.BuildSpan("Jobs.Find").StartActive())
 			{
-				IMongoCollection<JobDocument> Collection = ConsistentRead ? Jobs : Jobs.WithReadPreference(ReadPreference.SecondaryPreferred);
+				IMongoCollection<JobDocument> collection = consistentRead ? _jobs : _jobs.WithReadPreference(ReadPreference.SecondaryPreferred);
 
-				FindOptions? FindOptions = null;
-				if (IndexHint != null)
+				FindOptions? findOptions = null;
+				if (indexHint != null)
 				{
-					FindOptions = new FindOptions { Hint = new BsonString(IndexHint) };
+					findOptions = new FindOptions { Hint = new BsonString(indexHint) };
 				}
 				
-				IFindFluent<JobDocument, JobDocument> Query = Collection.Find<JobDocument>(Filter, FindOptions).SortByDescending(x => x.CreateTimeUtc!);
+				IFindFluent<JobDocument, JobDocument> query = collection.Find<JobDocument>(filter, findOptions).SortByDescending(x => x.CreateTimeUtc!);
 				
-				if (Index != null)
+				if (index != null)
 				{
-					Query = Query.Skip(Index.Value);
+					query = query.Skip(index.Value);
 				}
-				if (Count != null)
+				if (count != null)
 				{
-					Query = Query.Limit(Count.Value);
+					query = query.Limit(count.Value);
 				}
-				Results = await Query.ToListAsync();
+				results = await query.ToListAsync();
 			}
-			foreach (JobDocument Result in Results)
+			foreach (JobDocument result in results)
 			{
-				await PostLoadAsync(Result);
+				await PostLoadAsync(result);
 			}
-			return Results.ConvertAll<JobDocument, IJob>(x => x);
+			return results.ConvertAll<JobDocument, IJob>(x => x);
 		}
 
 		/// <inheritdoc/>
-		public async Task<List<IJob>> FindLatestByStreamWithTemplatesAsync(StreamId StreamId, TemplateRefId[] Templates, UserId? PreflightStartedByUser, DateTimeOffset? MaxCreateTime, DateTimeOffset? ModifiedAfter, int? Index, int? Count, bool ConsistentRead)
+		public async Task<List<IJob>> FindLatestByStreamWithTemplatesAsync(StreamId streamId, TemplateRefId[] templates, UserId? preflightStartedByUser, DateTimeOffset? maxCreateTime, DateTimeOffset? modifiedAfter, int? index, int? count, bool consistentRead)
 		{
-			string IndexHint = Indexes.CreateTimeUtc;
-			if (ModifiedAfter != null) IndexHint = Indexes.UpdateTimeUtc;
+			string indexHint = _indexes.CreateTimeUtc;
+			if (modifiedAfter != null)
+			{
+				indexHint = _indexes.UpdateTimeUtc;
+			}
 			
 			// This find call uses an index hint. Modifying the parameter passed to FindAsync can affect execution time a lot as the query planner is forced to use the specified index.
 			// Casting to interface to benefit from default parameter values
 			return await (this as IJobCollection).FindAsync(
-				StreamId: StreamId, Templates: Templates, PreflightStartedByUser: PreflightStartedByUser, ModifiedAfter: ModifiedAfter, MaxCreateTime: MaxCreateTime,
-				Index: Index, Count: Count, IndexHint: IndexHint, ConsistentRead: ConsistentRead);
+				streamId: streamId, templates: templates, preflightStartedByUser: preflightStartedByUser, modifiedAfter: modifiedAfter, maxCreateTime: maxCreateTime,
+				index: index, count: count, indexHint: indexHint, consistentRead: consistentRead);
 		}
 
 		/// <inheritdoc/>
-		public async Task<IJob?> TryUpdateJobAsync(IJob InJob, IGraph Graph, string? Name, Priority? Priority, bool? AutoSubmit, int? AutoSubmitChange, string? AutoSubmitMessage, UserId? AbortedByUserId, ObjectId? NotificationTriggerId, List<Report>? Reports, List<string>? Arguments, KeyValuePair<int, ObjectId>? LabelIdxToTriggerId, KeyValuePair<TemplateRefId, JobId>? JobTrigger)
+		public async Task<IJob?> TryUpdateJobAsync(IJob inJob, IGraph graph, string? name, Priority? priority, bool? autoSubmit, int? autoSubmitChange, string? autoSubmitMessage, UserId? abortedByUserId, ObjectId? notificationTriggerId, List<Report>? reports, List<string>? arguments, KeyValuePair<int, ObjectId>? labelIdxToTriggerId, KeyValuePair<TemplateRefId, JobId>? jobTrigger)
 		{
 			// Create the update 
-			UpdateDefinitionBuilder<JobDocument> UpdateBuilder = Builders<JobDocument>.Update;
-			List<UpdateDefinition<JobDocument>> Updates = new List<UpdateDefinition<JobDocument>>();
+			UpdateDefinitionBuilder<JobDocument> updateBuilder = Builders<JobDocument>.Update;
+			List<UpdateDefinition<JobDocument>> updates = new List<UpdateDefinition<JobDocument>>();
 
 			// Flag for whether to update batches
 			bool bUpdateBatches = false;
 
 			// Build the update list
-			JobDocument JobDocument = Clone((JobDocument)InJob);
-			if (Name != null)
+			JobDocument jobDocument = Clone((JobDocument)inJob);
+			if (name != null)
 			{
-				JobDocument.Name = Name;
-				Updates.Add(UpdateBuilder.Set(x => x.Name, JobDocument.Name));
+				jobDocument.Name = name;
+				updates.Add(updateBuilder.Set(x => x.Name, jobDocument.Name));
 			}
-			if (Priority != null)
+			if (priority != null)
 			{
-				JobDocument.Priority = Priority.Value;
-				Updates.Add(UpdateBuilder.Set(x => x.Priority, JobDocument.Priority));
+				jobDocument.Priority = priority.Value;
+				updates.Add(updateBuilder.Set(x => x.Priority, jobDocument.Priority));
 			}
-			if (AutoSubmit != null)
+			if (autoSubmit != null)
 			{
-				JobDocument.AutoSubmit = AutoSubmit.Value;
-				Updates.Add(UpdateBuilder.Set(x => x.AutoSubmit, JobDocument.AutoSubmit));
+				jobDocument.AutoSubmit = autoSubmit.Value;
+				updates.Add(updateBuilder.Set(x => x.AutoSubmit, jobDocument.AutoSubmit));
 			}
-			if (AutoSubmitChange != null)
+			if (autoSubmitChange != null)
 			{
-				JobDocument.AutoSubmitChange = AutoSubmitChange.Value;
-				Updates.Add(UpdateBuilder.Set(x => x.AutoSubmitChange, JobDocument.AutoSubmitChange));
+				jobDocument.AutoSubmitChange = autoSubmitChange.Value;
+				updates.Add(updateBuilder.Set(x => x.AutoSubmitChange, jobDocument.AutoSubmitChange));
 			}
-			if (AutoSubmitMessage != null)
+			if (autoSubmitMessage != null)
 			{
-				JobDocument.AutoSubmitMessage = (AutoSubmitMessage.Length == 0)? null : AutoSubmitMessage;
-				Updates.Add(UpdateBuilder.SetOrUnsetNullRef(x => x.AutoSubmitMessage, JobDocument.AutoSubmitMessage));
+				jobDocument.AutoSubmitMessage = (autoSubmitMessage.Length == 0)? null : autoSubmitMessage;
+				updates.Add(updateBuilder.SetOrUnsetNullRef(x => x.AutoSubmitMessage, jobDocument.AutoSubmitMessage));
 			}
-			if (AbortedByUserId != null && JobDocument.AbortedByUserId == null)
+			if (abortedByUserId != null && jobDocument.AbortedByUserId == null)
 			{
-				JobDocument.AbortedByUserId = AbortedByUserId;
-				Updates.Add(UpdateBuilder.Set(x => x.AbortedByUserId, JobDocument.AbortedByUserId));
+				jobDocument.AbortedByUserId = abortedByUserId;
+				updates.Add(updateBuilder.Set(x => x.AbortedByUserId, jobDocument.AbortedByUserId));
 				bUpdateBatches = true;
 			}
-			if (NotificationTriggerId != null)
+			if (notificationTriggerId != null)
 			{
-				JobDocument.NotificationTriggerId = NotificationTriggerId.Value;
-				Updates.Add(UpdateBuilder.Set(x => x.NotificationTriggerId, NotificationTriggerId));
+				jobDocument.NotificationTriggerId = notificationTriggerId.Value;
+				updates.Add(updateBuilder.Set(x => x.NotificationTriggerId, notificationTriggerId));
 			}
-			if (LabelIdxToTriggerId != null)
+			if (labelIdxToTriggerId != null)
 			{
-				if (JobDocument.LabelNotifications.Any(x => x.LabelIdx == LabelIdxToTriggerId.Value.Key))
+				if (jobDocument._labelNotifications.Any(x => x._labelIdx == labelIdxToTriggerId.Value.Key))
 				{
 					throw new ArgumentException("Cannot update label trigger that already exists");
 				}
-				JobDocument.LabelNotifications.Add(new LabelNotificationDocument { LabelIdx = LabelIdxToTriggerId.Value.Key, TriggerId = LabelIdxToTriggerId.Value.Value });
-				Updates.Add(UpdateBuilder.Set(x => x.LabelNotifications, JobDocument.LabelNotifications));
+				jobDocument._labelNotifications.Add(new LabelNotificationDocument { _labelIdx = labelIdxToTriggerId.Value.Key, _triggerId = labelIdxToTriggerId.Value.Value });
+				updates.Add(updateBuilder.Set(x => x._labelNotifications, jobDocument._labelNotifications));
 			}
-			if (JobTrigger != null)
+			if (jobTrigger != null)
 			{
-				for (int Idx = 0; Idx < JobDocument.ChainedJobs.Count; Idx++)
+				for (int idx = 0; idx < jobDocument.ChainedJobs.Count; idx++)
 				{
-					ChainedJobDocument JobTriggerDocument = JobDocument.ChainedJobs[Idx];
-					if (JobTriggerDocument.TemplateRefId == JobTrigger.Value.Key)
+					ChainedJobDocument jobTriggerDocument = jobDocument.ChainedJobs[idx];
+					if (jobTriggerDocument.TemplateRefId == jobTrigger.Value.Key)
 					{
-						int LocalIdx = Idx;
-						JobTriggerDocument.JobId = JobTrigger.Value.Value;
-						Updates.Add(UpdateBuilder.Set(x => x.ChainedJobs[LocalIdx].JobId, JobTrigger.Value.Value));
+						int localIdx = idx;
+						jobTriggerDocument.JobId = jobTrigger.Value.Value;
+						updates.Add(updateBuilder.Set(x => x.ChainedJobs[localIdx].JobId, jobTrigger.Value.Value));
 					}
 				}
 			}
-			if (Reports != null)
+			if (reports != null)
 			{
-				JobDocument.Reports ??= new List<Report>();
-				JobDocument.Reports.RemoveAll(x => Reports.Any(y => y.Name == x.Name));
-				JobDocument.Reports.AddRange(Reports);
-				Updates.Add(UpdateBuilder.Set(x => x.Reports, JobDocument.Reports));
+				jobDocument.Reports ??= new List<Report>();
+				jobDocument.Reports.RemoveAll(x => reports.Any(y => y.Name == x.Name));
+				jobDocument.Reports.AddRange(reports);
+				updates.Add(updateBuilder.Set(x => x.Reports, jobDocument.Reports));
 			}
-			if (Arguments != null)
+			if (arguments != null)
 			{
-				HashSet<string> ModifiedArguments = new HashSet<string>(JobDocument.Arguments);
-				ModifiedArguments.SymmetricExceptWith(Arguments);
+				HashSet<string> modifiedArguments = new HashSet<string>(jobDocument.Arguments);
+				modifiedArguments.SymmetricExceptWith(arguments);
 
-				foreach (string ModifiedArgument in ModifiedArguments)
+				foreach (string modifiedArgument in modifiedArguments)
 				{
-					if (ModifiedArgument.StartsWith(IJob.TargetArgumentPrefix, StringComparison.OrdinalIgnoreCase))
+					if (modifiedArgument.StartsWith(IJob.TargetArgumentPrefix, StringComparison.OrdinalIgnoreCase))
 					{
 						bUpdateBatches = true;
 					}
 				}
 
-				JobDocument.Arguments = Arguments.ToList();
-				Updates.Add(UpdateBuilder.Set(x => x.Arguments, JobDocument.Arguments));
+				jobDocument.Arguments = arguments.ToList();
+				updates.Add(updateBuilder.Set(x => x.Arguments, jobDocument.Arguments));
 			}
 
 			// Update the batches
 			if (bUpdateBatches)
 			{
-				UpdateBatches(JobDocument, Graph, Updates, Logger);
+				UpdateBatches(jobDocument, graph, updates, _logger);
 			}
 
 			// Update the new list of job steps
-			return await TryUpdateAsync(JobDocument, UpdateBuilder.Combine(Updates));
+			return await TryUpdateAsync(jobDocument, updateBuilder.Combine(updates));
 		}
 
 		/// <inheritdoc/>
-		public async Task<IJob?> TryUpdateBatchAsync(IJob Job, IGraph Graph, SubResourceId BatchId, LogId? NewLogId, JobStepBatchState? NewState, JobStepBatchError? NewError)
+		public async Task<IJob?> TryUpdateBatchAsync(IJob job, IGraph graph, SubResourceId batchId, LogId? newLogId, JobStepBatchState? newState, JobStepBatchError? newError)
 		{
-			JobDocument JobDocument = Clone((JobDocument)Job);
+			JobDocument jobDocument = Clone((JobDocument)job);
 
 			// Find the index of the appropriate batch
-			int BatchIdx = JobDocument.Batches.FindIndex(x => x.Id == BatchId);
-			if (BatchIdx == -1)
+			int batchIdx = jobDocument.Batches.FindIndex(x => x.Id == batchId);
+			if (batchIdx == -1)
 			{
 				return null;
 			}
 
 			// If we're marking the batch as complete and there are still steps to run (eg. because the agent crashed), we need to mark all the steps as complete first
-			JobStepBatchDocument Batch = JobDocument.Batches[BatchIdx];
+			JobStepBatchDocument batch = jobDocument.Batches[batchIdx];
 
 			// Create the update 
-			UpdateDefinitionBuilder<JobDocument> UpdateBuilder = Builders<JobDocument>.Update;
-			List<UpdateDefinition<JobDocument>> Updates = new List<UpdateDefinition<JobDocument>>();
+			UpdateDefinitionBuilder<JobDocument> updateBuilder = Builders<JobDocument>.Update;
+			List<UpdateDefinition<JobDocument>> updates = new List<UpdateDefinition<JobDocument>>();
 
 			// Update the batch
-			if (NewLogId != null)
+			if (newLogId != null)
 			{
-				Batch.LogId = NewLogId.Value;
-				Updates.Add(UpdateBuilder.Set(x => x.Batches[BatchIdx].LogId, Batch.LogId));
+				batch.LogId = newLogId.Value;
+				updates.Add(updateBuilder.Set(x => x.Batches[batchIdx].LogId, batch.LogId));
 			}
-			if (NewState != null)
+			if (newState != null)
 			{
-				Batch.State = NewState.Value;
-				Updates.Add(UpdateBuilder.Set(x => x.Batches[BatchIdx].State, Batch.State));
+				batch.State = newState.Value;
+				updates.Add(updateBuilder.Set(x => x.Batches[batchIdx].State, batch.State));
 
-				if (Batch.StartTime == null && NewState >= JobStepBatchState.Starting)
+				if (batch.StartTime == null && newState >= JobStepBatchState.Starting)
 				{
-					Batch.StartTime = Clock.UtcNow;
-					Updates.Add(UpdateBuilder.Set(x => x.Batches[BatchIdx].StartTime, Batch.StartTime));
+					batch.StartTime = _clock.UtcNow;
+					updates.Add(updateBuilder.Set(x => x.Batches[batchIdx].StartTime, batch.StartTime));
 				}
-				if (NewState == JobStepBatchState.Complete)
+				if (newState == JobStepBatchState.Complete)
 				{
-					Batch.FinishTime = Clock.UtcNow;
-					Updates.Add(UpdateBuilder.Set(x => x.Batches[BatchIdx].FinishTime, Batch.FinishTime));
+					batch.FinishTime = _clock.UtcNow;
+					updates.Add(updateBuilder.Set(x => x.Batches[batchIdx].FinishTime, batch.FinishTime));
 				}
 			}
-			if (NewError != null)
+			if (newError != null)
 			{
-				Batch.Error = NewError.Value;
-				Updates.Add(UpdateBuilder.Set(x => x.Batches[BatchIdx].Error, Batch.Error));
+				batch.Error = newError.Value;
+				updates.Add(updateBuilder.Set(x => x.Batches[batchIdx].Error, batch.Error));
 			}
 
 			// If the batch is being marked as incomplete, see if we can reschedule any of the work
-			if (NewError == JobStepBatchError.Incomplete)
+			if (newError == JobStepBatchError.Incomplete)
 			{
 				// Get the new list of retried nodes
-				List<NodeRef> RetriedNodes = JobDocument.RetriedNodes ?? new List<NodeRef>();
+				List<NodeRef> retriedNodes = jobDocument.RetriedNodes ?? new List<NodeRef>();
 
 				// Check if there are any steps that need to be run again
 				bool bUpdateState = false;
-				foreach (JobStepDocument Step in Batch.Steps)
+				foreach (JobStepDocument step in batch.Steps)
 				{
-					if (Step.State == JobStepState.Ready || Step.State == JobStepState.Waiting)
+					if (step.State == JobStepState.Ready || step.State == JobStepState.Waiting)
 					{
-						if (CanRetryNode(JobDocument, Batch.GroupIdx, Step.NodeIdx))
+						if (CanRetryNode(jobDocument, batch.GroupIdx, step.NodeIdx))
 						{
-							RetriedNodes.Add(new NodeRef(Batch.GroupIdx, Step.NodeIdx));
+							retriedNodes.Add(new NodeRef(batch.GroupIdx, step.NodeIdx));
 						}
 						else
 						{
-							Step.State = JobStepState.Skipped;
+							step.State = JobStepState.Skipped;
 						}
 						bUpdateState = true;
 					}
@@ -774,168 +769,168 @@ namespace Horde.Build.Collections.Impl
 				// Update the steps
 				if (bUpdateState)
 				{
-					Updates.Clear();
-					UpdateBatches(JobDocument, Graph, Updates, Logger);
+					updates.Clear();
+					UpdateBatches(jobDocument, graph, updates, _logger);
 
-					JobDocument.RetriedNodes = RetriedNodes;
-					Updates.Add(UpdateBuilder.Set(x => x.RetriedNodes, JobDocument.RetriedNodes));
+					jobDocument.RetriedNodes = retriedNodes;
+					updates.Add(updateBuilder.Set(x => x.RetriedNodes, jobDocument.RetriedNodes));
 				}
 			}
 
 			// Update the new list of job steps
-			return await TryUpdateAsync(JobDocument, Updates);
+			return await TryUpdateAsync(jobDocument, updates);
 		}
 
 		/// <inheritdoc/>
-		public Task<IJob?> TryUpdateStepAsync(IJob Job, IGraph Graph, SubResourceId BatchId, SubResourceId StepId, JobStepState NewState, JobStepOutcome NewOutcome, bool? NewAbortRequested, UserId? NewAbortByUserId, LogId? NewLogId, ObjectId? NewNotificationTriggerId, UserId? NewRetryByUserId, Priority? NewPriority, List<Report>? NewReports, Dictionary<string, string?>? NewProperties)
+		public Task<IJob?> TryUpdateStepAsync(IJob job, IGraph graph, SubResourceId batchId, SubResourceId stepId, JobStepState newState, JobStepOutcome newOutcome, bool? newAbortRequested, UserId? newAbortByUserId, LogId? newLogId, ObjectId? newNotificationTriggerId, UserId? newRetryByUserId, Priority? newPriority, List<Report>? newReports, Dictionary<string, string?>? newProperties)
 		{
-			JobDocument JobDocument = Clone((JobDocument)Job);
+			JobDocument jobDocument = Clone((JobDocument)job);
 
 			// Create the update 
-			UpdateDefinitionBuilder<JobDocument> UpdateBuilder = Builders<JobDocument>.Update;
-			List<UpdateDefinition<JobDocument>> Updates = new List<UpdateDefinition<JobDocument>>();
+			UpdateDefinitionBuilder<JobDocument> updateBuilder = Builders<JobDocument>.Update;
+			List<UpdateDefinition<JobDocument>> updates = new List<UpdateDefinition<JobDocument>>();
 
 			// Update the appropriate batch
 			bool bRefreshBatches = false;
 			bool bRefreshDependentJobSteps = false;
-			for (int LoopBatchIdx = 0; LoopBatchIdx < JobDocument.Batches.Count; LoopBatchIdx++)
+			for (int loopBatchIdx = 0; loopBatchIdx < jobDocument.Batches.Count; loopBatchIdx++)
 			{
-				int BatchIdx = LoopBatchIdx; // For lambda capture
-				JobStepBatchDocument Batch = JobDocument.Batches[BatchIdx];
-				if (Batch.Id == BatchId)
+				int batchIdx = loopBatchIdx; // For lambda capture
+				JobStepBatchDocument batch = jobDocument.Batches[batchIdx];
+				if (batch.Id == batchId)
 				{
-					for (int LoopStepIdx = 0; LoopStepIdx < Batch.Steps.Count; LoopStepIdx++)
+					for (int loopStepIdx = 0; loopStepIdx < batch.Steps.Count; loopStepIdx++)
 					{
-						int StepIdx = LoopStepIdx; // For lambda capture
-						JobStepDocument Step = Batch.Steps[StepIdx];
-						if (Step.Id == StepId)
+						int stepIdx = loopStepIdx; // For lambda capture
+						JobStepDocument step = batch.Steps[stepIdx];
+						if (step.Id == stepId)
 						{
 							// Update the state
-							if (NewState != JobStepState.Unspecified && Step.State != NewState)
+							if (newState != JobStepState.Unspecified && step.State != newState)
 							{
-								Step.State = NewState;
-								Updates.Add(UpdateBuilder.Set(x => x.Batches[BatchIdx].Steps[StepIdx].State, Step.State));
+								step.State = newState;
+								updates.Add(updateBuilder.Set(x => x.Batches[batchIdx].Steps[stepIdx].State, step.State));
 
-								if (Step.State == JobStepState.Running)
+								if (step.State == JobStepState.Running)
 								{
-									Step.StartTime = Clock.UtcNow;
-									Updates.Add(UpdateBuilder.Set(x => x.Batches[BatchIdx].Steps[StepIdx].StartTime, Step.StartTime));
+									step.StartTime = _clock.UtcNow;
+									updates.Add(updateBuilder.Set(x => x.Batches[batchIdx].Steps[stepIdx].StartTime, step.StartTime));
 								}
-								else if (Step.State == JobStepState.Completed || Step.State == JobStepState.Aborted)
+								else if (step.State == JobStepState.Completed || step.State == JobStepState.Aborted)
 								{
-									Step.FinishTime = Clock.UtcNow;
-									Updates.Add(UpdateBuilder.Set(x => x.Batches[BatchIdx].Steps[StepIdx].FinishTime, Step.FinishTime));
+									step.FinishTime = _clock.UtcNow;
+									updates.Add(updateBuilder.Set(x => x.Batches[batchIdx].Steps[stepIdx].FinishTime, step.FinishTime));
 								}
 
 								bRefreshDependentJobSteps = true;
 							}
 
 							// Update the job outcome
-							if (NewOutcome != JobStepOutcome.Unspecified && Step.Outcome != NewOutcome)
+							if (newOutcome != JobStepOutcome.Unspecified && step.Outcome != newOutcome)
 							{
-								Step.Outcome = NewOutcome;
-								Updates.Add(UpdateBuilder.Set(x => x.Batches[BatchIdx].Steps[StepIdx].Outcome, Step.Outcome));
+								step.Outcome = newOutcome;
+								updates.Add(updateBuilder.Set(x => x.Batches[batchIdx].Steps[stepIdx].Outcome, step.Outcome));
 
 								bRefreshDependentJobSteps = true;
 							}
 
 							// Update the request abort status
-							if (NewAbortRequested != null && Step.AbortRequested == false)
+							if (newAbortRequested != null && step.AbortRequested == false)
 							{
-								Step.AbortRequested = NewAbortRequested.Value;
-								Updates.Add(UpdateBuilder.Set(x => x.Batches[BatchIdx].Steps[StepIdx].AbortRequested, Step.AbortRequested));
+								step.AbortRequested = newAbortRequested.Value;
+								updates.Add(updateBuilder.Set(x => x.Batches[batchIdx].Steps[stepIdx].AbortRequested, step.AbortRequested));
 
 								bRefreshDependentJobSteps = true;
 							}
 
 							// Update the user that requested the abort
-							if (NewAbortByUserId != null && Step.AbortedByUserId == null)
+							if (newAbortByUserId != null && step.AbortedByUserId == null)
 							{
-								Step.AbortedByUserId = NewAbortByUserId;
-								Updates.Add(UpdateBuilder.Set(x => x.Batches[BatchIdx].Steps[StepIdx].AbortedByUserId, Step.AbortedByUserId));
+								step.AbortedByUserId = newAbortByUserId;
+								updates.Add(updateBuilder.Set(x => x.Batches[batchIdx].Steps[stepIdx].AbortedByUserId, step.AbortedByUserId));
 
 								bRefreshDependentJobSteps = true;
 							}
 
 							// Update the log id
-							if (NewLogId != null && Step.LogId != NewLogId.Value)
+							if (newLogId != null && step.LogId != newLogId.Value)
 							{
-								Step.LogId = NewLogId.Value;
-								Updates.Add(UpdateBuilder.Set(x => x.Batches[BatchIdx].Steps[StepIdx].LogId, Step.LogId));
+								step.LogId = newLogId.Value;
+								updates.Add(updateBuilder.Set(x => x.Batches[batchIdx].Steps[stepIdx].LogId, step.LogId));
 
 								bRefreshDependentJobSteps = true;
 							}
 
 							// Update the notification trigger id
-							if (NewNotificationTriggerId != null && Step.NotificationTriggerId == null)
+							if (newNotificationTriggerId != null && step.NotificationTriggerId == null)
 							{
-								Step.NotificationTriggerId = NewNotificationTriggerId.Value;
-								Updates.Add(UpdateBuilder.Set(x => x.Batches[BatchIdx].Steps[StepIdx].NotificationTriggerId, Step.NotificationTriggerId));
+								step.NotificationTriggerId = newNotificationTriggerId.Value;
+								updates.Add(updateBuilder.Set(x => x.Batches[batchIdx].Steps[stepIdx].NotificationTriggerId, step.NotificationTriggerId));
 							}
 
 							// Update the retry flag
-							if (NewRetryByUserId != null && Step.RetriedByUserId == null)
+							if (newRetryByUserId != null && step.RetriedByUserId == null)
 							{
-								Step.RetriedByUserId = NewRetryByUserId;
-								Updates.Add(UpdateBuilder.Set(x => x.Batches[BatchIdx].Steps[StepIdx].RetriedByUserId, Step.RetriedByUserId));
+								step.RetriedByUserId = newRetryByUserId;
+								updates.Add(updateBuilder.Set(x => x.Batches[batchIdx].Steps[stepIdx].RetriedByUserId, step.RetriedByUserId));
 
 								bRefreshBatches = true;
 							}
 
 							// Update the priority
-							if (NewPriority != null && NewPriority.Value != Step.Priority)
+							if (newPriority != null && newPriority.Value != step.Priority)
 							{
-								Step.Priority = NewPriority.Value;
-								Updates.Add(UpdateBuilder.Set(x => x.Batches[BatchIdx].Steps[StepIdx].Priority, Step.Priority));
+								step.Priority = newPriority.Value;
+								updates.Add(updateBuilder.Set(x => x.Batches[batchIdx].Steps[stepIdx].Priority, step.Priority));
 
 								bRefreshBatches = true;
 							}
 
 							// Add any new reports
-							if (NewReports != null)
+							if (newReports != null)
 							{
-								Step.Reports ??= new List<Report>();
-								Step.Reports.RemoveAll(x => NewReports.Any(y => y.Name == x.Name));
-								Step.Reports.AddRange(NewReports);
-								Updates.Add(UpdateBuilder.Set(x => x.Batches[BatchIdx].Steps[StepIdx].Reports, Step.Reports));
+								step.Reports ??= new List<Report>();
+								step.Reports.RemoveAll(x => newReports.Any(y => y.Name == x.Name));
+								step.Reports.AddRange(newReports);
+								updates.Add(updateBuilder.Set(x => x.Batches[batchIdx].Steps[stepIdx].Reports, step.Reports));
 							}
 
 							// Apply any property updates
-							if (NewProperties != null)
+							if (newProperties != null)
 							{
-								if (Step.Properties == null)
+								if (step.Properties == null)
 								{
-									Step.Properties = new Dictionary<string, string>(StringComparer.Ordinal);
+									step.Properties = new Dictionary<string, string>(StringComparer.Ordinal);
 								}
 
-								foreach (KeyValuePair<string, string?> Pair in NewProperties)
+								foreach (KeyValuePair<string, string?> pair in newProperties)
 								{
-									if (Pair.Value == null)
+									if (pair.Value == null)
 									{
-										Step.Properties.Remove(Pair.Key);
+										step.Properties.Remove(pair.Key);
 									}
 									else
 									{
-										Step.Properties[Pair.Key] = Pair.Value;
+										step.Properties[pair.Key] = pair.Value;
 									}
 								}
 
-								if (Step.Properties.Count == 0)
+								if (step.Properties.Count == 0)
 								{
-									Step.Properties = null;
-									Updates.Add(UpdateBuilder.Unset(x => x.Batches[BatchIdx].Steps[StepIdx].Properties));
+									step.Properties = null;
+									updates.Add(updateBuilder.Unset(x => x.Batches[batchIdx].Steps[stepIdx].Properties));
 								}
 								else
 								{
-									foreach (KeyValuePair<string, string?> Pair in NewProperties)
+									foreach (KeyValuePair<string, string?> pair in newProperties)
 									{
-										if (Pair.Value == null)
+										if (pair.Value == null)
 										{
-											Updates.Add(UpdateBuilder.Unset(x => x.Batches[BatchIdx].Steps[StepIdx].Properties![Pair.Key]));
+											updates.Add(updateBuilder.Unset(x => x.Batches[batchIdx].Steps[stepIdx].Properties![pair.Key]));
 										}
 										else
 										{
-											Updates.Add(UpdateBuilder.Set(x => x.Batches[BatchIdx].Steps[StepIdx].Properties![Pair.Key], Pair.Value));
+											updates.Add(updateBuilder.Set(x => x.Batches[batchIdx].Steps[stepIdx].Properties![pair.Key], pair.Value));
 										}
 									}
 								}
@@ -950,492 +945,492 @@ namespace Horde.Build.Collections.Impl
 			// Update the batches
 			if (bRefreshBatches)
 			{
-				Updates.Clear(); // UpdateBatches will update the entire batches list. We need to remove all the individual step updates to avoid an exception.
-				UpdateBatches(JobDocument, Graph, Updates, Logger);
+				updates.Clear(); // UpdateBatches will update the entire batches list. We need to remove all the individual step updates to avoid an exception.
+				UpdateBatches(jobDocument, graph, updates, _logger);
 			}
 
 			// Update the state of dependent jobsteps
 			if (bRefreshDependentJobSteps)
 			{
-				RefreshDependentJobSteps(JobDocument, Graph, Updates, Logger);
-				RefreshJobPriority(JobDocument, Updates);
+				RefreshDependentJobSteps(jobDocument, graph, updates, _logger);
+				RefreshJobPriority(jobDocument, updates);
 			}
 
 			// Update the new list of job steps
-			return TryUpdateAsync(JobDocument, Updates);
+			return TryUpdateAsync(jobDocument, updates);
 		}
 
-		Task<IJob?> TryUpdateAsync(JobDocument Job, List<UpdateDefinition<JobDocument>> Updates)
+		Task<IJob?> TryUpdateAsync(JobDocument job, List<UpdateDefinition<JobDocument>> updates)
 		{
-			if (Updates.Count == 0)
+			if (updates.Count == 0)
 			{
-				return Task.FromResult<IJob?>(Job);
+				return Task.FromResult<IJob?>(job);
 			}
 			else
 			{
-				return TryUpdateAsync(Job, Builders<JobDocument>.Update.Combine(Updates));
+				return TryUpdateAsync(job, Builders<JobDocument>.Update.Combine(updates));
 			}
 		}
 
-		async Task<IJob?> TryUpdateAsync(JobDocument Job, UpdateDefinition<JobDocument> Update)
+		async Task<IJob?> TryUpdateAsync(JobDocument job, UpdateDefinition<JobDocument> update)
 		{
-			int NewUpdateIndex = Job.UpdateIndex + 1;
-			Update = Update.Set(x => x.UpdateIndex, NewUpdateIndex);
+			int newUpdateIndex = job.UpdateIndex + 1;
+			update = update.Set(x => x.UpdateIndex, newUpdateIndex);
 
-			DateTime NewUpdateTimeUtc = DateTime.UtcNow;
-			Update = Update.Set(x => x.UpdateTimeUtc, NewUpdateTimeUtc);
+			DateTime newUpdateTimeUtc = DateTime.UtcNow;
+			update = update.Set(x => x.UpdateTimeUtc, newUpdateTimeUtc);
 
-			UpdateResult Result = await Jobs.UpdateOneAsync<JobDocument>(x => x.Id == Job.Id && x.UpdateIndex == Job.UpdateIndex, Update);
-			if (Result.ModifiedCount > 0)
+			UpdateResult result = await _jobs.UpdateOneAsync<JobDocument>(x => x.Id == job.Id && x.UpdateIndex == job.UpdateIndex, update);
+			if (result.ModifiedCount > 0)
 			{
-				Job.UpdateIndex = NewUpdateIndex;
-				Job.UpdateTimeUtc = NewUpdateTimeUtc;
-				return Job;
+				job.UpdateIndex = newUpdateIndex;
+				job.UpdateTimeUtc = newUpdateTimeUtc;
+				return job;
 			}
 			return null;
 		}
 
 		/// <inheritdoc/>
-		public Task<IJob?> TryRemoveFromDispatchQueueAsync(IJob Job)
+		public Task<IJob?> TryRemoveFromDispatchQueueAsync(IJob job)
 		{
-			JobDocument JobDocument = Clone((JobDocument)Job);
+			JobDocument jobDocument = Clone((JobDocument)job);
 
 			// Create the update 
-			UpdateDefinitionBuilder<JobDocument> UpdateBuilder = Builders<JobDocument>.Update;
-			List<UpdateDefinition<JobDocument>> Updates = new List<UpdateDefinition<JobDocument>>();
+			UpdateDefinitionBuilder<JobDocument> updateBuilder = Builders<JobDocument>.Update;
+			List<UpdateDefinition<JobDocument>> updates = new List<UpdateDefinition<JobDocument>>();
 
-			JobDocument.SchedulePriority = 0;
-			Updates.Add(UpdateBuilder.Set(x => x.SchedulePriority, JobDocument.SchedulePriority));
+			jobDocument.SchedulePriority = 0;
+			updates.Add(updateBuilder.Set(x => x.SchedulePriority, jobDocument.SchedulePriority));
 
 			// Update the new list of job steps
-			return TryUpdateAsync(JobDocument, Updates);
+			return TryUpdateAsync(jobDocument, updates);
 		}
 
 		/// <inheritdoc/>
-		public Task<IJob?> TryUpdateGraphAsync(IJob Job, IGraph NewGraph)
+		public Task<IJob?> TryUpdateGraphAsync(IJob job, IGraph newGraph)
 		{
-			JobDocument JobDocument = (JobDocument)Job;
+			JobDocument jobDocument = (JobDocument)job;
 
 			// Create the update 
-			UpdateDefinitionBuilder<JobDocument> UpdateBuilder = Builders<JobDocument>.Update;
-			List<UpdateDefinition<JobDocument>> Updates = new List<UpdateDefinition<JobDocument>>();
+			UpdateDefinitionBuilder<JobDocument> updateBuilder = Builders<JobDocument>.Update;
+			List<UpdateDefinition<JobDocument>> updates = new List<UpdateDefinition<JobDocument>>();
 
-			JobDocument.GraphHash = NewGraph.Id;
-			Updates.Add(UpdateBuilder.Set(x => x.GraphHash, Job.GraphHash));
+			jobDocument.GraphHash = newGraph.Id;
+			updates.Add(updateBuilder.Set(x => x.GraphHash, job.GraphHash));
 
-			UpdateBatches(JobDocument, NewGraph, Updates, Logger);
+			UpdateBatches(jobDocument, newGraph, updates, _logger);
 
 			// Update the new list of job steps
-			return TryUpdateAsync(JobDocument, Updates);
+			return TryUpdateAsync(jobDocument, updates);
 		}
 
 		/// <inheritdoc/>
-		public async Task AddIssueToJobAsync(JobId JobId, int IssueId)
+		public async Task AddIssueToJobAsync(JobId jobId, int issueId)
 		{
-			FilterDefinition<JobDocument> JobFilter = Builders<JobDocument>.Filter.Eq(x => x.Id, JobId);
-			UpdateDefinition<JobDocument> JobUpdate = Builders<JobDocument>.Update.AddToSet(x => x.ReferencedByIssues, IssueId).Inc(x => x.UpdateIndex, 1).Max(x => x.UpdateTimeUtc, DateTime.UtcNow);
-			await Jobs.UpdateOneAsync(JobFilter, JobUpdate);
+			FilterDefinition<JobDocument> jobFilter = Builders<JobDocument>.Filter.Eq(x => x.Id, jobId);
+			UpdateDefinition<JobDocument> jobUpdate = Builders<JobDocument>.Update.AddToSet(x => x.ReferencedByIssues, issueId).Inc(x => x.UpdateIndex, 1).Max(x => x.UpdateTimeUtc, DateTime.UtcNow);
+			await _jobs.UpdateOneAsync(jobFilter, jobUpdate);
 		}
 
 		/// <inheritdoc/>
 		public async Task<List<IJob>> GetDispatchQueueAsync()
 		{
-			List<JobDocument> NewJobs = await Jobs.Find(x => x.SchedulePriority > 0).SortByDescending(x => x.SchedulePriority).ThenBy(x => x.CreateTimeUtc).ToListAsync();
-			foreach (JobDocument Result in NewJobs)
+			List<JobDocument> newJobs = await _jobs.Find(x => x.SchedulePriority > 0).SortByDescending(x => x.SchedulePriority).ThenBy(x => x.CreateTimeUtc).ToListAsync();
+			foreach (JobDocument result in newJobs)
 			{
-				await PostLoadAsync(Result);
+				await PostLoadAsync(result);
 			}
-			return NewJobs.ConvertAll<JobDocument, IJob>(x => x);
+			return newJobs.ConvertAll<JobDocument, IJob>(x => x);
 		}
 
 		/// <summary>
 		/// Marks a job as skipped
 		/// </summary>
-		/// <param name="Job">The job to update</param>
-		/// <param name="Graph">Graph for the job</param>
-		/// <param name="Reason">Reason for this batch being failed</param>
+		/// <param name="job">The job to update</param>
+		/// <param name="graph">Graph for the job</param>
+		/// <param name="reason">Reason for this batch being failed</param>
 		/// <returns>Updated version of the job</returns>
-		public async Task<IJob?> SkipAllBatchesAsync(IJob? Job, IGraph Graph, JobStepBatchError Reason)
+		public async Task<IJob?> SkipAllBatchesAsync(IJob? job, IGraph graph, JobStepBatchError reason)
 		{
-			while (Job != null)
+			while (job != null)
 			{
-				JobDocument JobDocument = Clone((JobDocument)Job);
+				JobDocument jobDocument = Clone((JobDocument)job);
 
-				for (int BatchIdx = 0; BatchIdx < JobDocument.Batches.Count; BatchIdx++)
+				for (int batchIdx = 0; batchIdx < jobDocument.Batches.Count; batchIdx++)
 				{
-					JobStepBatchDocument Batch = JobDocument.Batches[BatchIdx];
-					if (Batch.State == JobStepBatchState.Ready || Batch.State == JobStepBatchState.Waiting)
+					JobStepBatchDocument batch = jobDocument.Batches[batchIdx];
+					if (batch.State == JobStepBatchState.Ready || batch.State == JobStepBatchState.Waiting)
 					{
-						Batch.State = JobStepBatchState.Complete;
-						Batch.Error = Reason;
-						Batch.FinishTime = DateTimeOffset.UtcNow;
+						batch.State = JobStepBatchState.Complete;
+						batch.Error = reason;
+						batch.FinishTime = DateTimeOffset.UtcNow;
 
-						for (int StepIdx = 0; StepIdx < Batch.Steps.Count; StepIdx++)
+						for (int stepIdx = 0; stepIdx < batch.Steps.Count; stepIdx++)
 						{
-							JobStepDocument Step = Batch.Steps[StepIdx];
-							if (Step.State == JobStepState.Ready || Step.State == JobStepState.Waiting)
+							JobStepDocument step = batch.Steps[stepIdx];
+							if (step.State == JobStepState.Ready || step.State == JobStepState.Waiting)
 							{
-								Step.State = JobStepState.Completed;
-								Step.Outcome = JobStepOutcome.Failure;
+								step.State = JobStepState.Completed;
+								step.Outcome = JobStepOutcome.Failure;
 							}
 						}
 					}
 				}
 
-				List<UpdateDefinition<JobDocument>> Updates = new List<UpdateDefinition<JobDocument>>();
-				UpdateBatches(JobDocument, Graph, Updates, Logger);
+				List<UpdateDefinition<JobDocument>> updates = new List<UpdateDefinition<JobDocument>>();
+				UpdateBatches(jobDocument, graph, updates, _logger);
 
-				IJob? NewJob = await TryUpdateAsync(JobDocument, Updates);
-				if (NewJob != null)
+				IJob? newJob = await TryUpdateAsync(jobDocument, updates);
+				if (newJob != null)
 				{
-					return NewJob;
+					return newJob;
 				}
 
-				Job = await GetAsync(Job.Id);
+				job = await GetAsync(job.Id);
 			}
-			return Job;
+			return job;
 		}
 
 		/// <inheritdoc/>
-		public async Task<IJob?> SkipBatchAsync(IJob? Job, SubResourceId BatchId, IGraph Graph, JobStepBatchError Reason)
+		public async Task<IJob?> SkipBatchAsync(IJob? job, SubResourceId batchId, IGraph graph, JobStepBatchError reason)
 		{
-			while (Job != null)
+			while (job != null)
 			{
-				JobDocument JobDocument = (JobDocument)Job;
+				JobDocument jobDocument = (JobDocument)job;
 
-				JobStepBatchDocument? Batch = JobDocument.Batches.FirstOrDefault(x => x.Id == BatchId);
-				if (Batch == null)
+				JobStepBatchDocument? batch = jobDocument.Batches.FirstOrDefault(x => x.Id == batchId);
+				if (batch == null)
 				{
-					return Job;
+					return job;
 				}
 
-				Batch.State = JobStepBatchState.Complete;
-				Batch.Error = Reason;
-				Batch.FinishTime = DateTimeOffset.UtcNow;
+				batch.State = JobStepBatchState.Complete;
+				batch.Error = reason;
+				batch.FinishTime = DateTimeOffset.UtcNow;
 
-				foreach (JobStepDocument Step in Batch.Steps)
+				foreach (JobStepDocument step in batch.Steps)
 				{
-					if (Step.State != JobStepState.Skipped)
+					if (step.State != JobStepState.Skipped)
 					{
-						Step.State = JobStepState.Skipped;
-						Step.Outcome = JobStepOutcome.Failure;
+						step.State = JobStepState.Skipped;
+						step.Outcome = JobStepOutcome.Failure;
 					}
 				}
 
-				List<UpdateDefinition<JobDocument>> Updates = new List<UpdateDefinition<JobDocument>>();
-				UpdateBatches(JobDocument, Graph, Updates, Logger);
+				List<UpdateDefinition<JobDocument>> updates = new List<UpdateDefinition<JobDocument>>();
+				UpdateBatches(jobDocument, graph, updates, _logger);
 
-				IJob? NewJob = await TryUpdateAsync(JobDocument, Updates);
-				if(NewJob != null)
+				IJob? newJob = await TryUpdateAsync(jobDocument, updates);
+				if(newJob != null)
 				{
-					return NewJob;
+					return newJob;
 				}
 
-				Job = await GetAsync(Job.Id);
+				job = await GetAsync(job.Id);
 			}
-			return Job;
+			return job;
 		}
 
 		/// <inheritdoc/>
-		public async Task<IJob?> TryAssignLeaseAsync(IJob Job, int BatchIdx, PoolId PoolId, AgentId AgentId, SessionId SessionId, LeaseId LeaseId, LogId LogId)
+		public async Task<IJob?> TryAssignLeaseAsync(IJob job, int batchIdx, PoolId poolId, AgentId agentId, SessionId sessionId, LeaseId leaseId, LogId logId)
 		{
 			// Try to update the job with this agent id
-			UpdateDefinitionBuilder<JobDocument> UpdateBuilder = Builders<JobDocument>.Update;
-			List<UpdateDefinition<JobDocument>> Updates = new List<UpdateDefinition<JobDocument>>();
+			UpdateDefinitionBuilder<JobDocument> updateBuilder = Builders<JobDocument>.Update;
+			List<UpdateDefinition<JobDocument>> updates = new List<UpdateDefinition<JobDocument>>();
 
-			Updates.Add(UpdateBuilder.Set(x => x.Batches[BatchIdx].PoolId, PoolId));
-			Updates.Add(UpdateBuilder.Set(x => x.Batches[BatchIdx].AgentId, AgentId));
-			Updates.Add(UpdateBuilder.Set(x => x.Batches[BatchIdx].SessionId, SessionId));
-			Updates.Add(UpdateBuilder.Set(x => x.Batches[BatchIdx].LeaseId, LeaseId));
-			Updates.Add(UpdateBuilder.Set(x => x.Batches[BatchIdx].LogId, LogId));
+			updates.Add(updateBuilder.Set(x => x.Batches[batchIdx].PoolId, poolId));
+			updates.Add(updateBuilder.Set(x => x.Batches[batchIdx].AgentId, agentId));
+			updates.Add(updateBuilder.Set(x => x.Batches[batchIdx].SessionId, sessionId));
+			updates.Add(updateBuilder.Set(x => x.Batches[batchIdx].LeaseId, leaseId));
+			updates.Add(updateBuilder.Set(x => x.Batches[batchIdx].LogId, logId));
 
 			// Extra logs for catching why session IDs sometimes doesn't match in prod. Resulting in PermissionDenied errors.
-			if (BatchIdx < Job.Batches.Count && Job.Batches[BatchIdx].SessionId != null)
+			if (batchIdx < job.Batches.Count && job.Batches[batchIdx].SessionId != null)
 			{
-				string CurrentSessionId = Job.Batches[BatchIdx].SessionId!.Value.ToString();
-				Logger.LogError("Attempt to replace current session ID {CurrSessionId} with {NewSessionId} for batch {JobId}:{BatchId}", CurrentSessionId, SessionId.ToString(), Job.Id.ToString(), Job.Batches[BatchIdx].Id);
+				string currentSessionId = job.Batches[batchIdx].SessionId!.Value.ToString();
+				_logger.LogError("Attempt to replace current session ID {CurrSessionId} with {NewSessionId} for batch {JobId}:{BatchId}", currentSessionId, sessionId.ToString(), job.Id.ToString(), job.Batches[batchIdx].Id);
 				return null;
 			}
 
-			JobDocument JobDocument = Clone((JobDocument)Job);
-			if (await TryUpdateAsync(JobDocument, Updates) == null)
+			JobDocument jobDocument = Clone((JobDocument)job);
+			if (await TryUpdateAsync(jobDocument, updates) == null)
 			{
 				return null;
 			}
 
-			JobStepBatchDocument Batch = JobDocument.Batches[BatchIdx];
-			Batch.AgentId = AgentId;
-			Batch.SessionId = SessionId;
-			Batch.LeaseId = LeaseId;
-			Batch.LogId = LogId;
+			JobStepBatchDocument batch = jobDocument.Batches[batchIdx];
+			batch.AgentId = agentId;
+			batch.SessionId = sessionId;
+			batch.LeaseId = leaseId;
+			batch.LogId = logId;
 
-			return JobDocument;
+			return jobDocument;
 		}
 
 		/// <inheritdoc/>
-		public async Task<IJob?> TryCancelLeaseAsync(IJob Job, int BatchIdx)
+		public async Task<IJob?> TryCancelLeaseAsync(IJob job, int batchIdx)
 		{
-			Logger.LogDebug("Cancelling lease {LeaseId} for agent {AgentId}", Job.Batches[BatchIdx].LeaseId, Job.Batches[BatchIdx].AgentId);
+			_logger.LogDebug("Cancelling lease {LeaseId} for agent {AgentId}", job.Batches[batchIdx].LeaseId, job.Batches[batchIdx].AgentId);
 
-			JobDocument JobDocument = Clone((JobDocument)Job);
+			JobDocument jobDocument = Clone((JobDocument)job);
 
-			UpdateDefinition<JobDocument> Update = Builders<JobDocument>.Update.Unset(x => x.Batches[BatchIdx].AgentId).Unset(x => x.Batches[BatchIdx].SessionId).Unset(x => x.Batches[BatchIdx].LeaseId);
-			if (await TryUpdateAsync(JobDocument, Update) == null)
+			UpdateDefinition<JobDocument> update = Builders<JobDocument>.Update.Unset(x => x.Batches[batchIdx].AgentId).Unset(x => x.Batches[batchIdx].SessionId).Unset(x => x.Batches[batchIdx].LeaseId);
+			if (await TryUpdateAsync(jobDocument, update) == null)
 			{
 				return null;
 			}
 
-			JobStepBatchDocument Batch = JobDocument.Batches[BatchIdx];
-			Batch.AgentId = null;
-			Batch.SessionId = null;
-			Batch.LeaseId = null;
+			JobStepBatchDocument batch = jobDocument.Batches[batchIdx];
+			batch.AgentId = null;
+			batch.SessionId = null;
+			batch.LeaseId = null;
 
-			return JobDocument;
+			return jobDocument;
 		}
 
 		/// <inheritdoc/>
-		public Task<IJob?> TryFailBatchAsync(IJob Job, int BatchIdx, IGraph Graph, JobStepBatchError Error)
+		public Task<IJob?> TryFailBatchAsync(IJob job, int batchIdx, IGraph graph, JobStepBatchError error)
 		{
-			JobDocument JobDocument = Clone((JobDocument)Job);
-			JobStepBatchDocument Batch = JobDocument.Batches[BatchIdx];
-			Logger.LogDebug("Failing batch {JobId}:{BatchId} with error {Error}", Job.Id, Batch.Id, Error);
+			JobDocument jobDocument = Clone((JobDocument)job);
+			JobStepBatchDocument batch = jobDocument.Batches[batchIdx];
+			_logger.LogDebug("Failing batch {JobId}:{BatchId} with error {Error}", job.Id, batch.Id, error);
 
-			UpdateDefinitionBuilder<JobDocument> UpdateBuilder = new UpdateDefinitionBuilder<JobDocument>();
-			List<UpdateDefinition<JobDocument>> Updates = new List<UpdateDefinition<JobDocument>>();
+			UpdateDefinitionBuilder<JobDocument> updateBuilder = new UpdateDefinitionBuilder<JobDocument>();
+			List<UpdateDefinition<JobDocument>> updates = new List<UpdateDefinition<JobDocument>>();
 
-			if (Batch.State != JobStepBatchState.Complete)
+			if (batch.State != JobStepBatchState.Complete)
 			{
-				Batch.State = JobStepBatchState.Complete;
-				Updates.Add(UpdateBuilder.Set(x => x.Batches[BatchIdx].State, Batch.State));
+				batch.State = JobStepBatchState.Complete;
+				updates.Add(updateBuilder.Set(x => x.Batches[batchIdx].State, batch.State));
 
-				Batch.Error = Error;
-				Updates.Add(UpdateBuilder.Set(x => x.Batches[BatchIdx].Error, Batch.Error));
+				batch.Error = error;
+				updates.Add(updateBuilder.Set(x => x.Batches[batchIdx].Error, batch.Error));
 
-				Batch.FinishTime = DateTimeOffset.UtcNow;
-				Updates.Add(UpdateBuilder.Set(x => x.Batches[BatchIdx].FinishTime, Batch.FinishTime));
+				batch.FinishTime = DateTimeOffset.UtcNow;
+				updates.Add(updateBuilder.Set(x => x.Batches[batchIdx].FinishTime, batch.FinishTime));
 			}
 
-			for (int StepIdx = 0; StepIdx < Batch.Steps.Count; StepIdx++)
+			for (int stepIdx = 0; stepIdx < batch.Steps.Count; stepIdx++)
 			{
-				JobStepDocument Step = Batch.Steps[StepIdx];
-				if (Step.State == JobStepState.Running)
+				JobStepDocument step = batch.Steps[stepIdx];
+				if (step.State == JobStepState.Running)
 				{
-					int StepIdxCopy = StepIdx;
+					int stepIdxCopy = stepIdx;
 
-					Step.State = JobStepState.Aborted;
-					Updates.Add(UpdateBuilder.Set(x => x.Batches[BatchIdx].Steps[StepIdxCopy].State, Step.State));
+					step.State = JobStepState.Aborted;
+					updates.Add(updateBuilder.Set(x => x.Batches[batchIdx].Steps[stepIdxCopy].State, step.State));
 
-					Step.Outcome = JobStepOutcome.Failure;
-					Updates.Add(UpdateBuilder.Set(x => x.Batches[BatchIdx].Steps[StepIdxCopy].Outcome, Step.Outcome));
+					step.Outcome = JobStepOutcome.Failure;
+					updates.Add(updateBuilder.Set(x => x.Batches[batchIdx].Steps[stepIdxCopy].Outcome, step.Outcome));
 
-					Step.FinishTime = Clock.UtcNow;
-					Updates.Add(UpdateBuilder.Set(x => x.Batches[BatchIdx].Steps[StepIdxCopy].FinishTime, Step.FinishTime));
+					step.FinishTime = _clock.UtcNow;
+					updates.Add(updateBuilder.Set(x => x.Batches[batchIdx].Steps[stepIdxCopy].FinishTime, step.FinishTime));
 				}
-				else if (Step.State == JobStepState.Ready)
+				else if (step.State == JobStepState.Ready)
 				{
-					int StepIdxCopy = StepIdx;
+					int stepIdxCopy = stepIdx;
 
-					Step.State = JobStepState.Skipped;
-					Updates.Add(UpdateBuilder.Set(x => x.Batches[BatchIdx].Steps[StepIdxCopy].State, Step.State));
+					step.State = JobStepState.Skipped;
+					updates.Add(updateBuilder.Set(x => x.Batches[batchIdx].Steps[stepIdxCopy].State, step.State));
 
-					Step.Outcome = JobStepOutcome.Failure;
-					Updates.Add(UpdateBuilder.Set(x => x.Batches[BatchIdx].Steps[StepIdxCopy].Outcome, Step.Outcome));
+					step.Outcome = JobStepOutcome.Failure;
+					updates.Add(updateBuilder.Set(x => x.Batches[batchIdx].Steps[stepIdxCopy].Outcome, step.Outcome));
 				}
 			}
 
-			RefreshDependentJobSteps(JobDocument, Graph, Updates, Logger);
-			RefreshJobPriority(JobDocument, Updates);
+			RefreshDependentJobSteps(jobDocument, graph, updates, _logger);
+			RefreshJobPriority(jobDocument, updates);
 
-			return TryUpdateAsync(JobDocument, UpdateBuilder.Combine(Updates));
+			return TryUpdateAsync(jobDocument, updateBuilder.Combine(updates));
 		}
 
 		/// <summary>
 		/// Populate the list of batches for a job
 		/// </summary>
-		/// <param name="Job">The job to update</param>
-		/// <param name="Graph">The graph definition for this job</param>
-		/// <param name="Logger">Logger for output messages</param>
-		private static void CreateBatches(JobDocument Job, IGraph Graph, ILogger Logger)
+		/// <param name="job">The job to update</param>
+		/// <param name="graph">The graph definition for this job</param>
+		/// <param name="logger">Logger for output messages</param>
+		private static void CreateBatches(JobDocument job, IGraph graph, ILogger logger)
 		{
-			UpdateBatches(Job, Graph, new List<UpdateDefinition<JobDocument>>(), Logger);
+			UpdateBatches(job, graph, new List<UpdateDefinition<JobDocument>>(), logger);
 		}
 
 		/// <summary>
 		/// Updates the list of batches for a job
 		/// </summary>
-		/// <param name="Job">Job to update</param>
-		/// <param name="Graph">Graph definition for this job</param>
-		/// <param name="Updates">List of updates for the job</param>
-		/// <param name="Logger">Logger for updates</param>
-		static void UpdateBatches(JobDocument Job, IGraph Graph, List<UpdateDefinition<JobDocument>> Updates, ILogger Logger)
+		/// <param name="job">Job to update</param>
+		/// <param name="graph">Graph definition for this job</param>
+		/// <param name="updates">List of updates for the job</param>
+		/// <param name="logger">Logger for updates</param>
+		static void UpdateBatches(JobDocument job, IGraph graph, List<UpdateDefinition<JobDocument>> updates, ILogger logger)
 		{
-			UpdateDefinitionBuilder<JobDocument> UpdateBuilder = Builders<JobDocument>.Update;
+			UpdateDefinitionBuilder<JobDocument> updateBuilder = Builders<JobDocument>.Update;
 
 			// Update the list of batches
-			CreateOrUpdateBatches(Job, Graph);
-			Updates.Add(UpdateBuilder.Set(x => x.Batches, Job.Batches));
-			Updates.Add(UpdateBuilder.Set(x => x.NextSubResourceId, Job.NextSubResourceId));
+			CreateOrUpdateBatches(job, graph);
+			updates.Add(updateBuilder.Set(x => x.Batches, job.Batches));
+			updates.Add(updateBuilder.Set(x => x.NextSubResourceId, job.NextSubResourceId));
 
 			// Update all the dependencies
-			RefreshDependentJobSteps(Job, Graph, new List<UpdateDefinition<JobDocument>>(), Logger);
-			RefreshJobPriority(Job, Updates);
+			RefreshDependentJobSteps(job, graph, new List<UpdateDefinition<JobDocument>>(), logger);
+			RefreshJobPriority(job, updates);
 		}
 
 		/// <summary>
 		/// Update the jobsteps for the given node graph 
 		/// </summary>
-		/// <param name="Job">The job to update</param>
-		/// <param name="Graph">The graph for this job</param>
-		private static void CreateOrUpdateBatches(JobDocument Job, IGraph Graph)
+		/// <param name="job">The job to update</param>
+		/// <param name="graph">The graph for this job</param>
+		private static void CreateOrUpdateBatches(JobDocument job, IGraph graph)
 		{
 			// Find the priorities of each node, incorporating all the per-step overrides
-			Dictionary<INode, Priority> NodePriorities = new Dictionary<INode, Priority>();
-			foreach (INodeGroup Group in Graph.Groups)
+			Dictionary<INode, Priority> nodePriorities = new Dictionary<INode, Priority>();
+			foreach (INodeGroup group in graph.Groups)
 			{
-				foreach (INode Node in Group.Nodes)
+				foreach (INode node in group.Nodes)
 				{
-					NodePriorities[Node] = Node.Priority;
+					nodePriorities[node] = node.Priority;
 				}
 			}
-			foreach (JobStepBatchDocument Batch in Job.Batches)
+			foreach (JobStepBatchDocument batch in job.Batches)
 			{
-				INodeGroup Group = Graph.Groups[Batch.GroupIdx];
-				foreach (JobStepDocument Step in Batch.Steps)
+				INodeGroup group = graph.Groups[batch.GroupIdx];
+				foreach (JobStepDocument step in batch.Steps)
 				{
-					if (Step.Priority != null)
+					if (step.Priority != null)
 					{
-						INode Node = Group.Nodes[Step.NodeIdx];
-						NodePriorities[Node] = Step.Priority.Value;
+						INode node = group.Nodes[step.NodeIdx];
+						nodePriorities[node] = step.Priority.Value;
 					}
 				}
 			}
 
 			// Remove any steps and batches that haven't started yet
-			foreach (JobStepBatchDocument Batch in Job.Batches)
+			foreach (JobStepBatchDocument batch in job.Batches)
 			{
-				Batch.Steps.RemoveAll(x => x.State == JobStepState.Waiting || x.State == JobStepState.Ready);
+				batch.Steps.RemoveAll(x => x.State == JobStepState.Waiting || x.State == JobStepState.Ready);
 			}
 
 			// Remove any skipped nodes whose skipped state is no longer valid
-			HashSet<INode> FailedNodes = new HashSet<INode>();
-			foreach (JobStepBatchDocument Batch in Job.Batches)
+			HashSet<INode> failedNodes = new HashSet<INode>();
+			foreach (JobStepBatchDocument batch in job.Batches)
 			{
-				INodeGroup Group = Graph.Groups[Batch.GroupIdx];
-				foreach (JobStepDocument Step in Batch.Steps)
+				INodeGroup group = graph.Groups[batch.GroupIdx];
+				foreach (JobStepDocument step in batch.Steps)
 				{
-					INode Node = Group.Nodes[Step.NodeIdx];
-					if (Step.Retry)
+					INode node = group.Nodes[step.NodeIdx];
+					if (step.Retry)
 					{
-						FailedNodes.Remove(Node);
+						failedNodes.Remove(node);
 					}
-					else if (Step.State == JobStepState.Skipped && (Node.InputDependencies.Any(x => FailedNodes.Contains(Graph.GetNode(x))) || !CanRetryNode(Job, Batch.GroupIdx, Step.NodeIdx)))
+					else if (step.State == JobStepState.Skipped && (node.InputDependencies.Any(x => failedNodes.Contains(graph.GetNode(x))) || !CanRetryNode(job, batch.GroupIdx, step.NodeIdx)))
 					{
-						FailedNodes.Add(Node);
+						failedNodes.Add(node);
 					}
-					else if (Step.Outcome == JobStepOutcome.Failure)
+					else if (step.Outcome == JobStepOutcome.Failure)
 					{
-						FailedNodes.Add(Node);
+						failedNodes.Add(node);
 					}
 					else
 					{
-						FailedNodes.Remove(Node);
+						failedNodes.Remove(node);
 					}
 				}
-				Batch.Steps.RemoveAll(x => x.State == JobStepState.Skipped && !FailedNodes.Contains(Group.Nodes[x.NodeIdx]));
+				batch.Steps.RemoveAll(x => x.State == JobStepState.Skipped && !failedNodes.Contains(group.Nodes[x.NodeIdx]));
 			}
 
 			// Remove any batches which are now empty
-			Job.Batches.RemoveAll(x => x.Steps.Count == 0 && x.Error == JobStepBatchError.None);
+			job.Batches.RemoveAll(x => x.Steps.Count == 0 && x.Error == JobStepBatchError.None);
 
 			// Find all the targets in this job
-			HashSet<string> Targets = new HashSet<string>(StringComparer.OrdinalIgnoreCase);
-			if (Job.AbortedByUserId == null)
+			HashSet<string> targets = new HashSet<string>(StringComparer.OrdinalIgnoreCase);
+			if (job.AbortedByUserId == null)
 			{
-				foreach (string Argument in Job.Arguments)
+				foreach (string argument in job.Arguments)
 				{
-					if (Argument.StartsWith(IJob.TargetArgumentPrefix, StringComparison.OrdinalIgnoreCase))
+					if (argument.StartsWith(IJob.TargetArgumentPrefix, StringComparison.OrdinalIgnoreCase))
 					{
-						Targets.UnionWith(Argument.Substring(IJob.TargetArgumentPrefix.Length).Split(';'));
+						targets.UnionWith(argument.Substring(IJob.TargetArgumentPrefix.Length).Split(';'));
 					}
 				}
-				Targets.Add(IJob.SetupNodeName);
+				targets.Add(IJob.SetupNodeName);
 			}
 
 			// Add all the referenced aggregates
-			HashSet<INode> NewNodesToExecute = new HashSet<INode>();
-			foreach (IAggregate Aggregate in Graph.Aggregates)
+			HashSet<INode> newNodesToExecute = new HashSet<INode>();
+			foreach (IAggregate aggregate in graph.Aggregates)
 			{
-				if (Targets.Contains(Aggregate.Name))
+				if (targets.Contains(aggregate.Name))
 				{
-					NewNodesToExecute.UnionWith(Aggregate.Nodes.Select(x => Graph.GetNode(x)));
+					newNodesToExecute.UnionWith(aggregate.Nodes.Select(x => graph.GetNode(x)));
 				}
 			}
 
 			// Add any individual nodes
-			foreach (INode Node in Graph.Groups.SelectMany(x => x.Nodes))
+			foreach (INode node in graph.Groups.SelectMany(x => x.Nodes))
 			{
-				if (Targets.Contains(Node.Name))
+				if (targets.Contains(node.Name))
 				{
-					NewNodesToExecute.Add(Node);
+					newNodesToExecute.Add(node);
 				}
 			}
 
 			// Also add any dependencies of these nodes
-			for (int GroupIdx = Graph.Groups.Count - 1; GroupIdx >= 0; GroupIdx--)
+			for (int groupIdx = graph.Groups.Count - 1; groupIdx >= 0; groupIdx--)
 			{
-				INodeGroup Group = Graph.Groups[GroupIdx];
-				for (int NodeIdx = Group.Nodes.Count - 1; NodeIdx >= 0; NodeIdx--)
+				INodeGroup group = graph.Groups[groupIdx];
+				for (int nodeIdx = group.Nodes.Count - 1; nodeIdx >= 0; nodeIdx--)
 				{
-					INode Node = Group.Nodes[NodeIdx];
-					if (NewNodesToExecute.Contains(Node))
+					INode node = group.Nodes[nodeIdx];
+					if (newNodesToExecute.Contains(node))
 					{
-						foreach (NodeRef Dependency in Node.InputDependencies)
+						foreach (NodeRef dependency in node.InputDependencies)
 						{
-							NewNodesToExecute.Add(Graph.GetNode(Dependency));
+							newNodesToExecute.Add(graph.GetNode(dependency));
 						}
 					}
 				}
 			}
 
 			// Cancel any batches which are still running but are no longer required
-			foreach (JobStepBatchDocument Batch in Job.Batches)
+			foreach (JobStepBatchDocument batch in job.Batches)
 			{
-				if (Batch.State == JobStepBatchState.Starting || Batch.State == JobStepBatchState.Running)
+				if (batch.State == JobStepBatchState.Starting || batch.State == JobStepBatchState.Running)
 				{
-					INodeGroup Group = Graph.Groups[Batch.GroupIdx];
-					if (!Batch.Steps.Any(x => NewNodesToExecute.Contains(Group.Nodes[x.NodeIdx])))
+					INodeGroup group = graph.Groups[batch.GroupIdx];
+					if (!batch.Steps.Any(x => newNodesToExecute.Contains(group.Nodes[x.NodeIdx])))
 					{
-						Batch.Error = JobStepBatchError.Cancelled;
+						batch.Error = JobStepBatchError.Cancelled;
 					}
 				}
 			}
 
 			// Remove all the nodes which have already succeeded
-			foreach (JobStepBatchDocument Batch in Job.Batches)
+			foreach (JobStepBatchDocument batch in job.Batches)
 			{
-				foreach (IJobStep Step in Batch.Steps)
+				foreach (IJobStep step in batch.Steps)
 				{
-					if ((Step.State == JobStepState.Running && Step.RetriedByUserId == null) || (Step.State == JobStepState.Completed && Step.RetriedByUserId == null) || (Step.State == JobStepState.Aborted && Step.RetriedByUserId == null) || Step.State == JobStepState.Skipped)
+					if ((step.State == JobStepState.Running && step.RetriedByUserId == null) || (step.State == JobStepState.Completed && step.RetriedByUserId == null) || (step.State == JobStepState.Aborted && step.RetriedByUserId == null) || step.State == JobStepState.Skipped)
 					{
-						NewNodesToExecute.Remove(Graph.Groups[Batch.GroupIdx].Nodes[Step.NodeIdx]);
+						newNodesToExecute.Remove(graph.Groups[batch.GroupIdx].Nodes[step.NodeIdx]);
 					}
 				}
 			}
 
 			// Re-add all the nodes that have input dependencies in the same group.
-			for (int GroupIdx = Graph.Groups.Count - 1; GroupIdx >= 0; GroupIdx--)
+			for (int groupIdx = graph.Groups.Count - 1; groupIdx >= 0; groupIdx--)
 			{
-				INodeGroup Group = Graph.Groups[GroupIdx];
-				for (int NodeIdx = Group.Nodes.Count - 1; NodeIdx >= 0; NodeIdx--)
+				INodeGroup group = graph.Groups[groupIdx];
+				for (int nodeIdx = group.Nodes.Count - 1; nodeIdx >= 0; nodeIdx--)
 				{
-					INode Node = Group.Nodes[NodeIdx];
-					if (NewNodesToExecute.Contains(Node))
+					INode node = group.Nodes[nodeIdx];
+					if (newNodesToExecute.Contains(node))
 					{
-						foreach (NodeRef Dependency in Node.InputDependencies)
+						foreach (NodeRef dependency in node.InputDependencies)
 						{
-							if (Dependency.GroupIdx == GroupIdx)
+							if (dependency.GroupIdx == groupIdx)
 							{
-								NewNodesToExecute.Add(Group.Nodes[Dependency.NodeIdx]);
+								newNodesToExecute.Add(group.Nodes[dependency.NodeIdx]);
 							}
 						}
 					}
@@ -1443,45 +1438,45 @@ namespace Horde.Build.Collections.Impl
 			}
 
 			// Build a list of nodes which are currently set to be executed
-			HashSet<INode> ExistingNodesToExecute = new HashSet<INode>();
-			foreach (JobStepBatchDocument Batch in Job.Batches)
+			HashSet<INode> existingNodesToExecute = new HashSet<INode>();
+			foreach (JobStepBatchDocument batch in job.Batches)
 			{
-				INodeGroup Group = Graph.Groups[Batch.GroupIdx];
-				foreach (IJobStep Step in Batch.Steps)
+				INodeGroup group = graph.Groups[batch.GroupIdx];
+				foreach (IJobStep step in batch.Steps)
 				{
-					INode Node = Group.Nodes[Step.NodeIdx];
-					ExistingNodesToExecute.Add(Node);
+					INode node = group.Nodes[step.NodeIdx];
+					existingNodesToExecute.Add(node);
 				}
 			}
 
 			// Figure out the existing batch for each group
-			JobStepBatchDocument?[] AppendToBatches = new JobStepBatchDocument?[Graph.Groups.Count];
-			foreach (JobStepBatchDocument Batch in Job.Batches)
+			JobStepBatchDocument?[] appendToBatches = new JobStepBatchDocument?[graph.Groups.Count];
+			foreach (JobStepBatchDocument batch in job.Batches)
 			{
-				if (Batch.CanBeAppendedTo())
+				if (batch.CanBeAppendedTo())
 				{
-					INodeGroup Group = Graph.Groups[Batch.GroupIdx];
-					INode FirstNode = Group.Nodes[Batch.Steps[0].NodeIdx];
-					AppendToBatches[Batch.GroupIdx] = Batch;
+					INodeGroup group = graph.Groups[batch.GroupIdx];
+					INode firstNode = group.Nodes[batch.Steps[0].NodeIdx];
+					appendToBatches[batch.GroupIdx] = batch;
 				}
 			}
 
 			// Invalidate all the entries for groups where we're too late to append new entries (ie. we need to execute an earlier node that wasn't executed previously)
-			for (int GroupIdx = 0; GroupIdx < Graph.Groups.Count; GroupIdx++)
+			for (int groupIdx = 0; groupIdx < graph.Groups.Count; groupIdx++)
 			{
-				INodeGroup Group = Graph.Groups[GroupIdx];
-				for (int NodeIdx = 0; NodeIdx < Group.Nodes.Count; NodeIdx++)
+				INodeGroup group = graph.Groups[groupIdx];
+				for (int nodeIdx = 0; nodeIdx < group.Nodes.Count; nodeIdx++)
 				{
-					INode Node = Group.Nodes[NodeIdx];
-					if (NewNodesToExecute.Contains(Node) && !ExistingNodesToExecute.Contains(Node))
+					INode node = group.Nodes[nodeIdx];
+					if (newNodesToExecute.Contains(node) && !existingNodesToExecute.Contains(node))
 					{
-						IJobStepBatch? Batch = AppendToBatches[GroupIdx];
-						if (Batch != null)
+						IJobStepBatch? batch = appendToBatches[groupIdx];
+						if (batch != null)
 						{
-							IJobStep LastStep = Batch.Steps[Batch.Steps.Count - 1];
-							if (NodeIdx < LastStep.NodeIdx)
+							IJobStep lastStep = batch.Steps[batch.Steps.Count - 1];
+							if (nodeIdx < lastStep.NodeIdx)
 							{
-								AppendToBatches[GroupIdx] = null;
+								appendToBatches[groupIdx] = null;
 							}
 						}
 					}
@@ -1489,82 +1484,82 @@ namespace Horde.Build.Collections.Impl
 			}
 
 			// Create all the new jobsteps
-			for (int GroupIdx = 0; GroupIdx < Graph.Groups.Count; GroupIdx++)
+			for (int groupIdx = 0; groupIdx < graph.Groups.Count; groupIdx++)
 			{
-				INodeGroup Group = Graph.Groups[GroupIdx];
-				for (int NodeIdx = 0; NodeIdx < Group.Nodes.Count; NodeIdx++)
+				INodeGroup group = graph.Groups[groupIdx];
+				for (int nodeIdx = 0; nodeIdx < group.Nodes.Count; nodeIdx++)
 				{
-					INode Node = Group.Nodes[NodeIdx];
-					if (NewNodesToExecute.Contains(Node))
+					INode node = group.Nodes[nodeIdx];
+					if (newNodesToExecute.Contains(node))
 					{
-						JobStepBatchDocument? Batch = AppendToBatches[GroupIdx];
-						if (Batch == null)
+						JobStepBatchDocument? batch = appendToBatches[groupIdx];
+						if (batch == null)
 						{
-							Job.NextSubResourceId = Job.NextSubResourceId.Next();
+							job.NextSubResourceId = job.NextSubResourceId.Next();
 
-							Batch = new JobStepBatchDocument(Job.NextSubResourceId, GroupIdx);
-							Job.Batches.Add(Batch);
+							batch = new JobStepBatchDocument(job.NextSubResourceId, groupIdx);
+							job.Batches.Add(batch);
 
-							AppendToBatches[GroupIdx] = Batch;
+							appendToBatches[groupIdx] = batch;
 						}
-						if (Batch.Steps.Count == 0 || NodeIdx > Batch.Steps[Batch.Steps.Count - 1].NodeIdx)
+						if (batch.Steps.Count == 0 || nodeIdx > batch.Steps[^1].NodeIdx)
 						{
-							Job.NextSubResourceId = Job.NextSubResourceId.Next();
+							job.NextSubResourceId = job.NextSubResourceId.Next();
 
-							JobStepDocument Step = new JobStepDocument(Job.NextSubResourceId, NodeIdx);
-							Batch.Steps.Add(Step);
+							JobStepDocument step = new JobStepDocument(job.NextSubResourceId, nodeIdx);
+							batch.Steps.Add(step);
 						}
 					}
 				}
 			}
 
 			// Find the priority of each node, propagating dependencies from dependent nodes
-			for (int GroupIdx = 0; GroupIdx < Graph.Groups.Count; GroupIdx++)
+			for (int groupIdx = 0; groupIdx < graph.Groups.Count; groupIdx++)
 			{
-				INodeGroup Group = Graph.Groups[GroupIdx];
-				for (int NodeIdx = 0; NodeIdx < Group.Nodes.Count; NodeIdx++)
+				INodeGroup group = graph.Groups[groupIdx];
+				for (int nodeIdx = 0; nodeIdx < group.Nodes.Count; nodeIdx++)
 				{
-					INode Node = Group.Nodes[NodeIdx];
-					Priority NodePriority = NodePriorities[Node];
+					INode node = group.Nodes[nodeIdx];
+					Priority nodePriority = nodePriorities[node];
 
-					foreach (NodeRef DependencyRef in Node.OrderDependencies)
+					foreach (NodeRef dependencyRef in node.OrderDependencies)
 					{
-						INode Dependency = Graph.Groups[DependencyRef.GroupIdx].Nodes[DependencyRef.NodeIdx];
-						if (NodePriorities[Node] > NodePriority)
+						INode dependency = graph.Groups[dependencyRef.GroupIdx].Nodes[dependencyRef.NodeIdx];
+						if (nodePriorities[node] > nodePriority)
 						{
-							NodePriorities[Dependency] = NodePriority;
+							nodePriorities[dependency] = nodePriority;
 						}
 					}
 				}
 			}
-			foreach (JobStepBatchDocument Batch in Job.Batches)
+			foreach (JobStepBatchDocument batch in job.Batches)
 			{
-				if (Batch.Steps.Count > 0)
+				if (batch.Steps.Count > 0)
 				{
-					INodeGroup Group = Graph.Groups[Batch.GroupIdx];
-					Priority NodePriority = Batch.Steps.Max(x => NodePriorities[Group.Nodes[x.NodeIdx]]);
-					Batch.SchedulePriority = ((int)Job.Priority * 10) + (int)NodePriority + 1; // Reserve '0' for none.
+					INodeGroup group = graph.Groups[batch.GroupIdx];
+					Priority nodePriority = batch.Steps.Max(x => nodePriorities[group.Nodes[x.NodeIdx]]);
+					batch.SchedulePriority = ((int)job.Priority * 10) + (int)nodePriority + 1; // Reserve '0' for none.
 				}
 			}
 
 			// Check we're not running a node which doesn't allow retries more than once
-			Dictionary<INode, int> NodeExecutionCount = new Dictionary<INode, int>();
-			foreach (IJobStepBatch Batch in Job.Batches)
+			Dictionary<INode, int> nodeExecutionCount = new Dictionary<INode, int>();
+			foreach (IJobStepBatch batch in job.Batches)
 			{
-				INodeGroup Group = Graph.Groups[Batch.GroupIdx];
-				foreach (IJobStep Step in Batch.Steps)
+				INodeGroup group = graph.Groups[batch.GroupIdx];
+				foreach (IJobStep step in batch.Steps)
 				{
-					INode Node = Group.Nodes[Step.NodeIdx];
+					INode node = group.Nodes[step.NodeIdx];
 
-					int Count;
-					NodeExecutionCount.TryGetValue(Node, out Count);
+					int count;
+					nodeExecutionCount.TryGetValue(node, out count);
 
-					if (!Node.AllowRetry && Count > 0)
+					if (!node.AllowRetry && count > 0)
 					{
-						throw new RetryNotAllowedException(Node.Name);
+						throw new RetryNotAllowedException(node.Name);
 					}
 
-					NodeExecutionCount[Node] = Count + 1;
+					nodeExecutionCount[node] = count + 1;
 				}
 			}
 		}
@@ -1572,99 +1567,99 @@ namespace Horde.Build.Collections.Impl
 		/// <summary>
 		/// Tests whether a node can be retried again
 		/// </summary>
-		static bool CanRetryNode(JobDocument Job, int GroupIdx, int NodeIdx)
+		static bool CanRetryNode(JobDocument job, int groupIdx, int nodeIdx)
 		{
-			return Job.RetriedNodes == null || Job.RetriedNodes.Count(x => x.GroupIdx == GroupIdx && x.NodeIdx == NodeIdx) < MaxRetries;
+			return job.RetriedNodes == null || job.RetriedNodes.Count(x => x.GroupIdx == groupIdx && x.NodeIdx == nodeIdx) < MaxRetries;
 		}
 
 		/// <summary>
 		/// Gets the scheduling priority of this job
 		/// </summary>
-		/// <param name="Job">Job to consider</param>
-		public static int GetSchedulePriority(IJob Job)
+		/// <param name="job">Job to consider</param>
+		public static int GetSchedulePriority(IJob job)
 		{
-			int NewSchedulePriority = 0;
-			foreach (IJobStepBatch Batch in Job.Batches)
+			int newSchedulePriority = 0;
+			foreach (IJobStepBatch batch in job.Batches)
 			{
-				if (Batch.State == JobStepBatchState.Ready)
+				if (batch.State == JobStepBatchState.Ready)
 				{
-					NewSchedulePriority = Math.Max(Batch.SchedulePriority, NewSchedulePriority);
+					newSchedulePriority = Math.Max(batch.SchedulePriority, newSchedulePriority);
 				}
 			}
-			return NewSchedulePriority;
+			return newSchedulePriority;
 		}
 
 		/// <summary>
 		/// Update the state of any jobsteps that are dependent on other jobsteps (eg. transition them from waiting to ready based on other steps completing)
 		/// </summary>
-		/// <param name="Job">The job to update</param>
-		/// <param name="Graph">The graph for this job</param>
-		/// <param name="Updates">List of updates to the job</param>
-		/// <param name="Logger">Logger instance</param>
-		static void RefreshDependentJobSteps(JobDocument Job, IGraph Graph, List<UpdateDefinition<JobDocument>> Updates, ILogger Logger)
+		/// <param name="job">The job to update</param>
+		/// <param name="graph">The graph for this job</param>
+		/// <param name="updates">List of updates to the job</param>
+		/// <param name="logger">Logger instance</param>
+		static void RefreshDependentJobSteps(JobDocument job, IGraph graph, List<UpdateDefinition<JobDocument>> updates, ILogger logger)
 		{
 			// Update the batches
-			UpdateDefinitionBuilder<JobDocument> UpdateBuilder = Builders<JobDocument>.Update;
-			if (Job.Batches != null)
+			UpdateDefinitionBuilder<JobDocument> updateBuilder = Builders<JobDocument>.Update;
+			if (job.Batches != null)
 			{
-				Dictionary<INode, IJobStep> StepForNode = new Dictionary<INode, IJobStep>();
-				for (int LoopBatchIdx = 0; LoopBatchIdx < Job.Batches.Count; LoopBatchIdx++)
+				Dictionary<INode, IJobStep> stepForNode = new Dictionary<INode, IJobStep>();
+				for (int loopBatchIdx = 0; loopBatchIdx < job.Batches.Count; loopBatchIdx++)
 				{
-					int BatchIdx = LoopBatchIdx; // For lambda capture
-					JobStepBatchDocument Batch = Job.Batches[BatchIdx];
+					int batchIdx = loopBatchIdx; // For lambda capture
+					JobStepBatchDocument batch = job.Batches[batchIdx];
 
-					for (int LoopStepIdx = 0; LoopStepIdx < Batch.Steps.Count; LoopStepIdx++)
+					for (int loopStepIdx = 0; loopStepIdx < batch.Steps.Count; loopStepIdx++)
 					{
-						int StepIdx = LoopStepIdx; // For lambda capture
-						JobStepDocument Step = Batch.Steps[StepIdx];
+						int stepIdx = loopStepIdx; // For lambda capture
+						JobStepDocument step = batch.Steps[stepIdx];
 
-						JobStepState NewState = Step.State;
-						JobStepOutcome NewOutcome = Step.Outcome;
+						JobStepState newState = step.State;
+						JobStepOutcome newOutcome = step.Outcome;
 
-						INode Node = Graph.Groups[Batch.GroupIdx].Nodes[Step.NodeIdx];
-						if (NewState == JobStepState.Waiting)
+						INode node = graph.Groups[batch.GroupIdx].Nodes[step.NodeIdx];
+						if (newState == JobStepState.Waiting)
 						{
-							List<IJobStep> Steps = GetDependentSteps(Graph, Node, StepForNode);
-							if (Steps.Any(x => x.IsFailedOrSkipped()))
+							List<IJobStep> steps = GetDependentSteps(graph, node, stepForNode);
+							if (steps.Any(x => x.IsFailedOrSkipped()))
 							{
-								NewState = JobStepState.Skipped;
-								NewOutcome = JobStepOutcome.Failure;
+								newState = JobStepState.Skipped;
+								newOutcome = JobStepOutcome.Failure;
 							}
-							else if (!Steps.Any(x => x.IsPending()))
+							else if (!steps.Any(x => x.IsPending()))
 							{
-								Logger.LogDebug("Transitioning job {JobId}, batch {BatchId}, step {StepId} to ready state ({Dependencies})", Job.Id, Batch.Id, Step.Id, String.Join(", ", Steps.Select(x => x.Id.ToString())));
-								NewState = JobStepState.Ready;
+								logger.LogDebug("Transitioning job {JobId}, batch {BatchId}, step {StepId} to ready state ({Dependencies})", job.Id, batch.Id, step.Id, String.Join(", ", steps.Select(x => x.Id.ToString())));
+								newState = JobStepState.Ready;
 							}
 						}
 
-						if (NewState != Step.State)
+						if (newState != step.State)
 						{
-							Step.State = NewState;
-							Updates.Add(UpdateBuilder.Set(x => x.Batches[BatchIdx].Steps[StepIdx].State, NewState));
+							step.State = newState;
+							updates.Add(updateBuilder.Set(x => x.Batches[batchIdx].Steps[stepIdx].State, newState));
 						}
 
-						if (NewOutcome != Step.Outcome)
+						if (newOutcome != step.Outcome)
 						{
-							Step.Outcome = NewOutcome;
-							Updates.Add(UpdateBuilder.Set(x => x.Batches[BatchIdx].Steps[StepIdx].Outcome, NewOutcome));
+							step.Outcome = newOutcome;
+							updates.Add(updateBuilder.Set(x => x.Batches[batchIdx].Steps[stepIdx].Outcome, newOutcome));
 						}
 
-						StepForNode[Node] = Step;
+						stepForNode[node] = step;
 					}
 
-					if (Batch.State == JobStepBatchState.Waiting || Batch.State == JobStepBatchState.Ready)
+					if (batch.State == JobStepBatchState.Waiting || batch.State == JobStepBatchState.Ready)
 					{
-						DateTime? NewReadyTime;
-						JobStepBatchState NewState = GetBatchState(Job, Graph, Batch, StepForNode, out NewReadyTime);
-						if (Batch.State != NewState)
+						DateTime? newReadyTime;
+						JobStepBatchState newState = GetBatchState(job, graph, batch, stepForNode, out newReadyTime);
+						if (batch.State != newState)
 						{
-							Batch.State = NewState;
-							Updates.Add(UpdateBuilder.Set(x => x.Batches[BatchIdx].State, Batch.State));
+							batch.State = newState;
+							updates.Add(updateBuilder.Set(x => x.Batches[batchIdx].State, batch.State));
 						}
-						if (Batch.ReadyTime != NewReadyTime)
+						if (batch.ReadyTime != newReadyTime)
 						{
-							Batch.ReadyTime = NewReadyTime;
-							Updates.Add(UpdateBuilder.Set(x => x.Batches[BatchIdx].ReadyTime, Batch.ReadyTime));
+							batch.ReadyTime = newReadyTime;
+							updates.Add(updateBuilder.Set(x => x.Batches[batchIdx].ReadyTime, batch.ReadyTime));
 						}
 					}
 				}
@@ -1674,101 +1669,101 @@ namespace Horde.Build.Collections.Impl
 		/// <summary>
 		/// Updates the schedule priority of a job
 		/// </summary>
-		/// <param name="Job"></param>
-		/// <param name="Updates"></param>
-		static void RefreshJobPriority(JobDocument Job, List<UpdateDefinition<JobDocument>> Updates)
+		/// <param name="job"></param>
+		/// <param name="updates"></param>
+		static void RefreshJobPriority(JobDocument job, List<UpdateDefinition<JobDocument>> updates)
 		{
 			// Update the weighted priority for the job
-			int NewSchedulePriority = GetSchedulePriority(Job);
-			if (Job.SchedulePriority != NewSchedulePriority)
+			int newSchedulePriority = GetSchedulePriority(job);
+			if (job.SchedulePriority != newSchedulePriority)
 			{
-				Job.SchedulePriority = NewSchedulePriority;
-				Updates.Add(Builders<JobDocument>.Update.Set(x => x.SchedulePriority, NewSchedulePriority));
+				job.SchedulePriority = newSchedulePriority;
+				updates.Add(Builders<JobDocument>.Update.Set(x => x.SchedulePriority, newSchedulePriority));
 			}
 		}
 
 		/// <summary>
 		/// Gets the steps that a node depends on
 		/// </summary>
-		/// <param name="Graph">The graph for this job</param>
-		/// <param name="Node">The node to test</param>
-		/// <param name="StepForNode">Map of node to step</param>
+		/// <param name="graph">The graph for this job</param>
+		/// <param name="node">The node to test</param>
+		/// <param name="stepForNode">Map of node to step</param>
 		/// <returns></returns>
-		static List<IJobStep> GetDependentSteps(IGraph Graph, INode Node, Dictionary<INode, IJobStep> StepForNode)
+		static List<IJobStep> GetDependentSteps(IGraph graph, INode node, Dictionary<INode, IJobStep> stepForNode)
 		{
-			List<IJobStep> Steps = new List<IJobStep>();
-			foreach (NodeRef OrderDependencyRef in Node.OrderDependencies)
+			List<IJobStep> steps = new List<IJobStep>();
+			foreach (NodeRef orderDependencyRef in node.OrderDependencies)
 			{
-				IJobStep? Step;
-				if (StepForNode.TryGetValue(Graph.GetNode(OrderDependencyRef), out Step))
+				IJobStep? step;
+				if (stepForNode.TryGetValue(graph.GetNode(orderDependencyRef), out step))
 				{
-					Steps.Add(Step);
+					steps.Add(step);
 				}
 			}
-			return Steps;
+			return steps;
 		}
 
 		/// <summary>
 		/// Gets the new state for a batch
 		/// </summary>
-		/// <param name="Job">The job being executed</param>
-		/// <param name="Graph">Graph for the job</param>
-		/// <param name="Batch">List of nodes in the job</param>
-		/// <param name="StepForNode">Array mapping each index to the appropriate step for that node</param>
-		/// <param name="OutReadyTimeUtc">Receives the time at which the batch was ready to execute</param>
+		/// <param name="job">The job being executed</param>
+		/// <param name="graph">Graph for the job</param>
+		/// <param name="batch">List of nodes in the job</param>
+		/// <param name="stepForNode">Array mapping each index to the appropriate step for that node</param>
+		/// <param name="outReadyTimeUtc">Receives the time at which the batch was ready to execute</param>
 		/// <returns>True if the batch is ready, false otherwise</returns>
-		static JobStepBatchState GetBatchState(IJob Job, IGraph Graph, IJobStepBatch Batch, Dictionary<INode, IJobStep> StepForNode, out DateTime? OutReadyTimeUtc)
+		static JobStepBatchState GetBatchState(IJob job, IGraph graph, IJobStepBatch batch, Dictionary<INode, IJobStep> stepForNode, out DateTime? outReadyTimeUtc)
 		{
 			// Check if the batch is already complete
-			if (Batch.Steps.All(x => x.State == JobStepState.Skipped || x.State == JobStepState.Completed || x.State == JobStepState.Aborted))
+			if (batch.Steps.All(x => x.State == JobStepState.Skipped || x.State == JobStepState.Completed || x.State == JobStepState.Aborted))
 			{
-				OutReadyTimeUtc = Batch.ReadyTimeUtc;
+				outReadyTimeUtc = batch.ReadyTimeUtc;
 				return JobStepBatchState.Complete;
 			}
 
 			// Get the dependencies for this batch to start. Some steps may be "after" dependencies that are optional parts of the graph.
-			List<INode> NodeDependencies = Batch.GetStartDependencies(Graph.Groups).ToList();
+			List<INode> nodeDependencies = batch.GetStartDependencies(graph.Groups).ToList();
 
 			// Check if we're still waiting on anything
-			DateTime ReadyTimeUtc = Job.CreateTimeUtc;
-			foreach (INode NodeDependency in NodeDependencies)
+			DateTime readyTimeUtc = job.CreateTimeUtc;
+			foreach (INode nodeDependency in nodeDependencies)
 			{
-				IJobStep? StepDependency;
-				if (StepForNode.TryGetValue(NodeDependency, out StepDependency))
+				IJobStep? stepDependency;
+				if (stepForNode.TryGetValue(nodeDependency, out stepDependency))
 				{
-					if (StepDependency.State != JobStepState.Completed && StepDependency.State != JobStepState.Skipped && StepDependency.State != JobStepState.Aborted)
+					if (stepDependency.State != JobStepState.Completed && stepDependency.State != JobStepState.Skipped && stepDependency.State != JobStepState.Aborted)
 					{
-						OutReadyTimeUtc = null;
+						outReadyTimeUtc = null;
 						return JobStepBatchState.Waiting;
 					}
 
-					if (StepDependency.FinishTimeUtc != null && StepDependency.FinishTimeUtc.Value > ReadyTimeUtc)
+					if (stepDependency.FinishTimeUtc != null && stepDependency.FinishTimeUtc.Value > readyTimeUtc)
 					{
-						ReadyTimeUtc = StepDependency.FinishTimeUtc.Value;
+						readyTimeUtc = stepDependency.FinishTimeUtc.Value;
 					}
 				}
 			}
 
 			// Otherwise return the ready state
-			OutReadyTimeUtc = ReadyTimeUtc;
+			outReadyTimeUtc = readyTimeUtc;
 			return JobStepBatchState.Ready;
 		}
 
 		/// <inheritdoc/>
 		public async Task UpgradeDocumentsAsync()
 		{
-			IAsyncCursor<JobDocument> Cursor = await Jobs.Find(Builders<JobDocument>.Filter.Eq(x => x.UpdateTimeUtc, null)).ToCursorAsync();
+			IAsyncCursor<JobDocument> cursor = await _jobs.Find(Builders<JobDocument>.Filter.Eq(x => x.UpdateTimeUtc, null)).ToCursorAsync();
 
-			int NumUpdated = 0;
-			while (await Cursor.MoveNextAsync())
+			int numUpdated = 0;
+			while (await cursor.MoveNextAsync())
 			{
-				foreach (JobDocument Document in Cursor.Current)
+				foreach (JobDocument document in cursor.Current)
 				{
-					UpdateDefinition<JobDocument> Update = Builders<JobDocument>.Update.Set(x => x.CreateTimeUtc, ((IJob)Document).CreateTimeUtc).Set(x => x.UpdateTimeUtc, ((IJob)Document).UpdateTimeUtc);
-					await Jobs.UpdateOneAsync(Builders<JobDocument>.Filter.Eq(x => x.Id, Document.Id), Update);
-					NumUpdated++;
+					UpdateDefinition<JobDocument> update = Builders<JobDocument>.Update.Set(x => x.CreateTimeUtc, ((IJob)document).CreateTimeUtc).Set(x => x.UpdateTimeUtc, ((IJob)document).UpdateTimeUtc);
+					await _jobs.UpdateOneAsync(Builders<JobDocument>.Filter.Eq(x => x.Id, document.Id), update);
+					numUpdated++;
 				}
-				Logger.LogInformation("Updated {NumDocuments} documents", NumUpdated);
+				_logger.LogInformation("Updated {NumDocuments} documents", numUpdated);
 			}
 		}
 	}
