Horde: Fix coding conventions within Horde.Build.

diff --git a/Engine/Source/Programs/Horde/Horde.Build/Jobs/Impl/JobService.cs b/Engine/Source/Programs/Horde/Horde.Build/Jobs/Impl/JobService.cs
--- a/Engine/Source/Programs/Horde/Horde.Build/Jobs/Impl/JobService.cs
+++ b/Engine/Source/Programs/Horde/Horde.Build/Jobs/Impl/JobService.cs
@@ -1,32 +1,24 @@
-﻿// Copyright Epic Games, Inc. All Rights Reserved.
+// Copyright Epic Games, Inc. All Rights Reserved.
 
+using System;
+using System.Collections.Generic;
+using System.Globalization;
+using System.Linq;
+using System.Security.Claims;
+using System.Threading.Tasks;
 using EpicGames.Core;
 using Horde.Build.Acls;
 using Horde.Build.Api;
 using Horde.Build.Collections;
-using HordeCommon;
 using Horde.Build.Models;
-using HordeCommon.Rpc;
+using Horde.Build.Tasks.Impl;
 using Horde.Build.Utilities;
+using HordeCommon;
 using Microsoft.Extensions.Logging;
-using Microsoft.Extensions.Options;
 using MongoDB.Bson;
-using MongoDB.Bson.Serialization.Serializers;
 using MongoDB.Driver;
-using System;
-using System.Collections.Generic;
-using System.ComponentModel;
-using System.Diagnostics;
-using System.Diagnostics.CodeAnalysis;
-using System.Linq;
-using System.Security.Claims;
-using System.Text.RegularExpressions;
-using System.Threading;
-using System.Threading.Tasks;
-using System.Globalization;
-using Horde.Build.Tasks.Impl;
-using OpenTracing.Util;
 using OpenTracing;
+using OpenTracing.Util;
 
 namespace Horde.Build.Services
 {
@@ -51,11 +43,11 @@ namespace Horde.Build.Services
 		/// <summary>
 		/// Constructor
 		/// </summary>
-		/// <param name="NodeName">Name of the node that does not permit retries</param>
-		public RetryNotAllowedException(string NodeName)
-			: base($"Node '{NodeName}' does not permit retries")
+		/// <param name="nodeName">Name of the node that does not permit retries</param>
+		public RetryNotAllowedException(string nodeName)
+			: base($"Node '{nodeName}' does not permit retries")
 		{
-			this.NodeName = NodeName;
+			NodeName = nodeName;
 		}
 	}
 
@@ -73,83 +65,26 @@ namespace Horde.Build.Services
 	/// <summary>
 	/// Wraps funtionality for manipulating jobs, jobsteps, and jobstep runs
 	/// </summary>
-	[SuppressMessage("Compiler", "CA1054:URI parameters should not be strings")]
 	public class JobService
 	{
-		/// <summary>
-		/// Collection of job documents
-		/// </summary>
-		IJobCollection Jobs;
-
-		/// <summary>
-		/// Collection of graph documents
-		/// </summary>
-		IGraphCollection Graphs;
-
-		/// <summary>
-		/// Collection of agent documents
-		/// </summary>
-		IAgentCollection Agents;
-
-		/// <summary>
-		/// Collection of jobstepref documents
-		/// </summary>
-		IJobStepRefCollection JobStepRefs;
-
-		/// <summary>
-		/// Collection of job timing documents
-		/// </summary>
-		IJobTimingCollection JobTimings;
-
-		/// <summary>
-		/// Collection of users
-		/// </summary>
-		IUserCollection UserCollection;
-
-		/// <summary>
-		/// Collection of notification triggers
-		/// </summary>
-		INotificationTriggerCollection TriggerCollection;
-
-		/// <summary>
-		/// The singleton instance of the queue service
-		/// </summary>
-		JobTaskSource JobTaskSource;
-
-		/// <summary>
-		/// Singleton instance of the stream service
-		/// </summary>
-		StreamService StreamService;
-
-		/// <summary>
-		/// Singleton instance of the template service
-		/// </summary>
-		ITemplateCollection TemplateCollection;
-
-		/// <summary>
-		/// Singleton instance of the issue service
-		/// </summary>
-		IIssueService? IssueService;
-
-		/// <summary>
-		/// The perforce service
-		/// </summary>
-		IPerforceService PerforceService;
-
-		/// <summary>
-		/// The server settings
-		/// </summary>
-		IOptionsMonitor<ServerSettings> Settings;
-
-		/// <summary>
-		/// Log output device
-		/// </summary>
-		ILogger Logger;
+		readonly IJobCollection _jobs;
+		readonly IGraphCollection _graphs;
+		readonly IAgentCollection _agents;
+		readonly IJobStepRefCollection _jobStepRefs;
+		readonly IJobTimingCollection _jobTimings;
+		readonly IUserCollection _userCollection;
+		readonly INotificationTriggerCollection _triggerCollection;
+		readonly JobTaskSource _jobTaskSource;
+		readonly StreamService _streamService;
+		readonly ITemplateCollection _templateCollection;
+		readonly IIssueService? _issueService;
+		readonly IPerforceService _perforceService;
+		readonly ILogger _logger;
 
 		/// <summary>
 		/// Delegate for label update events
 		/// </summary>
-		public delegate void LabelUpdateEvent(IJob Job, IReadOnlyList<(LabelState, LabelOutcome)> OldLabelStates, IReadOnlyList<(LabelState, LabelOutcome)> NewLabelStates);
+		public delegate void LabelUpdateEvent(IJob job, IReadOnlyList<(LabelState, LabelOutcome)> oldLabelStates, IReadOnlyList<(LabelState, LabelOutcome)> newLabelStates);
 
 		/// <summary>
 		/// Event triggered when a label state changes
@@ -159,7 +94,7 @@ namespace Horde.Build.Services
 		/// <summary>
 		/// Delegate for job step complete events
 		/// </summary>
-		public delegate void JobStepCompleteEvent(IJob Job, IGraph Graph, SubResourceId BatchId, SubResourceId StepId);
+		public delegate void JobStepCompleteEvent(IJob job, IGraph graph, SubResourceId batchId, SubResourceId stepId);
 
 		/// <summary>
 		/// Event triggered when a job step completes
@@ -176,174 +111,197 @@ namespace Horde.Build.Services
 		/// <summary>
 		/// Constructor
 		/// </summary>
-		/// <param name="Jobs">The jobs collection</param>
-		/// <param name="Graphs">The graphs collection</param>
-		/// <param name="Agents">The agents collection</param>
-		/// <param name="JobStepRefs">The jobsteprefs collection</param>
-		/// <param name="JobTimings">The job timing document collection</param>
-		/// <param name="UserCollection">User profiles</param>
-		/// <param name="TriggerCollection">Trigger collection</param>
-		/// <param name="JobTaskSource">The queue service</param>
-		/// <param name="StreamService">The stream service</param>
-		/// <param name="TemplateCollection">The template service</param>
-		/// <param name="IssueService">The issue service</param>
-		/// <param name="PerforceService">The perforce service</param>
-		/// <param name="Settings">Settings instance</param>
-		/// <param name="Logger">Log output</param>
-		public JobService(IJobCollection Jobs, IGraphCollection Graphs, IAgentCollection Agents, IJobStepRefCollection JobStepRefs, IJobTimingCollection JobTimings, IUserCollection UserCollection, INotificationTriggerCollection TriggerCollection, JobTaskSource JobTaskSource, StreamService StreamService, ITemplateCollection TemplateCollection, IIssueService IssueService, IPerforceService PerforceService, IOptionsMonitor<ServerSettings> Settings, ILogger<JobService> Logger)
+		/// <param name="jobs">The jobs collection</param>
+		/// <param name="graphs">The graphs collection</param>
+		/// <param name="agents">The agents collection</param>
+		/// <param name="jobStepRefs">The jobsteprefs collection</param>
+		/// <param name="jobTimings">The job timing document collection</param>
+		/// <param name="userCollection">User profiles</param>
+		/// <param name="triggerCollection">Trigger collection</param>
+		/// <param name="jobTaskSource">The queue service</param>
+		/// <param name="streamService">The stream service</param>
+		/// <param name="templateCollection">The template service</param>
+		/// <param name="issueService">The issue service</param>
+		/// <param name="perforceService">The perforce service</param>
+		/// <param name="logger">Log output</param>
+		public JobService(IJobCollection jobs, IGraphCollection graphs, IAgentCollection agents, IJobStepRefCollection jobStepRefs, IJobTimingCollection jobTimings, IUserCollection userCollection, INotificationTriggerCollection triggerCollection, JobTaskSource jobTaskSource, StreamService streamService, ITemplateCollection templateCollection, IIssueService issueService, IPerforceService perforceService, ILogger<JobService> logger)
 		{
-			this.Jobs = Jobs;
-			this.Graphs = Graphs;
-			this.Agents = Agents;
-			this.JobStepRefs = JobStepRefs;
-			this.JobTimings = JobTimings;
-			this.UserCollection = UserCollection;
-			this.TriggerCollection = TriggerCollection;
-			this.JobTaskSource = JobTaskSource;
-			this.StreamService = StreamService;
-			this.TemplateCollection = TemplateCollection;
-			this.IssueService = IssueService;
-			this.PerforceService = PerforceService;
-			this.Settings = Settings;
-			this.Logger = Logger;
-
-			this.JobTaskSource.OnJobScheduled += (Pool, NumAgentsOnline, Job, Graph, BatchId) =>
+			_jobs = jobs;
+			_graphs = graphs;
+			_agents = agents;
+			_jobStepRefs = jobStepRefs;
+			_jobTimings = jobTimings;
+			_userCollection = userCollection;
+			_triggerCollection = triggerCollection;
+			_jobTaskSource = jobTaskSource;
+			_streamService = streamService;
+			_templateCollection = templateCollection;
+			_issueService = issueService;
+			_perforceService = perforceService;
+			_logger = logger;
+
+			_jobTaskSource.OnJobScheduled += (pool, numAgentsOnline, job, graph, batchId) =>
 			{
-				OnJobScheduled?.Invoke(Pool, NumAgentsOnline, Job, Graph, BatchId);
+				OnJobScheduled?.Invoke(pool, numAgentsOnline, job, graph, batchId);
 			};
 		}
 
-#pragma warning disable CA1801
-		static bool ShouldClonePreflightChange(StreamId StreamId)
+		static bool ShouldClonePreflightChange(StreamId streamId)
 		{
 			return false; //			return StreamId == new StreamId("ue5-main");
 		}
-#pragma warning restore CA1801
 
 		/// <summary>
 		/// Creates a new job
 		/// </summary>
-		/// <param name="JobId">A requested job id</param>
-		/// <param name="Stream">The stream that this job belongs to</param>
-		/// <param name="TemplateRefId">Name of the template ref</param>
-		/// <param name="TemplateHash">Template for this job</param>
-		/// <param name="Graph">The graph for the new job</param>
-		/// <param name="Name">Name of the job</param>
-		/// <param name="Change">The change to build</param>
-		/// <param name="CodeChange">The corresponding code changelist</param>
-		/// <param name="PreflightChange">Optional changelist to preflight</param>
-		/// <param name="ClonedPreflightChange">Duplicated preflight change</param>
-		/// <param name="StartedByUserId">Id of the user that started the job</param>
-		/// <param name="Priority">Priority of the job</param>
-		/// <param name="AutoSubmit">Whether to automatically submit the preflighted change on completion</param>
-		/// <param name="UpdateIssues">Whether to update issues when this job completes</param>
-		/// <param name="PromoteIssuesByDefault">Whether to promote issues from this job by default</param>
-		/// <param name="JobTriggers">List of downstream job triggers</param>
-		/// <param name="ShowUgsBadges">Whether to show badges in UGS for this job</param>
-		/// <param name="ShowUgsAlerts">Whether to show alerts in UGS for this job</param>
-		/// <param name="NotificationChannel">Notification Channel for this job</param>
-		/// <param name="NotificationChannelFilter">Notification Channel filter for this job</param>
-		/// <param name="Arguments">Arguments for the job</param>
+		/// <param name="jobId">A requested job id</param>
+		/// <param name="stream">The stream that this job belongs to</param>
+		/// <param name="templateRefId">Name of the template ref</param>
+		/// <param name="templateHash">Template for this job</param>
+		/// <param name="graph">The graph for the new job</param>
+		/// <param name="name">Name of the job</param>
+		/// <param name="change">The change to build</param>
+		/// <param name="codeChange">The corresponding code changelist</param>
+		/// <param name="preflightChange">Optional changelist to preflight</param>
+		/// <param name="clonedPreflightChange">Duplicated preflight change</param>
+		/// <param name="startedByUserId">Id of the user that started the job</param>
+		/// <param name="priority">Priority of the job</param>
+		/// <param name="autoSubmit">Whether to automatically submit the preflighted change on completion</param>
+		/// <param name="updateIssues">Whether to update issues when this job completes</param>
+		/// <param name="promoteIssuesByDefault">Whether to promote issues from this job by default</param>
+		/// <param name="jobTriggers">List of downstream job triggers</param>
+		/// <param name="showUgsBadges">Whether to show badges in UGS for this job</param>
+		/// <param name="showUgsAlerts">Whether to show alerts in UGS for this job</param>
+		/// <param name="notificationChannel">Notification Channel for this job</param>
+		/// <param name="notificationChannelFilter">Notification Channel filter for this job</param>
+		/// <param name="arguments">Arguments for the job</param>
 		/// <returns>Unique id representing the job</returns>
-		public async Task<IJob> CreateJobAsync(JobId? JobId, IStream Stream, TemplateRefId TemplateRefId, ContentHash TemplateHash, IGraph Graph, string Name, int Change, int CodeChange, int? PreflightChange, int? ClonedPreflightChange, UserId? StartedByUserId, Priority? Priority, bool? AutoSubmit, bool? UpdateIssues, bool? PromoteIssuesByDefault, List<ChainedJobTemplate>? JobTriggers, bool ShowUgsBadges, bool ShowUgsAlerts, string? NotificationChannel, string? NotificationChannelFilter, IReadOnlyList<string> Arguments)
+		public async Task<IJob> CreateJobAsync(JobId? jobId, IStream stream, TemplateRefId templateRefId, ContentHash templateHash, IGraph graph, string name, int change, int codeChange, int? preflightChange, int? clonedPreflightChange, UserId? startedByUserId, Priority? priority, bool? autoSubmit, bool? updateIssues, bool? promoteIssuesByDefault, List<ChainedJobTemplate>? jobTriggers, bool showUgsBadges, bool showUgsAlerts, string? notificationChannel, string? notificationChannelFilter, IReadOnlyList<string> arguments)
 		{
-			using IScope TraceScope = GlobalTracer.Instance.BuildSpan("JobService.CreateJobAsync").StartActive();
-			TraceScope.Span.SetTag("JobId", JobId);
-			TraceScope.Span.SetTag("Stream", Stream.Name);
-			TraceScope.Span.SetTag("TemplateRefId", TemplateRefId);
-			TraceScope.Span.SetTag("TemplateHash", TemplateHash);
-			TraceScope.Span.SetTag("GraphId", Graph.Id);
-			TraceScope.Span.SetTag("Name", Name);
-			TraceScope.Span.SetTag("Change", Change);
-			TraceScope.Span.SetTag("CodeChange", CodeChange);
-			TraceScope.Span.SetTag("PreflightChange", PreflightChange);
-			TraceScope.Span.SetTag("ClonedPreflightChange", ClonedPreflightChange);
-			TraceScope.Span.SetTag("StartedByUserId", StartedByUserId.ToString());
-			TraceScope.Span.SetTag("Priority", Priority.ToString());
-			TraceScope.Span.SetTag("ShowUgsBadges", ShowUgsBadges);
-			TraceScope.Span.SetTag("NotificationChannel", NotificationChannel ?? "null");
-			TraceScope.Span.SetTag("NotificationChannelFilter", NotificationChannelFilter ?? "null");
-			TraceScope.Span.SetTag("Arguments", string.Join(',', Arguments));
-			
-			if (AutoSubmit != null) TraceScope.Span.SetTag("AutoSubmit", AutoSubmit.Value);
-			if (UpdateIssues != null) TraceScope.Span.SetTag("UpdateIssues", UpdateIssues.Value);
-			if (JobTriggers != null) TraceScope.Span.SetTag("JobTriggers.Count", JobTriggers.Count);
+			using IScope traceScope = GlobalTracer.Instance.BuildSpan("JobService.CreateJobAsync").StartActive();
+			traceScope.Span.SetTag("JobId", jobId);
+			traceScope.Span.SetTag("Stream", stream.Name);
+			traceScope.Span.SetTag("TemplateRefId", templateRefId);
+			traceScope.Span.SetTag("TemplateHash", templateHash);
+			traceScope.Span.SetTag("GraphId", graph.Id);
+			traceScope.Span.SetTag("Name", name);
+			traceScope.Span.SetTag("Change", change);
+			traceScope.Span.SetTag("CodeChange", codeChange);
+			traceScope.Span.SetTag("PreflightChange", preflightChange);
+			traceScope.Span.SetTag("ClonedPreflightChange", clonedPreflightChange);
+			traceScope.Span.SetTag("StartedByUserId", startedByUserId.ToString());
+			traceScope.Span.SetTag("Priority", priority.ToString());
+			traceScope.Span.SetTag("ShowUgsBadges", showUgsBadges);
+			traceScope.Span.SetTag("NotificationChannel", notificationChannel ?? "null");
+			traceScope.Span.SetTag("NotificationChannelFilter", notificationChannelFilter ?? "null");
+			traceScope.Span.SetTag("Arguments", String.Join(',', arguments));
+
+			if (autoSubmit != null)
+			{
+				traceScope.Span.SetTag("AutoSubmit", autoSubmit.Value);
+			}
+			if (updateIssues != null)
+			{
+				traceScope.Span.SetTag("UpdateIssues", updateIssues.Value);
+			}
+			if (jobTriggers != null)
+			{
+				traceScope.Span.SetTag("JobTriggers.Count", jobTriggers.Count);
+			}
 
-			JobId JobIdValue = JobId ?? Horde.Build.Utilities.ObjectId<IJob>.GenerateNewId();
-			using IDisposable Scope = Logger.BeginScope("CreateJobAsync({JobId})", JobIdValue);
+			JobId jobIdValue = jobId ?? Horde.Build.Utilities.ObjectId<IJob>.GenerateNewId();
+			using IDisposable scope = _logger.BeginScope("CreateJobAsync({JobId})", jobIdValue);
 
-			if (PreflightChange != null && ShouldClonePreflightChange(Stream.Id))
+			if (preflightChange != null && ShouldClonePreflightChange(stream.Id))
 			{
-				ClonedPreflightChange = await CloneShelvedChangeAsync(Stream.ClusterName, ClonedPreflightChange ?? PreflightChange.Value);
+				clonedPreflightChange = await CloneShelvedChangeAsync(stream.ClusterName, clonedPreflightChange ?? preflightChange.Value);
 			}
 
-			Logger.LogInformation("Creating job at CL {Change}, code CL {CodeChange}, preflight CL {PreflightChange}, cloned CL {ClonedPreflightChange}", Change, CodeChange, PreflightChange, ClonedPreflightChange);
+			_logger.LogInformation("Creating job at CL {Change}, code CL {CodeChange}, preflight CL {PreflightChange}, cloned CL {ClonedPreflightChange}", change, codeChange, preflightChange, clonedPreflightChange);
 
-			Dictionary<string, string> Properties = new Dictionary<string, string>();
-			Properties["Change"] = Change.ToString(CultureInfo.InvariantCulture);
-			Properties["CodeChange"] = CodeChange.ToString(CultureInfo.InvariantCulture);
-			Properties["PreflightChange"] = PreflightChange?.ToString(CultureInfo.InvariantCulture) ?? String.Empty;
-			Properties["ClonedPreflightChange"] = ClonedPreflightChange?.ToString(CultureInfo.InvariantCulture) ?? String.Empty;
-			Properties["StreamId"] = Stream.Id.ToString();
-			Properties["TemplateId"] = TemplateRefId.ToString();
-			Properties["JobId"] = JobIdValue.ToString();
+			Dictionary<string, string> properties = new Dictionary<string, string>();
+			properties["Change"] = change.ToString(CultureInfo.InvariantCulture);
+			properties["CodeChange"] = codeChange.ToString(CultureInfo.InvariantCulture);
+			properties["PreflightChange"] = preflightChange?.ToString(CultureInfo.InvariantCulture) ?? String.Empty;
+			properties["ClonedPreflightChange"] = clonedPreflightChange?.ToString(CultureInfo.InvariantCulture) ?? String.Empty;
+			properties["StreamId"] = stream.Id.ToString();
+			properties["TemplateId"] = templateRefId.ToString();
+			properties["JobId"] = jobIdValue.ToString();
 
-			List<string> ExpandedArguments = new List<string>();
-			if (Arguments != null)
+			List<string> expandedArguments = new List<string>();
+			if (arguments != null)
 			{
-				foreach (string Argument in Arguments)
+				foreach (string argument in arguments)
 				{
-					string ExpandedArgument = StringUtils.ExpandProperties(Argument, Properties);
-					ExpandedArguments.Add(ExpandedArgument);
+					string expandedArgument = StringUtils.ExpandProperties(argument, properties);
+					expandedArguments.Add(expandedArgument);
 				}
 			}
 
-			Name = StringUtils.ExpandProperties(Name, Properties);
+			name = StringUtils.ExpandProperties(name, properties);
 
-			IJob NewJob = await Jobs.AddAsync(JobIdValue, Stream.Id, TemplateRefId, TemplateHash, Graph, Name, Change, CodeChange, PreflightChange, ClonedPreflightChange, StartedByUserId, Priority, AutoSubmit, UpdateIssues, PromoteIssuesByDefault, JobTriggers, ShowUgsBadges, ShowUgsAlerts, NotificationChannel, NotificationChannelFilter, ExpandedArguments);
-			JobTaskSource.UpdateQueuedJob(NewJob, Graph);
+			IJob newJob = await _jobs.AddAsync(jobIdValue, stream.Id, templateRefId, templateHash, graph, name, change, codeChange, preflightChange, clonedPreflightChange, startedByUserId, priority, autoSubmit, updateIssues, promoteIssuesByDefault, jobTriggers, showUgsBadges, showUgsAlerts, notificationChannel, notificationChannelFilter, expandedArguments);
+			_jobTaskSource.UpdateQueuedJob(newJob, graph);
 
-			await JobTaskSource.UpdateUgsBadges(NewJob, Graph, new List<(LabelState, LabelOutcome)>());
+			await _jobTaskSource.UpdateUgsBadges(newJob, graph, new List<(LabelState, LabelOutcome)>());
 
-			if (StartedByUserId != null)
+			if (startedByUserId != null)
 			{
-				await UserCollection.UpdateSettingsAsync(StartedByUserId.Value, AddPinnedJobIds: new[] { NewJob.Id });
+				await _userCollection.UpdateSettingsAsync(startedByUserId.Value, addPinnedJobIds: new[] { newJob.Id });
 			}
 
-			await AbortAnyDuplicateJobs(NewJob);
+			await AbortAnyDuplicateJobs(newJob);
 
-			return NewJob;
+			return newJob;
 		}
 
-		private async Task AbortAnyDuplicateJobs(IJob NewJob)
+		private async Task AbortAnyDuplicateJobs(IJob newJob)
 		{
-			using IScope Scope = GlobalTracer.Instance.BuildSpan("JobService.AbortAnyDuplicateJobs").StartActive();
-			Scope.Span.SetTag("JobId", NewJob.Id.ToString());
-			Scope.Span.SetTag("JobName", NewJob.Name);
+			using IScope scope = GlobalTracer.Instance.BuildSpan("JobService.AbortAnyDuplicateJobs").StartActive();
+			scope.Span.SetTag("JobId", newJob.Id.ToString());
+			scope.Span.SetTag("JobName", newJob.Name);
 			
-			List<IJob> JobsToAbort = new List<IJob>();
-			if (NewJob.PreflightChange > 0)
+			List<IJob> jobsToAbort = new List<IJob>();
+			if (newJob.PreflightChange > 0)
 			{
-				JobsToAbort = await Jobs.FindAsync(PreflightChange: NewJob.PreflightChange);
+				jobsToAbort = await _jobs.FindAsync(preflightChange: newJob.PreflightChange);
 			}
 
-			foreach (IJob Job in JobsToAbort)
+			foreach (IJob job in jobsToAbort)
 			{
-				if (Job.GetState() == JobState.Complete) continue;
-				if (Job.Id == NewJob.Id) continue; // Don't remove the new job
-				if (Job.TemplateId != NewJob.TemplateId) continue;
-				if (Job.TemplateHash != NewJob.TemplateHash) continue;
-				if (Job.StartedByUserId != NewJob.StartedByUserId) continue;
-				if (string.Join(",", Job.Arguments) != string.Join(",", NewJob.Arguments)) continue;
-
-				IJob? UpdatedJob = await UpdateJobAsync(Job, null, null, null, KnownUsers.System, null, null, null);
-				if (UpdatedJob == null)
+				if (job.GetState() == JobState.Complete)
+				{
+					continue;
+				}
+				if (job.Id == newJob.Id)
+				{
+					continue; // Don't remove the new job
+				}
+				if (job.TemplateId != newJob.TemplateId)
+				{
+					continue;
+				}
+				if (job.TemplateHash != newJob.TemplateHash)
+				{
+					continue;
+				}
+				if (job.StartedByUserId != newJob.StartedByUserId)
+				{
+					continue;
+				}
+				if (String.Join(",", job.Arguments) != String.Join(",", newJob.Arguments))
+				{
+					continue;
+				}
+
+				IJob? updatedJob = await UpdateJobAsync(job, null, null, null, KnownUsers.System, null, null, null);
+				if (updatedJob == null)
 				{
-					Logger.LogError("Failed marking duplicate job as aborted! Job ID: {JobId}", Job.Id);
+					_logger.LogError("Failed marking duplicate job as aborted! Job ID: {JobId}", job.Id);
 				}
 
-				IJob? UpdatedJob2 = await GetJobAsync(Job.Id);
-				if (UpdatedJob2?.AbortedByUserId != UpdatedJob?.AbortedByUserId)
+				IJob? updatedJob2 = await GetJobAsync(job.Id);
+				if (updatedJob2?.AbortedByUserId != updatedJob?.AbortedByUserId)
 				{
 					throw new NotImplementedException();
 				}
@@ -353,40 +311,40 @@ namespace Horde.Build.Services
 		/// <summary>
 		/// Deletes a job
 		/// </summary>
-		/// <param name="Job">The job to delete</param>
-		public async Task<bool> DeleteJobAsync(IJob Job)
+		/// <param name="job">The job to delete</param>
+		public async Task<bool> DeleteJobAsync(IJob job)
 		{
-			using IScope TraceScope = GlobalTracer.Instance.BuildSpan("JobService.DeleteJobAsync").StartActive();
-			TraceScope.Span.SetTag("JobId", Job.Id.ToString());
-			TraceScope.Span.SetTag("JobName", Job.Name);
+			using IScope traceScope = GlobalTracer.Instance.BuildSpan("JobService.DeleteJobAsync").StartActive();
+			traceScope.Span.SetTag("JobId", job.Id.ToString());
+			traceScope.Span.SetTag("JobName", job.Name);
 
-			using IDisposable Scope = Logger.BeginScope("DeleteJobAsync({JobId})", Job.Id);
+			using IDisposable scope = _logger.BeginScope("DeleteJobAsync({JobId})", job.Id);
 
 			// Delete the job
-			while (!await Jobs.RemoveAsync(Job))
+			while (!await _jobs.RemoveAsync(job))
 			{
-				IJob? NewJob = await Jobs.GetAsync(Job.Id);
-				if (NewJob == null)
+				IJob? newJob = await _jobs.GetAsync(job.Id);
+				if (newJob == null)
 				{
 					return false;
 				}
-				Job = NewJob;
+				job = newJob;
 			}
 
 			// Remove all the triggers from it
-			List<ObjectId> TriggerIds = new List<ObjectId>();
-			if (Job.NotificationTriggerId != null)
+			List<ObjectId> triggerIds = new List<ObjectId>();
+			if (job.NotificationTriggerId != null)
 			{
-				TriggerIds.Add(Job.NotificationTriggerId.Value);
+				triggerIds.Add(job.NotificationTriggerId.Value);
 			}
-			foreach (IJobStep Step in Job.Batches.SelectMany(x => x.Steps))
+			foreach (IJobStep step in job.Batches.SelectMany(x => x.Steps))
 			{
-				if (Step.NotificationTriggerId != null)
+				if (step.NotificationTriggerId != null)
 				{
-					TriggerIds.Add(Step.NotificationTriggerId.Value);
+					triggerIds.Add(step.NotificationTriggerId.Value);
 				}
 			}
-			await TriggerCollection.DeleteAsync(TriggerIds);
+			await _triggerCollection.DeleteAsync(triggerIds);
 
 			return true;
 		}
@@ -394,55 +352,55 @@ namespace Horde.Build.Services
 		/// <summary>
 		/// Delete all the jobs for a stream
 		/// </summary>
-		/// <param name="StreamId">Unique id of the stream</param>
+		/// <param name="streamId">Unique id of the stream</param>
 		/// <returns>Async task</returns>
-		public async Task DeleteJobsForStreamAsync(StreamId StreamId)
+		public async Task DeleteJobsForStreamAsync(StreamId streamId)
 		{
-			await Jobs.RemoveStreamAsync(StreamId);
+			await _jobs.RemoveStreamAsync(streamId);
 		}
 
 		/// <summary>
 		/// Updates a new job
 		/// </summary>
-		/// <param name="Job">The job document to update</param>
-		/// <param name="Name">Name of the job</param>
-		/// <param name="Priority">Priority of the job</param>
-		/// <param name="AutoSubmit">Whether to automatically submit the preflighted change on completion</param>
-		/// <param name="AbortedByUserId">Name of the user that aborted this job</param>
-		/// <param name="OnCompleteTriggerId">Object id for a notification trigger</param>
-		/// <param name="Reports">New reports to add</param>
-		/// <param name="Arguments">New arguments for the job</param>
-		/// <param name="LabelIdxToTriggerId">New trigger ID for a label in the job</param>
-		public async Task<IJob?> UpdateJobAsync(IJob Job, string? Name = null, Priority? Priority = null, bool? AutoSubmit = null, UserId? AbortedByUserId = null, ObjectId? OnCompleteTriggerId = null, List<Report>? Reports = null, List<string>? Arguments = null, KeyValuePair<int, ObjectId>? LabelIdxToTriggerId = null)
+		/// <param name="job">The job document to update</param>
+		/// <param name="name">Name of the job</param>
+		/// <param name="priority">Priority of the job</param>
+		/// <param name="autoSubmit">Whether to automatically submit the preflighted change on completion</param>
+		/// <param name="abortedByUserId">Name of the user that aborted this job</param>
+		/// <param name="onCompleteTriggerId">Object id for a notification trigger</param>
+		/// <param name="reports">New reports to add</param>
+		/// <param name="arguments">New arguments for the job</param>
+		/// <param name="labelIdxToTriggerId">New trigger ID for a label in the job</param>
+		public async Task<IJob?> UpdateJobAsync(IJob job, string? name = null, Priority? priority = null, bool? autoSubmit = null, UserId? abortedByUserId = null, ObjectId? onCompleteTriggerId = null, List<Report>? reports = null, List<string>? arguments = null, KeyValuePair<int, ObjectId>? labelIdxToTriggerId = null)
 		{
-			using IScope TraceScope = GlobalTracer.Instance.BuildSpan("JobService.UpdateJobAsync").StartActive();
-			TraceScope.Span.SetTag("JobId", Job.Id.ToString());
-			TraceScope.Span.SetTag("Name", Name);
+			using IScope traceScope = GlobalTracer.Instance.BuildSpan("JobService.UpdateJobAsync").StartActive();
+			traceScope.Span.SetTag("JobId", job.Id.ToString());
+			traceScope.Span.SetTag("Name", name);
 			
-			using IDisposable Scope = Logger.BeginScope("UpdateJobAsync({JobId})", Job.Id);
-			for(IJob? NewJob = Job; NewJob != null; NewJob = await GetJobAsync(Job.Id))
+			using IDisposable scope = _logger.BeginScope("UpdateJobAsync({JobId})", job.Id);
+			for(IJob? newJob = job; newJob != null; newJob = await GetJobAsync(job.Id))
 			{
-				IGraph Graph = await GetGraphAsync(NewJob);
+				IGraph graph = await GetGraphAsync(newJob);
 
 				// Capture the previous label states
-				IReadOnlyList<(LabelState, LabelOutcome)> OldLabelStates = NewJob.GetLabelStates(Graph);
+				IReadOnlyList<(LabelState, LabelOutcome)> oldLabelStates = newJob.GetLabelStates(graph);
 
 				// Update the new list of job steps
-				NewJob = await Jobs.TryUpdateJobAsync(NewJob, Graph, Name, Priority, AutoSubmit, null, null, AbortedByUserId, OnCompleteTriggerId, Reports, Arguments, LabelIdxToTriggerId);
-				if (NewJob != null)
+				newJob = await _jobs.TryUpdateJobAsync(newJob, graph, name, priority, autoSubmit, null, null, abortedByUserId, onCompleteTriggerId, reports, arguments, labelIdxToTriggerId);
+				if (newJob != null)
 				{
 					// Update any badges that have been modified
-					await JobTaskSource.UpdateUgsBadges(NewJob, Graph, OldLabelStates);
+					await _jobTaskSource.UpdateUgsBadges(newJob, graph, oldLabelStates);
 
 					// Cancel any leases which are no longer required
-					foreach (IJobStepBatch Batch in NewJob.Batches)
+					foreach (IJobStepBatch batch in newJob.Batches)
 					{
-						if (Batch.Error == JobStepBatchError.Cancelled && (Batch.State == JobStepBatchState.Starting || Batch.State == JobStepBatchState.Running) && Batch.AgentId != null && Batch.LeaseId != null)
+						if (batch.Error == JobStepBatchError.Cancelled && (batch.State == JobStepBatchState.Starting || batch.State == JobStepBatchState.Running) && batch.AgentId != null && batch.LeaseId != null)
 						{
-							await CancelLeaseAsync(Batch.AgentId.Value, Batch.LeaseId.Value);
+							await CancelLeaseAsync(batch.AgentId.Value, batch.LeaseId.Value);
 						}
 					}
-					return NewJob;
+					return newJob;
 				}
 			}
 			return null;
@@ -451,43 +409,43 @@ namespace Horde.Build.Services
 		/// <summary>
 		/// Cancel an active lease
 		/// </summary>
-		/// <param name="AgentId">The agent to retreive</param>
-		/// <param name="LeaseId">The lease id to update</param>
+		/// <param name="agentId">The agent to retreive</param>
+		/// <param name="leaseId">The lease id to update</param>
 		/// <returns></returns>
-		async Task CancelLeaseAsync(AgentId AgentId, LeaseId LeaseId)
+		async Task CancelLeaseAsync(AgentId agentId, LeaseId leaseId)
 		{
-			using IScope Scope = GlobalTracer.Instance.BuildSpan("JobService.CancelLeaseAsync").StartActive();
-			Scope.Span.SetTag("AgentId", AgentId.ToString());
-			Scope.Span.SetTag("LeaseId", LeaseId.ToString());
+			using IScope scope = GlobalTracer.Instance.BuildSpan("JobService.CancelLeaseAsync").StartActive();
+			scope.Span.SetTag("AgentId", agentId.ToString());
+			scope.Span.SetTag("LeaseId", leaseId.ToString());
 			
 			for (; ; )
 			{
-				IAgent? Agent = await Agents.GetAsync(AgentId);
-				if (Agent == null)
+				IAgent? agent = await _agents.GetAsync(agentId);
+				if (agent == null)
 				{
 					break;
 				}
 
-				int Index = 0;
-				while (Index < Agent.Leases.Count && Agent.Leases[Index].Id != LeaseId)
+				int index = 0;
+				while (index < agent.Leases.Count && agent.Leases[index].Id != leaseId)
 				{
-					Index++;
+					index++;
 				}
-				if (Index == Agent.Leases.Count)
+				if (index == agent.Leases.Count)
 				{
 					break;
 				}
 
-				AgentLease Lease = Agent.Leases[Index];
-				if (Lease.State != LeaseState.Active)
+				AgentLease lease = agent.Leases[index];
+				if (lease.State != LeaseState.Active)
 				{
 					break;
 				}
 
-				IAgent? NewAgent = await Agents.TryCancelLeaseAsync(Agent, Index);
-				if (NewAgent != null)
+				IAgent? newAgent = await _agents.TryCancelLeaseAsync(agent, index);
+				if (newAgent != null)
 				{
-					JobTaskSource.CancelLongPollForAgent(Agent.Id);
+					_jobTaskSource.CancelLongPollForAgent(agent.Id);
 					break;
 				}
 			}
@@ -496,117 +454,117 @@ namespace Horde.Build.Services
 		/// <summary>
 		/// Gets a job with the given unique id
 		/// </summary>
-		/// <param name="JobId">Job id to search for</param>
+		/// <param name="jobId">Job id to search for</param>
 		/// <returns>Information about the given job</returns>
-		public Task<IJob?> GetJobAsync(JobId JobId)
+		public Task<IJob?> GetJobAsync(JobId jobId)
 		{
-			return Jobs.GetAsync(JobId);
+			return _jobs.GetAsync(jobId);
 		}
 
 		/// <summary>
 		/// Gets the graph for a job
 		/// </summary>
-		/// <param name="Job">Job to retrieve the graph for</param>
+		/// <param name="job">Job to retrieve the graph for</param>
 		/// <returns>The graph for this job</returns>
-		public Task<IGraph> GetGraphAsync(IJob Job)
+		public Task<IGraph> GetGraphAsync(IJob job)
 		{
-			return Graphs.GetAsync(Job.GraphHash);
+			return _graphs.GetAsync(job.GraphHash);
 		}
 
 		/// <summary>
 		/// Gets a job's permissions info by ID
 		/// </summary>
-		/// <param name="JobId">Unique id of the job</param>
+		/// <param name="jobId">Unique id of the job</param>
 		/// <returns>The job document</returns>
-		public Task<IJobPermissions?> GetJobPermissionsAsync(JobId JobId)
+		public Task<IJobPermissions?> GetJobPermissionsAsync(JobId jobId)
 		{
-			return Jobs.GetPermissionsAsync(JobId);
+			return _jobs.GetPermissionsAsync(jobId);
 		}
 
 		/// <summary>
 		/// Searches for jobs matching the given criteria
 		/// </summary>
-		/// <param name="JobIds">List of job ids to return</param>
-		/// <param name="StreamId">The stream containing the job</param>
-		/// <param name="Name">Name of the job</param>
-		/// <param name="Templates">Templates to look for</param>
-		/// <param name="MinChange">The minimum changelist number</param>
-		/// <param name="MaxChange">The maximum changelist number</param>		
-		/// <param name="PreflightChange">The preflight change to look for</param>
-		/// <param name="PreflightOnly">Whether to only include preflights</param>
-		/// <param name="PreflightStartedByUser">User for which to include preflight jobs</param>		
-		/// <param name="StartedByUser">User for which to include jobs</param>
-		/// <param name="MinCreateTime">The minimum creation time</param>
-		/// <param name="MaxCreateTime">The maximum creation time</param>
-		/// <param name="Target">The target to query</param>
-		/// <param name="State">State to query</param>
-		/// <param name="Outcome">Outcomes to return</param>
-		/// <param name="ModifiedBefore">Filter the results by last modified time</param>
-		/// <param name="ModifiedAfter">Filter the results by last modified time</param>
-		/// <param name="Index">Index of the first result to return</param>
-		/// <param name="Count">Number of results to return</param>
-		/// <param name="ConsistentRead">If the database read should be made to the replica server</param>
-		/// <param name="ExcludeUserJobs">Whether to exclude user jobs from the find</param>
+		/// <param name="jobIds">List of job ids to return</param>
+		/// <param name="streamId">The stream containing the job</param>
+		/// <param name="name">Name of the job</param>
+		/// <param name="templates">Templates to look for</param>
+		/// <param name="minChange">The minimum changelist number</param>
+		/// <param name="maxChange">The maximum changelist number</param>		
+		/// <param name="preflightChange">The preflight change to look for</param>
+		/// <param name="preflightOnly">Whether to only include preflights</param>
+		/// <param name="preflightStartedByUser">User for which to include preflight jobs</param>		
+		/// <param name="startedByUser">User for which to include jobs</param>
+		/// <param name="minCreateTime">The minimum creation time</param>
+		/// <param name="maxCreateTime">The maximum creation time</param>
+		/// <param name="target">The target to query</param>
+		/// <param name="state">State to query</param>
+		/// <param name="outcome">Outcomes to return</param>
+		/// <param name="modifiedBefore">Filter the results by last modified time</param>
+		/// <param name="modifiedAfter">Filter the results by last modified time</param>
+		/// <param name="index">Index of the first result to return</param>
+		/// <param name="count">Number of results to return</param>
+		/// <param name="consistentRead">If the database read should be made to the replica server</param>
+		/// <param name="excludeUserJobs">Whether to exclude user jobs from the find</param>
 		/// <returns>List of jobs matching the given criteria</returns>
-		public async Task<List<IJob>> FindJobsAsync(JobId[]? JobIds = null, StreamId? StreamId = null, string? Name = null, TemplateRefId[]? Templates = null, int? MinChange = null, int? MaxChange = null, int? PreflightChange = null, bool? PreflightOnly = null, UserId? PreflightStartedByUser = null, UserId? StartedByUser = null, DateTimeOffset ? MinCreateTime = null, DateTimeOffset? MaxCreateTime = null, string? Target = null, JobStepState[]? State = null, JobStepOutcome[]? Outcome = null, DateTimeOffset? ModifiedBefore = null, DateTimeOffset? ModifiedAfter = null, int? Index = null, int? Count = null, bool ConsistentRead = true, bool? ExcludeUserJobs = null)
+		public async Task<List<IJob>> FindJobsAsync(JobId[]? jobIds = null, StreamId? streamId = null, string? name = null, TemplateRefId[]? templates = null, int? minChange = null, int? maxChange = null, int? preflightChange = null, bool? preflightOnly = null, UserId? preflightStartedByUser = null, UserId? startedByUser = null, DateTimeOffset ? minCreateTime = null, DateTimeOffset? maxCreateTime = null, string? target = null, JobStepState[]? state = null, JobStepOutcome[]? outcome = null, DateTimeOffset? modifiedBefore = null, DateTimeOffset? modifiedAfter = null, int? index = null, int? count = null, bool consistentRead = true, bool? excludeUserJobs = null)
 		{
-			using IScope Scope = GlobalTracer.Instance.BuildSpan("JobService.FindJobsAsync").StartActive();
-			Scope.Span.SetTag("JobIds", JobIds);
-			Scope.Span.SetTag("StreamId", StreamId);
-			Scope.Span.SetTag("Name", Name);
-			Scope.Span.SetTag("Templates", Templates);
-			Scope.Span.SetTag("MinChange", MinChange);
-			Scope.Span.SetTag("MaxChange", MaxChange);
-			Scope.Span.SetTag("PreflightChange", PreflightChange);
-			Scope.Span.SetTag("PreflightStartedByUser", PreflightStartedByUser);
-			Scope.Span.SetTag("StartedByUser", StartedByUser);
-			Scope.Span.SetTag("MinCreateTime", MinCreateTime);
-			Scope.Span.SetTag("MaxCreateTime", MaxCreateTime);
-			Scope.Span.SetTag("Target", Target);
-			Scope.Span.SetTag("State", State?.ToString());
-			Scope.Span.SetTag("Outcome", Outcome?.ToString());
-			Scope.Span.SetTag("ModifiedBefore", ModifiedBefore);
-			Scope.Span.SetTag("ModifiedAfter", ModifiedAfter);
-			Scope.Span.SetTag("Index", Index);
-			Scope.Span.SetTag("Count", Count);
+			using IScope scope = GlobalTracer.Instance.BuildSpan("JobService.FindJobsAsync").StartActive();
+			scope.Span.SetTag("JobIds", jobIds);
+			scope.Span.SetTag("StreamId", streamId);
+			scope.Span.SetTag("Name", name);
+			scope.Span.SetTag("Templates", templates);
+			scope.Span.SetTag("MinChange", minChange);
+			scope.Span.SetTag("MaxChange", maxChange);
+			scope.Span.SetTag("PreflightChange", preflightChange);
+			scope.Span.SetTag("PreflightStartedByUser", preflightStartedByUser);
+			scope.Span.SetTag("StartedByUser", startedByUser);
+			scope.Span.SetTag("MinCreateTime", minCreateTime);
+			scope.Span.SetTag("MaxCreateTime", maxCreateTime);
+			scope.Span.SetTag("Target", target);
+			scope.Span.SetTag("State", state?.ToString());
+			scope.Span.SetTag("Outcome", outcome?.ToString());
+			scope.Span.SetTag("ModifiedBefore", modifiedBefore);
+			scope.Span.SetTag("ModifiedAfter", modifiedAfter);
+			scope.Span.SetTag("Index", index);
+			scope.Span.SetTag("Count", count);
 			
-			if (Target == null && (State == null || State.Length == 0) && (Outcome == null || Outcome.Length == 0))
+			if (target == null && (state == null || state.Length == 0) && (outcome == null || outcome.Length == 0))
 			{
-				return await Jobs.FindAsync(JobIds, StreamId, Name, Templates, MinChange, MaxChange, PreflightChange, PreflightOnly, PreflightStartedByUser, StartedByUser, MinCreateTime, MaxCreateTime, ModifiedBefore, ModifiedAfter, Index, Count, ConsistentRead, null, ExcludeUserJobs);
+				return await _jobs.FindAsync(jobIds, streamId, name, templates, minChange, maxChange, preflightChange, preflightOnly, preflightStartedByUser, startedByUser, minCreateTime, maxCreateTime, modifiedBefore, modifiedAfter, index, count, consistentRead, null, excludeUserJobs);
 			}
 			else
 			{
-				List<IJob> Results = new List<IJob>();
-				Logger.LogInformation("Performing scan for job with ");
+				List<IJob> results = new List<IJob>();
+				_logger.LogInformation("Performing scan for job with ");
 
-				int MaxCount = (Count ?? 1);
-				while (Results.Count < MaxCount)
+				int maxCount = (count ?? 1);
+				while (results.Count < maxCount)
 				{
-					List<IJob> ScanJobs = await Jobs.FindAsync(JobIds, StreamId, Name, Templates, MinChange, MaxChange, PreflightChange, PreflightOnly, PreflightStartedByUser, StartedByUser, MinCreateTime, MaxCreateTime, ModifiedBefore, ModifiedAfter, 0, 5, ConsistentRead, null, ExcludeUserJobs);
-					if (ScanJobs.Count == 0)
+					List<IJob> scanJobs = await _jobs.FindAsync(jobIds, streamId, name, templates, minChange, maxChange, preflightChange, preflightOnly, preflightStartedByUser, startedByUser, minCreateTime, maxCreateTime, modifiedBefore, modifiedAfter, 0, 5, consistentRead, null, excludeUserJobs);
+					if (scanJobs.Count == 0)
 					{
 						break;
 					}
 
-					foreach (IJob Job in ScanJobs.OrderByDescending(x => x.Change))
+					foreach (IJob job in scanJobs.OrderByDescending(x => x.Change))
 					{
-						(JobStepState, JobStepOutcome)? Result;
-						if (Target == null)
+						(JobStepState, JobStepOutcome)? result;
+						if (target == null)
 						{
-							Result = Job.GetTargetState();
+							result = job.GetTargetState();
 						}
 						else
 						{
-							Result = Job.GetTargetState(await GetGraphAsync(Job), Target);
+							result = job.GetTargetState(await GetGraphAsync(job), target);
 						}
 
-						if (Result != null)
+						if (result != null)
 						{
-							(JobStepState JobState, JobStepOutcome JobOutcome) = Result.Value;
-							if ((State == null || State.Length == 0 || State.Contains(JobState)) && (Outcome == null || Outcome.Length == 0 || Outcome.Contains(JobOutcome)))
+							(JobStepState jobState, JobStepOutcome jobOutcome) = result.Value;
+							if ((state == null || state.Length == 0 || state.Contains(jobState)) && (outcome == null || outcome.Length == 0 || outcome.Contains(jobOutcome)))
 							{
-								Results.Add(Job);
-								if (Results.Count == MaxCount)
+								results.Add(job);
+								if (results.Count == maxCount)
 								{
 									break;
 								}
@@ -614,130 +572,130 @@ namespace Horde.Build.Services
 						}
 					}
 
-					MaxChange = ScanJobs.Min(x => x.Change) - 1;
+					maxChange = scanJobs.Min(x => x.Change) - 1;
 				}
 
-				return Results;
+				return results;
 			}
 		}
 		
 		/// <summary>
 		/// Searches for jobs matching a stream with given templates
 		/// </summary>
-		/// <param name="StreamId">The stream containing the job</param>
-		/// <param name="Templates">Templates to look for</param>
-		/// <param name="PreflightStartedByUser">User for which to include preflight jobs</param>
-		/// <param name="MaxCreateTime">The maximum creation time</param>
-		/// <param name="ModifiedAfter">Filter the results by last modified time</param>	
-		/// <param name="Index">Index of the first result to return</param>
-		/// <param name="Count">Number of results to return</param>
-		/// <param name="ConsistentRead">If the database read should be made to the replica server</param>
+		/// <param name="streamId">The stream containing the job</param>
+		/// <param name="templates">Templates to look for</param>
+		/// <param name="preflightStartedByUser">User for which to include preflight jobs</param>
+		/// <param name="maxCreateTime">The maximum creation time</param>
+		/// <param name="modifiedAfter">Filter the results by last modified time</param>	
+		/// <param name="index">Index of the first result to return</param>
+		/// <param name="count">Number of results to return</param>
+		/// <param name="consistentRead">If the database read should be made to the replica server</param>
 		/// <returns>List of jobs matching the given criteria</returns>
-		public async Task<List<IJob>> FindJobsByStreamWithTemplatesAsync(StreamId StreamId, TemplateRefId[] Templates, UserId? PreflightStartedByUser = null, DateTimeOffset? MaxCreateTime = null, DateTimeOffset? ModifiedAfter = null, int? Index = null, int? Count = null, bool ConsistentRead = true)
+		public async Task<List<IJob>> FindJobsByStreamWithTemplatesAsync(StreamId streamId, TemplateRefId[] templates, UserId? preflightStartedByUser = null, DateTimeOffset? maxCreateTime = null, DateTimeOffset? modifiedAfter = null, int? index = null, int? count = null, bool consistentRead = true)
 		{
-			using IScope Scope = GlobalTracer.Instance.BuildSpan("JobService.FindJobsByStreamWithTemplatesAsync").StartActive();
-			Scope.Span.SetTag("StreamId", StreamId);
-			Scope.Span.SetTag("Templates", Templates);
-			Scope.Span.SetTag("PreflightStartedByUser", PreflightStartedByUser);
-			Scope.Span.SetTag("MaxCreateTime", MaxCreateTime);
-			Scope.Span.SetTag("ModifiedAfter", ModifiedAfter);
-			Scope.Span.SetTag("Index", Index);
-			Scope.Span.SetTag("Count", Count);
+			using IScope scope = GlobalTracer.Instance.BuildSpan("JobService.FindJobsByStreamWithTemplatesAsync").StartActive();
+			scope.Span.SetTag("StreamId", streamId);
+			scope.Span.SetTag("Templates", templates);
+			scope.Span.SetTag("PreflightStartedByUser", preflightStartedByUser);
+			scope.Span.SetTag("MaxCreateTime", maxCreateTime);
+			scope.Span.SetTag("ModifiedAfter", modifiedAfter);
+			scope.Span.SetTag("Index", index);
+			scope.Span.SetTag("Count", count);
 			
-			return await Jobs.FindLatestByStreamWithTemplatesAsync(StreamId, Templates, PreflightStartedByUser, MaxCreateTime, ModifiedAfter, Index, Count, ConsistentRead);
+			return await _jobs.FindLatestByStreamWithTemplatesAsync(streamId, templates, preflightStartedByUser, maxCreateTime, modifiedAfter, index, count, consistentRead);
 		}
 
 		/// <summary>
 		/// Attempts to update the node groups to be executed for a job. Fails if another write happens in the meantime.
 		/// </summary>
-		/// <param name="Job">The job to update</param>
-		/// <param name="NewGraph">New graph for this job</param>
+		/// <param name="job">The job to update</param>
+		/// <param name="newGraph">New graph for this job</param>
 		/// <returns>True if the groups were updated to the given list. False if another write happened first.</returns>
-		public async Task<IJob?> TryUpdateGraphAsync(IJob Job, IGraph NewGraph)
+		public async Task<IJob?> TryUpdateGraphAsync(IJob job, IGraph newGraph)
 		{
-			using IScope TraceScope = GlobalTracer.Instance.BuildSpan("JobService.TryUpdateGraphAsync").StartActive();
-			TraceScope.Span.SetTag("Job", Job.Id);
-			TraceScope.Span.SetTag("NewGraph", NewGraph.Id);
+			using IScope traceScope = GlobalTracer.Instance.BuildSpan("JobService.TryUpdateGraphAsync").StartActive();
+			traceScope.Span.SetTag("Job", job.Id);
+			traceScope.Span.SetTag("NewGraph", newGraph.Id);
 
-			using IDisposable Scope = Logger.BeginScope("TryUpdateGraphAsync({JobId})", Job.Id);
+			using IDisposable scope = _logger.BeginScope("TryUpdateGraphAsync({JobId})", job.Id);
 
-			IReadOnlyList<(LabelState, LabelOutcome)> OldLabelStates = Job.GetLabelStates(NewGraph);
+			IReadOnlyList<(LabelState, LabelOutcome)> oldLabelStates = job.GetLabelStates(newGraph);
 
-			IJob? NewJob = await Jobs.TryUpdateGraphAsync(Job, NewGraph);
-			if(NewJob != null)
+			IJob? newJob = await _jobs.TryUpdateGraphAsync(job, newGraph);
+			if(newJob != null)
 			{
-				await JobTaskSource.UpdateUgsBadges(NewJob, NewGraph, OldLabelStates);
-				JobTaskSource.UpdateQueuedJob(NewJob, NewGraph);
+				await _jobTaskSource.UpdateUgsBadges(newJob, newGraph, oldLabelStates);
+				_jobTaskSource.UpdateQueuedJob(newJob, newGraph);
 			}
-			return NewJob;
+			return newJob;
 		}
 
 		/// <summary>
 		/// Gets the timing info for a particular job.
 		/// </summary>
-		/// <param name="Job">The job to get timing info for</param>
+		/// <param name="job">The job to get timing info for</param>
 		/// <returns>Timing info for the given job</returns>
-		public async Task<IJobTiming> GetJobTimingAsync(IJob Job)
+		public async Task<IJobTiming> GetJobTimingAsync(IJob job)
 		{
-			using IScope TraceScope = GlobalTracer.Instance.BuildSpan("JobService.GetJobTimingAsync").StartActive();
-			TraceScope.Span.SetTag("Job", Job.Id);
+			using IScope traceScope = GlobalTracer.Instance.BuildSpan("JobService.GetJobTimingAsync").StartActive();
+			traceScope.Span.SetTag("Job", job.Id);
 
-			using IDisposable Scope = Logger.BeginScope("GetJobTimingAsync({JobId})", Job.Id);
+			using IDisposable scope = _logger.BeginScope("GetJobTimingAsync({JobId})", job.Id);
 
-			IGraph Graph = await Graphs.GetAsync(Job.GraphHash);
+			IGraph graph = await _graphs.GetAsync(job.GraphHash);
 
-			Dictionary<string, JobStepTimingData> CachedNewSteps = new Dictionary<string, JobStepTimingData>();
+			Dictionary<string, JobStepTimingData> cachedNewSteps = new Dictionary<string, JobStepTimingData>();
 			for (; ; )
 			{
 				// Try to get the current timing document
-				IJobTiming? JobTiming = await JobTimings.TryGetAsync(Job.Id);
+				IJobTiming? jobTiming = await _jobTimings.TryGetAsync(job.Id);
 
 				// Calculate timing info for any missing steps
-				List<JobStepTimingData> NewSteps = new List<JobStepTimingData>();
-				foreach (IJobStepBatch Batch in Job.Batches)
+				List<JobStepTimingData> newSteps = new List<JobStepTimingData>();
+				foreach (IJobStepBatch batch in job.Batches)
 				{
-					INodeGroup Group = Graph.Groups[Batch.GroupIdx];
-					foreach (IJobStep Step in Batch.Steps)
+					INodeGroup group = graph.Groups[batch.GroupIdx];
+					foreach (IJobStep step in batch.Steps)
 					{
-						INode Node = Group.Nodes[Step.NodeIdx];
-						if (JobTiming == null || !JobTiming.TryGetStepTiming(Node.Name, out _))
+						INode node = group.Nodes[step.NodeIdx];
+						if (jobTiming == null || !jobTiming.TryGetStepTiming(node.Name, out _))
 						{
-							JobStepTimingData? StepTimingData;
-							if (!CachedNewSteps.TryGetValue(Node.Name, out StepTimingData))
+							JobStepTimingData? stepTimingData;
+							if (!cachedNewSteps.TryGetValue(node.Name, out stepTimingData))
 							{
-								StepTimingData = await GetStepTimingInfo(Job.StreamId, Job.TemplateId, Node.Name, Job.Change);
+								stepTimingData = await GetStepTimingInfo(job.StreamId, job.TemplateId, node.Name, job.Change);
 							}
-							NewSteps.Add(StepTimingData);
+							newSteps.Add(stepTimingData);
 						}
 					}
 				}
 
 				// Try to add or update the timing document
-				if (JobTiming == null)
+				if (jobTiming == null)
 				{
-					IJobTiming? NewJobTiming = await JobTimings.TryAddAsync(Job.Id, NewSteps);
-					if (NewJobTiming != null)
+					IJobTiming? newJobTiming = await _jobTimings.TryAddAsync(job.Id, newSteps);
+					if (newJobTiming != null)
 					{
-						return NewJobTiming;
+						return newJobTiming;
 					}
 				}
-				else if (NewSteps.Count == 0)
+				else if (newSteps.Count == 0)
 				{
-					return JobTiming;
+					return jobTiming;
 				}
 				else
 				{
-					IJobTiming? NewJobTiming = await JobTimings.TryAddStepsAsync(JobTiming, NewSteps);
-					if (NewJobTiming != null)
+					IJobTiming? newJobTiming = await _jobTimings.TryAddStepsAsync(jobTiming, newSteps);
+					if (newJobTiming != null)
 					{
-						return NewJobTiming;
+						return newJobTiming;
 					}
 				}
 
 				// Update the cached list of steps for the next iteration
-				foreach (JobStepTimingData NewStep in NewSteps)
+				foreach (JobStepTimingData newStep in newSteps)
 				{
-					CachedNewSteps[NewStep.Name] = NewStep;
+					cachedNewSteps[newStep.Name] = newStep;
 				}
 			}
 		}
@@ -745,108 +703,108 @@ namespace Horde.Build.Services
 		/// <summary>
 		/// Finds the average duration for the given step
 		/// </summary>
-		/// <param name="StreamId">The stream to search</param>
-		/// <param name="TemplateId">The template id</param>
-		/// <param name="NodeName">Name of the node</param>
-		/// <param name="Change">Maximum changelist to consider</param>
+		/// <param name="streamId">The stream to search</param>
+		/// <param name="templateId">The template id</param>
+		/// <param name="nodeName">Name of the node</param>
+		/// <param name="change">Maximum changelist to consider</param>
 		/// <returns>Expected duration for the given step</returns>
-		async Task<JobStepTimingData> GetStepTimingInfo(StreamId StreamId, TemplateRefId TemplateId, string NodeName, int? Change)
+		async Task<JobStepTimingData> GetStepTimingInfo(StreamId streamId, TemplateRefId templateId, string nodeName, int? change)
 		{
-			using IScope TraceScope = GlobalTracer.Instance.BuildSpan("JobService.GetStepTimingInfo").StartActive();
-			TraceScope.Span.SetTag("StreamId", StreamId);
-			TraceScope.Span.SetTag("TemplateId", TemplateId);
-			TraceScope.Span.SetTag("NodeName", NodeName);
-			TraceScope.Span.SetTag("Change", Change);
+			using IScope traceScope = GlobalTracer.Instance.BuildSpan("JobService.GetStepTimingInfo").StartActive();
+			traceScope.Span.SetTag("StreamId", streamId);
+			traceScope.Span.SetTag("TemplateId", templateId);
+			traceScope.Span.SetTag("NodeName", nodeName);
+			traceScope.Span.SetTag("Change", change);
 			
 			// Find all the steps matching the given criteria
-			List<IJobStepRef> Steps = await JobStepRefs.GetStepsForNodeAsync(StreamId, TemplateId, NodeName, Change, false, 10);
+			List<IJobStepRef> steps = await _jobStepRefs.GetStepsForNodeAsync(streamId, templateId, nodeName, change, false, 10);
 
 			// Sum up all the durations and wait times
-			int Count = 0;
-			float WaitTimeSum = 0;
-			float InitTimeSum = 0;
-			float DurationSum = 0;
+			int count = 0;
+			float waitTimeSum = 0;
+			float initTimeSum = 0;
+			float durationSum = 0;
 
-			foreach (IJobStepRef Step in Steps)
+			foreach (IJobStepRef step in steps)
 			{
-				if (Step.FinishTimeUtc != null)
+				if (step.FinishTimeUtc != null)
 				{
-					WaitTimeSum += Step.BatchWaitTime;
-					InitTimeSum += Step.BatchInitTime;
-					DurationSum += (float)(Step.FinishTimeUtc.Value - Step.StartTimeUtc).TotalSeconds;
-					Count++;
+					waitTimeSum += step.BatchWaitTime;
+					initTimeSum += step.BatchInitTime;
+					durationSum += (float)(step.FinishTimeUtc.Value - step.StartTimeUtc).TotalSeconds;
+					count++;
 				}
 			}
 
 			// Compute the averages
-			if (Count == 0)
+			if (count == 0)
 			{
-				return new JobStepTimingData(NodeName, null, null, null);
+				return new JobStepTimingData(nodeName, null, null, null);
 			}
 			else
 			{
-				return new JobStepTimingData(NodeName, WaitTimeSum / Count, InitTimeSum / Count, DurationSum / Count);
+				return new JobStepTimingData(nodeName, waitTimeSum / count, initTimeSum / count, durationSum / count);
 			}
 		}
 
 		/// <summary>
 		/// Updates the state of a batch
 		/// </summary>
-		/// <param name="Job">Job to update</param>
-		/// <param name="BatchId">Unique id of the batch to update</param>
-		/// <param name="NewLogId">The new log file id</param>
-		/// <param name="NewState">New state of the jobstep</param>
+		/// <param name="job">Job to update</param>
+		/// <param name="batchId">Unique id of the batch to update</param>
+		/// <param name="newLogId">The new log file id</param>
+		/// <param name="newState">New state of the jobstep</param>
 		/// <returns>True if the job was updated, false if it was deleted</returns>
-		public async Task<IJob?> UpdateBatchAsync(IJob Job, SubResourceId BatchId, LogId? NewLogId = null, JobStepBatchState? NewState = null)
+		public async Task<IJob?> UpdateBatchAsync(IJob job, SubResourceId batchId, LogId? newLogId = null, JobStepBatchState? newState = null)
 		{
-			using IScope TraceScope = GlobalTracer.Instance.BuildSpan("JobService.UpdateBatchAsync").StartActive();
-			TraceScope.Span.SetTag("Job", Job.Id);
-			TraceScope.Span.SetTag("BatchId", BatchId);
-			TraceScope.Span.SetTag("NewLogId", NewLogId);
-			TraceScope.Span.SetTag("NewState", NewState.ToString());
+			using IScope traceScope = GlobalTracer.Instance.BuildSpan("JobService.UpdateBatchAsync").StartActive();
+			traceScope.Span.SetTag("Job", job.Id);
+			traceScope.Span.SetTag("BatchId", batchId);
+			traceScope.Span.SetTag("NewLogId", newLogId);
+			traceScope.Span.SetTag("NewState", newState.ToString());
 
-			using IDisposable Scope = Logger.BeginScope("UpdateBatchAsync({JobId})", Job.Id);
+			using IDisposable scope = _logger.BeginScope("UpdateBatchAsync({JobId})", job.Id);
 
 			bool bCheckForBadAgent = true;
 			for (; ; )
 			{
 				// Find the index of the appropriate batch
-				int BatchIdx = Job.Batches.FindIndex(x => x.Id == BatchId);
-				if (BatchIdx == -1)
+				int batchIdx = job.Batches.FindIndex(x => x.Id == batchId);
+				if (batchIdx == -1)
 				{
 					return null;
 				}
 
-				IJobStepBatch Batch = Job.Batches[BatchIdx];
+				IJobStepBatch batch = job.Batches[batchIdx];
 
 				// If the batch has already been marked as complete, error out
-				if (Batch.State == JobStepBatchState.Complete)
+				if (batch.State == JobStepBatchState.Complete)
 				{
 					return null;
 				}
 
 				// If we're marking the batch as complete before the agent has run everything, mark it to conform
-				JobStepBatchError? NewError = null;
-				if (NewState == JobStepBatchState.Complete && Batch.AgentId != null)
+				JobStepBatchError? newError = null;
+				if (newState == JobStepBatchState.Complete && batch.AgentId != null)
 				{
-					if (Batch.Steps.Any(x => x.State == JobStepState.Waiting || x.State == JobStepState.Ready))
+					if (batch.Steps.Any(x => x.State == JobStepState.Waiting || x.State == JobStepState.Ready))
 					{
 						// Mark the batch as incomplete
-						NewError = JobStepBatchError.Incomplete;
+						newError = JobStepBatchError.Incomplete;
 
 						// Find the agent and set the conform flag
 						if (bCheckForBadAgent)
 						{
 							for (; ; )
 							{
-								IAgent? Agent = await Agents.GetAsync(Batch.AgentId.Value);
-								if (Agent == null || Agent.RequestConform)
+								IAgent? agent = await _agents.GetAsync(batch.AgentId.Value);
+								if (agent == null || agent.RequestConform)
 								{
 									break;
 								}
-								if (await Agents.TryUpdateSettingsAsync(Agent, bRequestConform: true) != null)
+								if (await _agents.TryUpdateSettingsAsync(agent, bRequestConform: true) != null)
 								{
-									Logger.LogError("Agent {AgentId} did not complete lease; marking for conform", Agent.Id);
+									_logger.LogError("Agent {AgentId} did not complete lease; marking for conform", agent.Id);
 									break;
 								}
 							}
@@ -856,234 +814,234 @@ namespace Horde.Build.Services
 				}
 
 				// Update the batch state
-				IJob? NewJob = await TryUpdateBatchAsync(Job, BatchId, NewLogId, NewState, NewError);
-				if (NewJob != null)
+				IJob? newJob = await TryUpdateBatchAsync(job, batchId, newLogId, newState, newError);
+				if (newJob != null)
 				{
-					return NewJob;
+					return newJob;
 				}
 
 				// Update the job
-				NewJob = await GetJobAsync(Job.Id);
-				if (NewJob == null)
+				newJob = await GetJobAsync(job.Id);
+				if (newJob == null)
 				{
 					return null;
 				}
 
-				Job = NewJob;
+				job = newJob;
 			}
 		}
 
 		/// <summary>
 		/// Attempts to update the state of a batch
 		/// </summary>
-		/// <param name="Job">Job to update</param>
-		/// <param name="BatchId">Unique id of the batch to update</param>
-		/// <param name="NewLogId">The new log file id</param>
-		/// <param name="NewState">New state of the jobstep</param>
-		/// <param name="NewError">New error state</param>
+		/// <param name="job">Job to update</param>
+		/// <param name="batchId">Unique id of the batch to update</param>
+		/// <param name="newLogId">The new log file id</param>
+		/// <param name="newState">New state of the jobstep</param>
+		/// <param name="newError">New error state</param>
 		/// <returns>The updated job, otherwise null</returns>
-		public async Task<IJob?> TryUpdateBatchAsync(IJob Job, SubResourceId BatchId, LogId? NewLogId = null, JobStepBatchState? NewState = null, JobStepBatchError? NewError = null)
+		public async Task<IJob?> TryUpdateBatchAsync(IJob job, SubResourceId batchId, LogId? newLogId = null, JobStepBatchState? newState = null, JobStepBatchError? newError = null)
 		{
-			IGraph Graph = await GetGraphAsync(Job);
-			return await Jobs.TryUpdateBatchAsync(Job, Graph, BatchId, NewLogId, NewState, NewError);
+			IGraph graph = await GetGraphAsync(job);
+			return await _jobs.TryUpdateBatchAsync(job, graph, batchId, newLogId, newState, newError);
 		}
 
 		/// <summary>
 		/// Update a jobstep state
 		/// </summary>
-		/// <param name="Job">Job to update</param>
-		/// <param name="BatchId">Unique id of the batch containing the step</param>
-		/// <param name="StepId">Unique id of the step to update</param>
-		/// <param name="NewState">New state of the jobstep</param>
-		/// <param name="NewOutcome">New outcome of the jobstep</param>
-		/// <param name="NewAbortRequested">New state of abort request</param>
-		/// <param name="NewAbortByUserId">New user that requested the abort</param>
-		/// <param name="NewLogId">New log id for the jobstep</param>
-		/// <param name="NewNotificationTriggerId">New notification trigger id for the jobstep</param>
-		/// <param name="NewRetryByUserId">Whether the step should be retried</param>
-		/// <param name="NewPriority">New priority for this step</param>
-		/// <param name="NewReports">New list of reports</param>
-		/// <param name="NewProperties">Property changes. Any properties with a null value will be removed.</param>
+		/// <param name="job">Job to update</param>
+		/// <param name="batchId">Unique id of the batch containing the step</param>
+		/// <param name="stepId">Unique id of the step to update</param>
+		/// <param name="newState">New state of the jobstep</param>
+		/// <param name="newOutcome">New outcome of the jobstep</param>
+		/// <param name="newAbortRequested">New state of abort request</param>
+		/// <param name="newAbortByUserId">New user that requested the abort</param>
+		/// <param name="newLogId">New log id for the jobstep</param>
+		/// <param name="newNotificationTriggerId">New notification trigger id for the jobstep</param>
+		/// <param name="newRetryByUserId">Whether the step should be retried</param>
+		/// <param name="newPriority">New priority for this step</param>
+		/// <param name="newReports">New list of reports</param>
+		/// <param name="newProperties">Property changes. Any properties with a null value will be removed.</param>
 		/// <returns>True if the job was updated, false if it was deleted in the meantime</returns>
-		public async Task<IJob?> UpdateStepAsync(IJob Job, SubResourceId BatchId, SubResourceId StepId, JobStepState NewState = JobStepState.Unspecified, JobStepOutcome NewOutcome = JobStepOutcome.Unspecified, bool? NewAbortRequested = null, UserId? NewAbortByUserId = null, LogId? NewLogId = null, ObjectId? NewNotificationTriggerId = null, UserId? NewRetryByUserId = null, Priority? NewPriority = null, List<Report>? NewReports = null, Dictionary<string, string?>? NewProperties = null)
+		public async Task<IJob?> UpdateStepAsync(IJob job, SubResourceId batchId, SubResourceId stepId, JobStepState newState = JobStepState.Unspecified, JobStepOutcome newOutcome = JobStepOutcome.Unspecified, bool? newAbortRequested = null, UserId? newAbortByUserId = null, LogId? newLogId = null, ObjectId? newNotificationTriggerId = null, UserId? newRetryByUserId = null, Priority? newPriority = null, List<Report>? newReports = null, Dictionary<string, string?>? newProperties = null)
 		{
-			using IScope TraceScope = GlobalTracer.Instance.BuildSpan("JobService.UpdateStepAsync").StartActive();
-			TraceScope.Span.SetTag("Job", Job.Id);
-			TraceScope.Span.SetTag("BatchId", BatchId);
-			TraceScope.Span.SetTag("StepId", StepId);
+			using IScope traceScope = GlobalTracer.Instance.BuildSpan("JobService.UpdateStepAsync").StartActive();
+			traceScope.Span.SetTag("Job", job.Id);
+			traceScope.Span.SetTag("BatchId", batchId);
+			traceScope.Span.SetTag("StepId", stepId);
 			
-			using IDisposable Scope = Logger.BeginScope("UpdateStepAsync({JobId})", Job.Id);
+			using IDisposable scope = _logger.BeginScope("UpdateStepAsync({JobId})", job.Id);
 			for (; ;)
 			{
-				IJob? NewJob = await TryUpdateStepAsync(Job, BatchId, StepId, NewState, NewOutcome, NewAbortRequested, NewAbortByUserId, NewLogId, NewNotificationTriggerId, NewRetryByUserId, NewPriority, NewReports, NewProperties);
-				if (NewJob != null)
+				IJob? newJob = await TryUpdateStepAsync(job, batchId, stepId, newState, newOutcome, newAbortRequested, newAbortByUserId, newLogId, newNotificationTriggerId, newRetryByUserId, newPriority, newReports, newProperties);
+				if (newJob != null)
 				{
-					return NewJob;
+					return newJob;
 				}
 
-				NewJob = await GetJobAsync(Job.Id);
-				if(NewJob == null)
+				newJob = await GetJobAsync(job.Id);
+				if(newJob == null)
 				{
 					return null;
 				}
 
-				Job = NewJob;
+				job = newJob;
 			}
 		}
 
 		/// <summary>
 		/// Update a jobstep state
 		/// </summary>
-		/// <param name="Job">Job to update</param>
-		/// <param name="BatchId">Unique id of the batch containing the step</param>
-		/// <param name="StepId">Unique id of the step to update</param>
-		/// <param name="NewState">New state of the jobstep</param>
-		/// <param name="NewOutcome">New outcome of the jobstep</param>
-		/// <param name="NewAbortRequested">New state for request abort</param>
-		/// <param name="NewAbortByUserId">New name of user that requested the abort</param>
-		/// <param name="NewLogId">New log id for the jobstep</param>
-		/// <param name="NewTriggerId">New trigger id for the jobstep</param>
-		/// <param name="NewRetryByUserId">Whether the step should be retried</param>
-		/// <param name="NewPriority">New priority for this step</param>
-		/// <param name="NewReports">New reports</param>
-		/// <param name="NewProperties">Property changes. Any properties with a null value will be removed.</param>
+		/// <param name="job">Job to update</param>
+		/// <param name="batchId">Unique id of the batch containing the step</param>
+		/// <param name="stepId">Unique id of the step to update</param>
+		/// <param name="newState">New state of the jobstep</param>
+		/// <param name="newOutcome">New outcome of the jobstep</param>
+		/// <param name="newAbortRequested">New state for request abort</param>
+		/// <param name="newAbortByUserId">New name of user that requested the abort</param>
+		/// <param name="newLogId">New log id for the jobstep</param>
+		/// <param name="newTriggerId">New trigger id for the jobstep</param>
+		/// <param name="newRetryByUserId">Whether the step should be retried</param>
+		/// <param name="newPriority">New priority for this step</param>
+		/// <param name="newReports">New reports</param>
+		/// <param name="newProperties">Property changes. Any properties with a null value will be removed.</param>
 		/// <returns>True if the job was updated, false if it was deleted in the meantime</returns>
-		public async Task<IJob?> TryUpdateStepAsync(IJob Job, SubResourceId BatchId, SubResourceId StepId, JobStepState NewState = JobStepState.Unspecified, JobStepOutcome NewOutcome = JobStepOutcome.Unspecified, bool? NewAbortRequested = null, UserId? NewAbortByUserId = null, LogId? NewLogId = null, ObjectId? NewTriggerId = null, UserId? NewRetryByUserId = null, Priority? NewPriority = null, List<Report>? NewReports = null, Dictionary<string, string?>? NewProperties = null)
+		public async Task<IJob?> TryUpdateStepAsync(IJob job, SubResourceId batchId, SubResourceId stepId, JobStepState newState = JobStepState.Unspecified, JobStepOutcome newOutcome = JobStepOutcome.Unspecified, bool? newAbortRequested = null, UserId? newAbortByUserId = null, LogId? newLogId = null, ObjectId? newTriggerId = null, UserId? newRetryByUserId = null, Priority? newPriority = null, List<Report>? newReports = null, Dictionary<string, string?>? newProperties = null)
 		{
-			using IScope TraceScope = GlobalTracer.Instance.BuildSpan("JobService.TryUpdateStepAsync").StartActive();
-			TraceScope.Span.SetTag("Job", Job.Id);
-			TraceScope.Span.SetTag("BatchId", BatchId);
-			TraceScope.Span.SetTag("StepId", StepId);
+			using IScope traceScope = GlobalTracer.Instance.BuildSpan("JobService.TryUpdateStepAsync").StartActive();
+			traceScope.Span.SetTag("Job", job.Id);
+			traceScope.Span.SetTag("BatchId", batchId);
+			traceScope.Span.SetTag("StepId", stepId);
 
-			using IDisposable Scope = Logger.BeginScope("TryUpdateStepAsync({JobId})", Job.Id);
+			using IDisposable scope = _logger.BeginScope("TryUpdateStepAsync({JobId})", job.Id);
 
 			// Get the graph for this job
-			IGraph Graph = await GetGraphAsync(Job);
+			IGraph graph = await GetGraphAsync(job);
 
 			// Make sure the job state is set to unspecified
-			JobStepState OldState = JobStepState.Unspecified;
-			JobStepOutcome OldOutcome = JobStepOutcome.Unspecified;
-			if (NewState != JobStepState.Unspecified)
+			JobStepState oldState = JobStepState.Unspecified;
+			JobStepOutcome oldOutcome = JobStepOutcome.Unspecified;
+			if (newState != JobStepState.Unspecified)
 			{
-				if (Job.TryGetStep(BatchId, StepId, out IJobStep? Step))
+				if (job.TryGetStep(batchId, stepId, out IJobStep? step))
 				{
-					OldState = Step.State;
-					OldOutcome = Step.Outcome;
+					oldState = step.State;
+					oldOutcome = step.Outcome;
 				}
 			}
 
 			// If we're changing the state of a step, capture the label states beforehand so we can update any that change
-			IReadOnlyList<(LabelState, LabelOutcome)>? OldLabelStates = null;
-			if (OldState != NewState || OldOutcome != NewOutcome)
+			IReadOnlyList<(LabelState, LabelOutcome)>? oldLabelStates = null;
+			if (oldState != newState || oldOutcome != newOutcome)
 			{
-				OldLabelStates = Job.GetLabelStates(Graph);
+				oldLabelStates = job.GetLabelStates(graph);
 			}
 
 			// Update the step
-			IJob? NewJob = await Jobs.TryUpdateStepAsync(Job, Graph, BatchId, StepId, NewState, NewOutcome, NewAbortRequested, NewAbortByUserId, NewLogId, NewTriggerId, NewRetryByUserId, NewPriority, NewReports, NewProperties);
-			if (NewJob != null)
+			IJob? newJob = await _jobs.TryUpdateStepAsync(job, graph, batchId, stepId, newState, newOutcome, newAbortRequested, newAbortByUserId, newLogId, newTriggerId, newRetryByUserId, newPriority, newReports, newProperties);
+			if (newJob != null)
 			{
-				Job = NewJob;
+				job = newJob;
 
-				using IScope DdScope = GlobalTracer.Instance.BuildSpan("TryUpdateStepAsync").StartActive();
-				DdScope.Span.SetTag("JobId", Job.Id.ToString());
-				DdScope.Span.SetTag("BatchId", BatchId.ToString());
-				DdScope.Span.SetTag("StepId", StepId.ToString());
+				using IScope ddScope = GlobalTracer.Instance.BuildSpan("TryUpdateStepAsync").StartActive();
+				ddScope.Span.SetTag("JobId", job.Id.ToString());
+				ddScope.Span.SetTag("BatchId", batchId.ToString());
+				ddScope.Span.SetTag("StepId", stepId.ToString());
 				
-				if (OldState != NewState || OldOutcome != NewOutcome)
+				if (oldState != newState || oldOutcome != newOutcome)
 				{
-					Logger.LogDebug("Transitioned job {JobId}, batch {BatchId}, step {StepId} from {OldState} to {NewState}", Job.Id, BatchId, StepId, OldState, NewState);
+					_logger.LogDebug("Transitioned job {JobId}, batch {BatchId}, step {StepId} from {OldState} to {NewState}", job.Id, batchId, stepId, oldState, newState);
 
 					// Send any updates for modified badges
-					if (OldLabelStates != null)
+					if (oldLabelStates != null)
 					{
 						using IScope _ = GlobalTracer.Instance.BuildSpan("Send badge updates").StartActive();
-						IReadOnlyList<(LabelState, LabelOutcome)> NewLabelStates = Job.GetLabelStates(Graph);
-						OnLabelUpdate?.Invoke(Job, OldLabelStates, NewLabelStates);
-						await JobTaskSource.UpdateUgsBadges(Job, Graph, OldLabelStates, NewLabelStates);
+						IReadOnlyList<(LabelState, LabelOutcome)> newLabelStates = job.GetLabelStates(graph);
+						OnLabelUpdate?.Invoke(job, oldLabelStates, newLabelStates);
+						await _jobTaskSource.UpdateUgsBadges(job, graph, oldLabelStates, newLabelStates);
 					}
 
 					// Submit the change if auto-submit is enabled
-					if (Job.PreflightChange != 0 && NewState == JobStepState.Completed)
+					if (job.PreflightChange != 0 && newState == JobStepState.Completed)
 					{
-						(JobStepState State, JobStepOutcome Outcome) = Job.GetTargetState();
-						if (State == JobStepState.Completed)
+						(JobStepState state, JobStepOutcome outcome) = job.GetTargetState();
+						if (state == JobStepState.Completed)
 						{
-							if (Job.AutoSubmit && Outcome == JobStepOutcome.Success && Job.AbortedByUserId == null)
+							if (job.AutoSubmit && outcome == JobStepOutcome.Success && job.AbortedByUserId == null)
 							{
-								Job = await AutoSubmitChangeAsync(Job, Graph);
+								job = await AutoSubmitChangeAsync(job, graph);
 							}
-							else if (Job.ClonedPreflightChange != 0)
+							else if (job.ClonedPreflightChange != 0)
 							{
-								IStream? Stream = await StreamService.GetCachedStream(Job.StreamId);
-								if (Stream != null)
+								IStream? stream = await _streamService.GetCachedStream(job.StreamId);
+								if (stream != null)
 								{
-									await DeleteShelvedChangeAsync(Stream.ClusterName, Job.ClonedPreflightChange);
+									await DeleteShelvedChangeAsync(stream.ClusterName, job.ClonedPreflightChange);
 								}
 							}
 						}
 					}
 
 					// Notify subscribers
-					if (NewState == JobStepState.Completed)
+					if (newState == JobStepState.Completed)
 					{
 						using IScope _ = GlobalTracer.Instance.BuildSpan("Notify subscribers").StartActive();
-						OnJobStepComplete?.Invoke(Job, Graph, BatchId, StepId);
+						OnJobStepComplete?.Invoke(job, graph, batchId, stepId);
 					}
 
 					// Create any downstream jobs
-					if (NewState == JobStepState.Completed && NewOutcome != JobStepOutcome.Failure)
+					if (newState == JobStepState.Completed && newOutcome != JobStepOutcome.Failure)
 					{
 						using IScope _ = GlobalTracer.Instance.BuildSpan("Create downstream jobs").StartActive();
-						for (int Idx = 0; Idx < Job.ChainedJobs.Count; Idx++)
+						for (int idx = 0; idx < job.ChainedJobs.Count; idx++)
 						{
-							IChainedJob JobTrigger = Job.ChainedJobs[Idx];
-							if (JobTrigger.JobId == null)
+							IChainedJob jobTrigger = job.ChainedJobs[idx];
+							if (jobTrigger.JobId == null)
 							{
-								JobStepOutcome JobTriggerOutcome = Job.GetTargetOutcome(Graph, JobTrigger.Target);
-								if (JobTriggerOutcome == JobStepOutcome.Success || JobTriggerOutcome == JobStepOutcome.Warnings)
+								JobStepOutcome jobTriggerOutcome = job.GetTargetOutcome(graph, jobTrigger.Target);
+								if (jobTriggerOutcome == JobStepOutcome.Success || jobTriggerOutcome == JobStepOutcome.Warnings)
 								{
-									Job = await FireJobTriggerAsync(Job, Graph, JobTrigger) ?? Job;
+									job = await FireJobTriggerAsync(job, graph, jobTrigger) ?? job;
 								}
 							}
 						}
 					}
 
 					// Update the jobstep ref if it completed
-					if (NewState == JobStepState.Running || NewState == JobStepState.Completed || NewState == JobStepState.Aborted)
+					if (newState == JobStepState.Running || newState == JobStepState.Completed || newState == JobStepState.Aborted)
 					{
 						using IScope _ = GlobalTracer.Instance.BuildSpan("Update job step ref").StartActive();
-						if (Job.TryGetBatch(BatchId, out IJobStepBatch? Batch) && Batch.TryGetStep(StepId, out IJobStep? Step) && Step.StartTimeUtc != null)
+						if (job.TryGetBatch(batchId, out IJobStepBatch? batch) && batch.TryGetStep(stepId, out IJobStep? step) && step.StartTimeUtc != null)
 						{
-							await JobStepRefs.UpdateAsync(Job, Batch, Step, Graph);
+							await _jobStepRefs.UpdateAsync(job, batch, step, graph);
 						}
 					}
 
 					// Update any issues that depend on this step
-					if (NewState == JobStepState.Completed && Job.UpdateIssues)
+					if (newState == JobStepState.Completed && job.UpdateIssues)
 					{
-						if (IssueService != null)
+						if (_issueService != null)
 						{
 							using IScope _ = GlobalTracer.Instance.BuildSpan("Update issues (V2)").StartActive();
 							try
 							{
-								await IssueService.UpdateCompleteStep(Job, Graph, BatchId, StepId);
+								await _issueService.UpdateCompleteStep(job, graph, batchId, stepId);
 							}
-							catch(Exception Ex)
+							catch(Exception ex)
 							{
-								Logger.LogError(Ex, "Exception while updating issue service for job {JobId}:{BatchId}:{StepId}: {Message}", Job.Id, BatchId, StepId, Ex.Message);
+								_logger.LogError(ex, "Exception while updating issue service for job {JobId}:{BatchId}:{StepId}: {Message}", job.Id, batchId, stepId, ex.Message);
 							}
 						}
 					}
 				}
 
-				using (IScope DispatchScope = GlobalTracer.Instance.BuildSpan("Update queued jobs").StartActive())
+				using (IScope dispatchScope = GlobalTracer.Instance.BuildSpan("Update queued jobs").StartActive())
 				{
 					// Notify the dispatch service that the job has changed
-					JobTaskSource.UpdateQueuedJob(Job, Graph);
+					_jobTaskSource.UpdateQueuedJob(job, graph);
 				}
 
-				return Job;
+				return job;
 			}
 
 			return null;
@@ -1092,130 +1050,128 @@ namespace Horde.Build.Services
 		/// <summary>
 		/// Submit the given change for a preflight
 		/// </summary>
-		/// <param name="Job">The job being run</param>
-		/// <param name="Graph">Graph for the job</param>
+		/// <param name="job">The job being run</param>
+		/// <param name="graph">Graph for the job</param>
 		/// <returns></returns>
-		[SuppressMessage("Design", "CA1031:Do not catch general exception types", Justification = "<Pending>")]
-		private async Task<IJob> AutoSubmitChangeAsync(IJob Job, IGraph Graph)
+		private async Task<IJob> AutoSubmitChangeAsync(IJob job, IGraph graph)
 		{
-			using IScope TraceScope = GlobalTracer.Instance.BuildSpan("JobService.AutoSubmitChangeAsync").StartActive();
-			TraceScope.Span.SetTag("Job", Job.Id);
-			TraceScope.Span.SetTag("Graph", Graph.Id);
+			using IScope traceScope = GlobalTracer.Instance.BuildSpan("JobService.AutoSubmitChangeAsync").StartActive();
+			traceScope.Span.SetTag("Job", job.Id);
+			traceScope.Span.SetTag("Graph", graph.Id);
 			
-			int? Change;
-			string Message;
+			int? change;
+			string message;
 			try
 			{
-				IStream? Stream = await StreamService.GetCachedStream(Job.StreamId);
-				if (Stream != null)
+				IStream? stream = await _streamService.GetCachedStream(job.StreamId);
+				if (stream != null)
 				{
-					int ClonedPreflightChange = Job.ClonedPreflightChange;
-					if (ClonedPreflightChange == 0)
+					int clonedPreflightChange = job.ClonedPreflightChange;
+					if (clonedPreflightChange == 0)
 					{
-						if (ShouldClonePreflightChange(Job.StreamId))
+						if (ShouldClonePreflightChange(job.StreamId))
 						{
-							ClonedPreflightChange = await CloneShelvedChangeAsync(Stream.ClusterName, Job.PreflightChange);
+							clonedPreflightChange = await CloneShelvedChangeAsync(stream.ClusterName, job.PreflightChange);
 						}
 						else
 						{
-							ClonedPreflightChange = Job.PreflightChange;
+							clonedPreflightChange = job.PreflightChange;
 						}
 					}
 
-					Logger.LogInformation("Updating description for {ClonedPreflightChange}", ClonedPreflightChange);
+					_logger.LogInformation("Updating description for {ClonedPreflightChange}", clonedPreflightChange);
 
-					ChangeDetails Details = await PerforceService.GetChangeDetailsAsync(Stream.ClusterName, Stream.Name, ClonedPreflightChange, null);
-					await PerforceService.UpdateChangelistDescription(Stream.ClusterName, ClonedPreflightChange, Details.Description.TrimEnd() + $"\n#preflight {Job.Id}");
+					ChangeDetails details = await _perforceService.GetChangeDetailsAsync(stream.ClusterName, stream.Name, clonedPreflightChange, null);
+					await _perforceService.UpdateChangelistDescription(stream.ClusterName, clonedPreflightChange, details.Description.TrimEnd() + $"\n#preflight {job.Id}");
 
-					Logger.LogInformation("Submitting change {Change} (through {ChangeCopy}) after successful completion of {JobId}", Job.PreflightChange, ClonedPreflightChange, Job.Id);
-					(Change, Message) = await PerforceService.SubmitShelvedChangeAsync(Stream.ClusterName, ClonedPreflightChange, Job.PreflightChange);
+					_logger.LogInformation("Submitting change {Change} (through {ChangeCopy}) after successful completion of {JobId}", job.PreflightChange, clonedPreflightChange, job.Id);
+					(change, message) = await _perforceService.SubmitShelvedChangeAsync(stream.ClusterName, clonedPreflightChange, job.PreflightChange);
 
-					Logger.LogInformation("Attempt to submit {Change} (through {ChangeCopy}): {Message}", Job.PreflightChange, ClonedPreflightChange, Message);
+					_logger.LogInformation("Attempt to submit {Change} (through {ChangeCopy}): {Message}", job.PreflightChange, clonedPreflightChange, message);
 
-					if (ShouldClonePreflightChange(Job.StreamId))
+					if (ShouldClonePreflightChange(job.StreamId))
 					{
-						if (Change != null && Job.ClonedPreflightChange != 0)
+						if (change != null && job.ClonedPreflightChange != 0)
 						{
-							await DeleteShelvedChangeAsync(Stream.ClusterName, Job.PreflightChange);
+							await DeleteShelvedChangeAsync(stream.ClusterName, job.PreflightChange);
 						}
 					}
 					else
 					{
-						if (Change != null && Job.PreflightChange != 0)
+						if (change != null && job.PreflightChange != 0)
 						{
-							await DeleteShelvedChangeAsync(Stream.ClusterName, Job.PreflightChange);
+							await DeleteShelvedChangeAsync(stream.ClusterName, job.PreflightChange);
 						}
 					}
 				}
 				else
 				{
-					(Change, Message) = ((int?)null, "Stream not found");
+					(change, message) = ((int?)null, "Stream not found");
 				}
 			}
-			catch (Exception Ex)
+			catch (Exception ex)
 			{
-				(Change, Message) = ((int?)null, "Internal error");
-				Logger.LogError(Ex, "Unable to submit shelved change");
+				(change, message) = ((int?)null, "Internal error");
+				_logger.LogError(ex, "Unable to submit shelved change");
 			}
 
 			for (; ; )
 			{
-				IJob? NewJob = await Jobs.TryUpdateJobAsync(Job, Graph, AutoSubmitChange: Change, AutoSubmitMessage: Message);
-				if (NewJob != null)
+				IJob? newJob = await _jobs.TryUpdateJobAsync(job, graph, autoSubmitChange: change, autoSubmitMessage: message);
+				if (newJob != null)
 				{
-					return NewJob;
+					return newJob;
 				}
 
-				NewJob = await GetJobAsync(Job.Id);
-				if (NewJob == null)
+				newJob = await GetJobAsync(job.Id);
+				if (newJob == null)
 				{
-					return Job;
+					return job;
 				}
 
-				Job = NewJob;
+				job = newJob;
 			}
 		}
 
 		/// <summary>
 		/// Clone a shelved changelist for running a preflight
 		/// </summary>
-		/// <param name="ClusterName">Name of the Perforce cluster</param>
-		/// <param name="Change">The changelist to clone</param>
+		/// <param name="clusterName">Name of the Perforce cluster</param>
+		/// <param name="change">The changelist to clone</param>
 		/// <returns></returns>
-		private async Task<int> CloneShelvedChangeAsync(string ClusterName, int Change)
+		private async Task<int> CloneShelvedChangeAsync(string clusterName, int change)
 		{
-			int ClonedChange;
+			int clonedChange;
 			try
 			{
-				ClonedChange = await PerforceService.DuplicateShelvedChangeAsync(ClusterName, Change);
-				Logger.LogInformation("CL {Change} was duplicated into {ClonedChange}", Change, ClonedChange);
+				clonedChange = await _perforceService.DuplicateShelvedChangeAsync(clusterName, change);
+				_logger.LogInformation("CL {Change} was duplicated into {ClonedChange}", change, clonedChange);
 			}
-			catch (Exception Ex)
+			catch (Exception ex)
 			{
-				Logger.LogError(Ex, "Unable to CL {Change} for preflight: {Message}", Change, Ex.Message);
+				_logger.LogError(ex, "Unable to CL {Change} for preflight: {Message}", change, ex.Message);
 				throw;
 			}
-			return ClonedChange;
+			return clonedChange;
 		}
 
 		/// <summary>
 		/// Deletes a shelved changelist, catching and logging any exceptions
 		/// </summary>
-		/// <param name="ClusterName"></param>
-		/// <param name="Change">The changelist to delete</param>
+		/// <param name="clusterName"></param>
+		/// <param name="change">The changelist to delete</param>
 		/// <returns>True if the change was deleted successfully, false otherwise</returns>
-		[SuppressMessage("Design", "CA1031:Do not catch general exception types", Justification = "<Pending>")]
-		private async Task<bool> DeleteShelvedChangeAsync(string ClusterName, int Change)
+		private async Task<bool> DeleteShelvedChangeAsync(string clusterName, int change)
 		{
-			Logger.LogInformation("Removing shelf {Change}", Change);
+			_logger.LogInformation("Removing shelf {Change}", change);
 			try
 			{
-				await PerforceService.DeleteShelvedChangeAsync(ClusterName, Change);
+				await _perforceService.DeleteShelvedChangeAsync(clusterName, change);
 				return true;
 			}
-			catch (Exception Ex)
+			catch (Exception ex)
 			{
-				Logger.LogError(Ex, "Unable to delete shelved change {Change}", Change);
+				_logger.LogError(ex, "Unable to delete shelved change {Change}", change);
 				return false;
 			}
 		}
@@ -1223,63 +1179,63 @@ namespace Horde.Build.Services
 		/// <summary>
 		/// Fires a trigger for a chained job
 		/// </summary>
-		/// <param name="Job">The job containing the trigger</param>
-		/// <param name="Graph">Graph for the job containing the trigger</param>
-		/// <param name="JobTrigger">The trigger object to fire</param>
+		/// <param name="job">The job containing the trigger</param>
+		/// <param name="graph">Graph for the job containing the trigger</param>
+		/// <param name="jobTrigger">The trigger object to fire</param>
 		/// <returns>New job instance</returns>
-		private async Task<IJob?> FireJobTriggerAsync(IJob Job, IGraph Graph, IChainedJob JobTrigger)
+		private async Task<IJob?> FireJobTriggerAsync(IJob job, IGraph graph, IChainedJob jobTrigger)
 		{
-			using IScope TraceScope = GlobalTracer.Instance.BuildSpan("JobService.FireJobTriggerAsync").StartActive();
-			TraceScope.Span.SetTag("Job", Job.Id);
-			TraceScope.Span.SetTag("Graph", Graph.Id);
-			TraceScope.Span.SetTag("JobTrigger.Target", JobTrigger.Target);
-			TraceScope.Span.SetTag("JobTrigger.JobId", JobTrigger.JobId);
-			TraceScope.Span.SetTag("JobTrigger.TemplateRefId", JobTrigger.TemplateRefId);
+			using IScope traceScope = GlobalTracer.Instance.BuildSpan("JobService.FireJobTriggerAsync").StartActive();
+			traceScope.Span.SetTag("Job", job.Id);
+			traceScope.Span.SetTag("Graph", graph.Id);
+			traceScope.Span.SetTag("JobTrigger.Target", jobTrigger.Target);
+			traceScope.Span.SetTag("JobTrigger.JobId", jobTrigger.JobId);
+			traceScope.Span.SetTag("JobTrigger.TemplateRefId", jobTrigger.TemplateRefId);
 			
 			for (; ; )
 			{
 				// Update the job
-				JobId ChainedJobId = JobId.GenerateNewId();
+				JobId chainedJobId = JobId.GenerateNewId();
 
-				IJob? NewJob = await Jobs.TryUpdateJobAsync(Job, Graph, JobTrigger: new KeyValuePair<TemplateRefId, JobId>(JobTrigger.TemplateRefId, ChainedJobId));
-				if(NewJob != null)
+				IJob? newJob = await _jobs.TryUpdateJobAsync(job, graph, jobTrigger: new KeyValuePair<TemplateRefId, JobId>(jobTrigger.TemplateRefId, chainedJobId));
+				if(newJob != null)
 				{
-					IStream? Stream = await StreamService.GetStreamAsync(NewJob.StreamId);
-					if(Stream == null)
+					IStream? stream = await _streamService.GetStreamAsync(newJob.StreamId);
+					if(stream == null)
 					{
-						Logger.LogWarning("Cannot find stream {StreamId} for downstream job", NewJob.StreamId);
+						_logger.LogWarning("Cannot find stream {StreamId} for downstream job", newJob.StreamId);
 						break;
 					}
 
-					TemplateRef? TemplateRef;
-					if (!Stream.Templates.TryGetValue(JobTrigger.TemplateRefId, out TemplateRef))
+					TemplateRef? templateRef;
+					if (!stream.Templates.TryGetValue(jobTrigger.TemplateRefId, out templateRef))
 					{
-						Logger.LogWarning("Cannot find template {TemplateId} in stream {StreamId}", JobTrigger.TemplateRefId, NewJob.StreamId);
+						_logger.LogWarning("Cannot find template {TemplateId} in stream {StreamId}", jobTrigger.TemplateRefId, newJob.StreamId);
 						break;
 					}
 
-					ITemplate? Template = await TemplateCollection.GetAsync(TemplateRef.Hash);
-					if (Template == null)
+					ITemplate? template = await _templateCollection.GetAsync(templateRef.Hash);
+					if (template == null)
 					{
-						Logger.LogWarning("Cannot find template {TemplateHash}", TemplateRef.Hash);
+						_logger.LogWarning("Cannot find template {TemplateHash}", templateRef.Hash);
 						break;
 					}
 
-					IGraph TriggerGraph = await Graphs.AddAsync(Template);
-					Logger.LogInformation("Creating downstream job {ChainedJobId} from job {JobId}", ChainedJobId, NewJob.Id);
+					IGraph triggerGraph = await _graphs.AddAsync(template);
+					_logger.LogInformation("Creating downstream job {ChainedJobId} from job {JobId}", chainedJobId, newJob.Id);
 
-					await CreateJobAsync(ChainedJobId, Stream, JobTrigger.TemplateRefId, TemplateRef.Hash, TriggerGraph, TemplateRef.Name, NewJob.Change, NewJob.CodeChange, NewJob.PreflightChange, NewJob.ClonedPreflightChange, NewJob.StartedByUserId, Template.Priority, null, NewJob.UpdateIssues, NewJob.PromoteIssuesByDefault, TemplateRef.ChainedJobs, false, false, TemplateRef.NotificationChannel, TemplateRef.NotificationChannelFilter, Template.Arguments);
-					return NewJob;
+					await CreateJobAsync(chainedJobId, stream, jobTrigger.TemplateRefId, templateRef.Hash, triggerGraph, templateRef.Name, newJob.Change, newJob.CodeChange, newJob.PreflightChange, newJob.ClonedPreflightChange, newJob.StartedByUserId, template.Priority, null, newJob.UpdateIssues, newJob.PromoteIssuesByDefault, templateRef.ChainedJobs, false, false, templateRef.NotificationChannel, templateRef.NotificationChannelFilter, template.Arguments);
+					return newJob;
 				}
 
 				// Fetch the job again
-				NewJob = await Jobs.GetAsync(Job.Id);
-				if(NewJob == null)
+				newJob = await _jobs.GetAsync(job.Id);
+				if(newJob == null)
 				{
 					break;
 				}
 
-				Job = NewJob;
+				job = newJob;
 			}
 			return null;
 		}
@@ -1287,89 +1243,89 @@ namespace Horde.Build.Services
 		/// <summary>
 		/// Determines if the user is authorized to perform an action on a particular stream
 		/// </summary>
-		/// <param name="Acl">The ACL to check</param>
-		/// <param name="StreamId">The stream containing this job</param>
-		/// <param name="Action">The action being performed</param>
-		/// <param name="User">The principal to authorize</param>
-		/// <param name="Cache">Cache of stream permissions</param>
+		/// <param name="acl">The ACL to check</param>
+		/// <param name="streamId">The stream containing this job</param>
+		/// <param name="action">The action being performed</param>
+		/// <param name="user">The principal to authorize</param>
+		/// <param name="cache">Cache of stream permissions</param>
 		/// <returns>True if the action is authorized</returns>
-		private Task<bool> AuthorizeAsync(Acl? Acl, StreamId StreamId, AclAction Action, ClaimsPrincipal User, StreamPermissionsCache? Cache)
+		private Task<bool> AuthorizeAsync(Acl? acl, StreamId streamId, AclAction action, ClaimsPrincipal user, StreamPermissionsCache? cache)
 		{
 			// Do the regular authorization
-			bool? Result = Acl?.Authorize(Action, User);
-			if (Result == null)
+			bool? result = acl?.Authorize(action, user);
+			if (result == null)
 			{
-				return StreamService.AuthorizeAsync(StreamId, Action, User, Cache);
+				return _streamService.AuthorizeAsync(streamId, action, user, cache);
 			}
 			else
 			{
-				return Task.FromResult(Result.Value);
+				return Task.FromResult(result.Value);
 			}
 		}
 
 		/// <summary>
 		/// Determines if the user is authorized to perform an action on a particular stream
 		/// </summary>
-		/// <param name="Job">The job to check</param>
-		/// <param name="Action">The action being performed</param>
-		/// <param name="User">The principal to authorize</param>
-		/// <param name="Cache">Cache of stream permissions</param>
+		/// <param name="job">The job to check</param>
+		/// <param name="action">The action being performed</param>
+		/// <param name="user">The principal to authorize</param>
+		/// <param name="cache">Cache of stream permissions</param>
 		/// <returns>True if the action is authorized</returns>
-		public Task<bool> AuthorizeAsync(IJob Job, AclAction Action, ClaimsPrincipal User, StreamPermissionsCache? Cache)
+		public Task<bool> AuthorizeAsync(IJob job, AclAction action, ClaimsPrincipal user, StreamPermissionsCache? cache)
 		{
-			return AuthorizeAsync(Job.Acl, Job.StreamId, Action, User, Cache);
+			return AuthorizeAsync(job.Acl, job.StreamId, action, user, cache);
 		}
 
 		/// <summary>
 		/// Determines if the user is authorized to perform an action on a particular stream
 		/// </summary>
-		/// <param name="JobId">The job to check</param>
-		/// <param name="Action">The action being performed</param>
-		/// <param name="User">The principal to authorize</param>
-		/// <param name="Cache">Cache of job permissions</param>
+		/// <param name="jobId">The job to check</param>
+		/// <param name="action">The action being performed</param>
+		/// <param name="user">The principal to authorize</param>
+		/// <param name="cache">Cache of job permissions</param>
 		/// <returns>True if the action is authorized</returns>
-		public async Task<bool> AuthorizeAsync(JobId JobId, AclAction Action, ClaimsPrincipal User, JobPermissionsCache? Cache)
+		public async Task<bool> AuthorizeAsync(JobId jobId, AclAction action, ClaimsPrincipal user, JobPermissionsCache? cache)
 		{
-			IJobPermissions? Permissions;
-			if (Cache == null)
+			IJobPermissions? permissions;
+			if (cache == null)
 			{
-				Permissions = await GetJobPermissionsAsync(JobId);
+				permissions = await GetJobPermissionsAsync(jobId);
 			}
-			else if (!Cache.Jobs.TryGetValue(JobId, out Permissions))
+			else if (!cache.Jobs.TryGetValue(jobId, out permissions))
 			{
-				Permissions = await GetJobPermissionsAsync(JobId);
-				Cache.Jobs.Add(JobId, Permissions);
+				permissions = await GetJobPermissionsAsync(jobId);
+				cache.Jobs.Add(jobId, permissions);
 			}
-			return Permissions != null && await AuthorizeAsync(Permissions.Acl, Permissions.StreamId, Action, User, Cache);
+			return permissions != null && await AuthorizeAsync(permissions.Acl, permissions.StreamId, action, user, cache);
 		}
 
 		/// <summary>
 		/// Determines if the user is authorized to perform an action on a particular stream
 		/// </summary>
-		/// <param name="Job">The job to check</param>
-		/// <param name="User">The principal to authorize</param>
+		/// <param name="job">The job to check</param>
+		/// <param name="user">The principal to authorize</param>
 		/// <returns>True if the action is authorized</returns>
-		public static bool AuthorizeSession(IJob Job, ClaimsPrincipal User)
+		public static bool AuthorizeSession(IJob job, ClaimsPrincipal user)
 		{
-			return Job.Batches.Any(x => AuthorizeSession(x, User));
+			return job.Batches.Any(x => AuthorizeSession(x, user));
 		}
 
 		/// <summary>
 		/// Determines if the user is authorized to perform an action on a particular stream
 		/// </summary>
-		/// <param name="Batch">The batch being executed</param>
-		/// <param name="User">The principal to authorize</param>
+		/// <param name="batch">The batch being executed</param>
+		/// <param name="user">The principal to authorize</param>
 		/// <returns>True if the action is authorized</returns>
-		public static bool AuthorizeSession(IJobStepBatch Batch, ClaimsPrincipal User)
+		public static bool AuthorizeSession(IJobStepBatch batch, ClaimsPrincipal user)
 		{
-			if(Batch.SessionId != null)
+			if(batch.SessionId != null)
 			{
-				foreach (Claim Claim in User.Claims)
+				foreach (Claim claim in user.Claims)
 				{
-					if(Claim.Type == HordeClaimTypes.AgentSessionId)
+					if(claim.Type == HordeClaimTypes.AgentSessionId)
 					{
-						SessionId SessionIdValue;
-						if (SessionId.TryParse(Claim.Value, out SessionIdValue) && SessionIdValue == Batch.SessionId.Value)
+						SessionId sessionIdValue;
+						if (SessionId.TryParse(claim.Value, out sessionIdValue) && sessionIdValue == batch.SessionId.Value)
 						{
 							return true;
 						}
@@ -1382,11 +1338,11 @@ namespace Horde.Build.Services
 		/// <summary>
 		/// Get a list of each batch's session ID formatted as a string for debugging purposes
 		/// </summary>
-		/// <param name="Job">The job to list</param>
+		/// <param name="job">The job to list</param>
 		/// <returns>List of session IDs</returns>
-		public static string GetAllBatchSessionIds(IJob Job)
+		public static string GetAllBatchSessionIds(IJob job)
 		{
-			return string.Join(",", Job.Batches.Select(b => b.SessionId));
+			return String.Join(",", job.Batches.Select(b => b.SessionId));
 		}
 	}
 }
