Horde: Fix coding conventions within Horde.Build.

diff --git a/Engine/Source/Programs/Horde/Horde.Build/Jobs/Impl/JobStepRefCollection.cs b/Engine/Source/Programs/Horde/Horde.Build/Jobs/Impl/JobStepRefCollection.cs
--- a/Engine/Source/Programs/Horde/Horde.Build/Jobs/Impl/JobStepRefCollection.cs
+++ b/Engine/Source/Programs/Horde/Horde.Build/Jobs/Impl/JobStepRefCollection.cs
@@ -1,16 +1,15 @@
 // Copyright Epic Games, Inc. All Rights Reserved.
 
-using HordeCommon;
-using Horde.Build.Models;
-using Horde.Build.Services;
-using MongoDB.Bson;
-using MongoDB.Bson.Serialization.Attributes;
-using MongoDB.Driver;
 using System;
 using System.Collections.Generic;
 using System.Linq;
 using System.Threading.Tasks;
+using Horde.Build.Models;
+using Horde.Build.Services;
 using Horde.Build.Utilities;
+using HordeCommon;
+using MongoDB.Bson.Serialization.Attributes;
+using MongoDB.Driver;
 
 namespace Horde.Build.Collections.Impl
 {
@@ -66,86 +65,86 @@ namespace Horde.Build.Collections.Impl
 			DateTime? IJobStepRef.FinishTimeUtc => FinishTimeUtc ?? FinishTime?.UtcDateTime;
 			string IJobStepRef.NodeName => Name;
 
-			public JobStepRef(JobStepRefId Id, string JobName, string NodeName, StreamId StreamId, TemplateRefId TemplateId, int Change, LogId? LogId, PoolId? PoolId, AgentId? AgentId, JobStepOutcome? Outcome, int? LastSuccess, int? LastWarning, float BatchWaitTime, float BatchInitTime, DateTime StartTimeUtc, DateTime? FinishTimeUtc)
+			public JobStepRef(JobStepRefId id, string jobName, string nodeName, StreamId streamId, TemplateRefId templateId, int change, LogId? logId, PoolId? poolId, AgentId? agentId, JobStepOutcome? outcome, int? lastSuccess, int? lastWarning, float batchWaitTime, float batchInitTime, DateTime startTimeUtc, DateTime? finishTimeUtc)
 			{
-				this.Id = Id;
-				this.JobName = JobName;
-				this.Name = NodeName;
-				this.StreamId = StreamId;
-				this.TemplateId = TemplateId;
-				this.Change = Change;
-				this.LogId = LogId;
-				this.PoolId = PoolId;
-				this.AgentId = AgentId;
-				this.Outcome = Outcome;
-				this.LastSuccess = LastSuccess;
-				this.LastWarning = LastWarning;
-				this.BatchWaitTime = BatchWaitTime;
-				this.BatchInitTime = BatchInitTime;
-				this.StartTime = StartTimeUtc;
-				this.StartTimeUtc = StartTimeUtc;
-				this.FinishTime = FinishTimeUtc;
-				this.FinishTimeUtc = FinishTimeUtc;
+				Id = id;
+				JobName = jobName;
+				Name = nodeName;
+				StreamId = streamId;
+				TemplateId = templateId;
+				Change = change;
+				LogId = logId;
+				PoolId = poolId;
+				AgentId = agentId;
+				Outcome = outcome;
+				LastSuccess = lastSuccess;
+				LastWarning = lastWarning;
+				BatchWaitTime = batchWaitTime;
+				BatchInitTime = batchInitTime;
+				StartTime = startTimeUtc;
+				StartTimeUtc = startTimeUtc;
+				FinishTime = finishTimeUtc;
+				FinishTimeUtc = finishTimeUtc;
 			}
 		}
 
-		IMongoCollection<JobStepRef> JobStepRefs;
+		readonly IMongoCollection<JobStepRef> _jobStepRefs;
 
 		/// <summary>
 		/// Constructor
 		/// </summary>
-		/// <param name="DatabaseService">The database service instance</param>
-		public JobStepRefCollection(DatabaseService DatabaseService)
+		/// <param name="databaseService">The database service instance</param>
+		public JobStepRefCollection(DatabaseService databaseService)
 		{
-			JobStepRefs = DatabaseService.GetCollection<JobStepRef>("JobStepRefs");
+			_jobStepRefs = databaseService.GetCollection<JobStepRef>("JobStepRefs");
 
-			if (!DatabaseService.ReadOnlyMode)
+			if (!databaseService.ReadOnlyMode)
 			{
-				JobStepRefs.Indexes.CreateOne(new CreateIndexModel<JobStepRef>(Builders<JobStepRef>.IndexKeys.Ascending(x => x.StreamId).Ascending(x => x.TemplateId).Ascending(x => x.Name).Descending(x => x.Change)));
+				_jobStepRefs.Indexes.CreateOne(new CreateIndexModel<JobStepRef>(Builders<JobStepRef>.IndexKeys.Ascending(x => x.StreamId).Ascending(x => x.TemplateId).Ascending(x => x.Name).Descending(x => x.Change)));
 			}
 		}
 
 		/// <inheritdoc/>
-		public async Task<IJobStepRef> InsertOrReplaceAsync(JobStepRefId Id, string JobName, string StepName, StreamId StreamId, TemplateRefId TemplateId, int Change, LogId? LogId, PoolId? PoolId, AgentId? AgentId, JobStepOutcome? Outcome, int? LastSuccess, int? LastWarning, float WaitTime, float InitTime, DateTime StartTimeUtc, DateTime? FinishTimeUtc)
+		public async Task<IJobStepRef> InsertOrReplaceAsync(JobStepRefId id, string jobName, string stepName, StreamId streamId, TemplateRefId templateId, int change, LogId? logId, PoolId? poolId, AgentId? agentId, JobStepOutcome? outcome, int? lastSuccess, int? lastWarning, float waitTime, float initTime, DateTime startTimeUtc, DateTime? finishTimeUtc)
 		{
-			JobStepRef NewJobStepRef = new JobStepRef(Id, JobName, StepName, StreamId, TemplateId, Change, LogId, PoolId, AgentId, Outcome, LastSuccess, LastWarning, WaitTime, InitTime, StartTimeUtc, FinishTimeUtc);
-			await JobStepRefs.ReplaceOneAsync(Builders<JobStepRef>.Filter.Eq(x => x.Id, NewJobStepRef.Id), NewJobStepRef, new ReplaceOptions { IsUpsert = true });
-			return NewJobStepRef;
+			JobStepRef newJobStepRef = new JobStepRef(id, jobName, stepName, streamId, templateId, change, logId, poolId, agentId, outcome, lastSuccess, lastWarning, waitTime, initTime, startTimeUtc, finishTimeUtc);
+			await _jobStepRefs.ReplaceOneAsync(Builders<JobStepRef>.Filter.Eq(x => x.Id, newJobStepRef.Id), newJobStepRef, new ReplaceOptions { IsUpsert = true });
+			return newJobStepRef;
 		}
 
 		/// <inheritdoc/>
-		public async Task<List<IJobStepRef>> GetStepsForNodeAsync(StreamId StreamId, TemplateRefId TemplateId, string NodeName, int? Change, bool IncludeFailed, int Count)
+		public async Task<List<IJobStepRef>> GetStepsForNodeAsync(StreamId streamId, TemplateRefId templateId, string nodeName, int? change, bool includeFailed, int count)
 		{
 			// Find all the steps matching the given criteria
-			FilterDefinitionBuilder<JobStepRef> FilterBuilder = Builders<JobStepRef>.Filter;
+			FilterDefinitionBuilder<JobStepRef> filterBuilder = Builders<JobStepRef>.Filter;
 
-			FilterDefinition<JobStepRef> Filter = FilterDefinition<JobStepRef>.Empty;
-			Filter &= FilterBuilder.Eq(x => x.StreamId, StreamId);
-			Filter &= FilterBuilder.Eq(x => x.TemplateId, TemplateId);
-			Filter &= FilterBuilder.Eq(x => x.Name, NodeName);
-			if (Change != null)
+			FilterDefinition<JobStepRef> filter = FilterDefinition<JobStepRef>.Empty;
+			filter &= filterBuilder.Eq(x => x.StreamId, streamId);
+			filter &= filterBuilder.Eq(x => x.TemplateId, templateId);
+			filter &= filterBuilder.Eq(x => x.Name, nodeName);
+			if (change != null)
 			{
-				Filter &= FilterBuilder.Lte(x => x.Change, Change.Value);
+				filter &= filterBuilder.Lte(x => x.Change, change.Value);
 			}
-			if (!IncludeFailed)
+			if (!includeFailed)
 			{
-				Filter &= FilterBuilder.Ne(x => x.Outcome, JobStepOutcome.Failure);
+				filter &= filterBuilder.Ne(x => x.Outcome, JobStepOutcome.Failure);
 			}
 
-			List<JobStepRef> Steps = await JobStepRefs.Find(Filter).SortByDescending(x => x.Change).ThenByDescending(x => x.StartTime).Limit(Count).ToListAsync();
-			return Steps.ConvertAll<IJobStepRef>(x => x);
+			List<JobStepRef> steps = await _jobStepRefs.Find(filter).SortByDescending(x => x.Change).ThenByDescending(x => x.StartTime).Limit(count).ToListAsync();
+			return steps.ConvertAll<IJobStepRef>(x => x);
 		}
 
 		/// <inheritdoc/>
-		public async Task<IJobStepRef?> GetPrevStepForNodeAsync(StreamId StreamId, TemplateRefId TemplateId, string NodeName, int Change)
+		public async Task<IJobStepRef?> GetPrevStepForNodeAsync(StreamId streamId, TemplateRefId templateId, string nodeName, int change)
 		{
-			return await JobStepRefs.Find(x => x.StreamId == StreamId && x.TemplateId == TemplateId && x.Name == NodeName && x.Change < Change && x.Outcome != null).SortByDescending(x => x.Change).FirstOrDefaultAsync();
+			return await _jobStepRefs.Find(x => x.StreamId == streamId && x.TemplateId == templateId && x.Name == nodeName && x.Change < change && x.Outcome != null).SortByDescending(x => x.Change).FirstOrDefaultAsync();
 		}
 
 		/// <inheritdoc/>
-		public async Task<IJobStepRef?> GetNextStepForNodeAsync(StreamId StreamId, TemplateRefId TemplateId, string NodeName, int Change)
+		public async Task<IJobStepRef?> GetNextStepForNodeAsync(StreamId streamId, TemplateRefId templateId, string nodeName, int change)
 		{
-			return await JobStepRefs.Find(x => x.StreamId == StreamId && x.TemplateId == TemplateId && x.Name == NodeName && x.Change > Change && x.Outcome != null).SortBy(x => x.Change).FirstOrDefaultAsync();
+			return await _jobStepRefs.Find(x => x.StreamId == streamId && x.TemplateId == templateId && x.Name == nodeName && x.Change > change && x.Outcome != null).SortBy(x => x.Change).FirstOrDefaultAsync();
 		}
 	}
 }
diff --git a/Engine/Source/Programs/Horde/Horde.Build/Jobs/Impl/JobTaskSource.cs b/Engine/Source/Programs/Horde/Horde.Build/Jobs/Impl/JobTaskSource.cs
--- a/Engine/Source/Programs/Horde/Horde.Build/Jobs/Impl/JobTaskSource.cs
+++ b/Engine/Source/Programs/Horde/Horde.Build/Jobs/Impl/JobTaskSource.cs
@@ -1,35 +1,28 @@
 // Copyright Epic Games, Inc. All Rights Reserved.
 
+using System;
+using System.Collections.Generic;
+using System.Diagnostics;
+using System.Diagnostics.CodeAnalysis;
+using System.Globalization;
+using System.Linq;
+using System.Text;
+using System.Threading;
+using System.Threading.Tasks;
 using EpicGames.Core;
-using EpicGames.Horde.Common;
 using Google.Protobuf;
-using Google.Protobuf.Reflection;
 using Google.Protobuf.WellKnownTypes;
 using Horde.Build.Api;
 using Horde.Build.Collections;
-using HordeCommon;
 using Horde.Build.Models;
-using HordeCommon.Rpc;
-using HordeCommon.Rpc.Tasks;
-using Horde.Build.Tasks;
+using Horde.Build.Services;
 using Horde.Build.Utilities;
+using HordeCommon;
+using HordeCommon.Rpc.Tasks;
 using Microsoft.Extensions.Hosting;
 using Microsoft.Extensions.Logging;
 using Microsoft.Extensions.Options;
-using MongoDB.Bson;
 using MongoDB.Driver;
-using System;
-using System.Collections.Generic;
-using System.Diagnostics;
-using System.Diagnostics.CodeAnalysis;
-using System.Linq;
-using System.Runtime.CompilerServices;
-using System.Text;
-using System.Threading;
-using System.Threading.Tasks;
-using Horde.Build.Services;
-using EpicGames.Horde.Compute;
-using System.Globalization;
 
 namespace Horde.Build.Tasks.Impl
 {
@@ -53,77 +46,71 @@ namespace Horde.Build.Tasks.Impl
 		/// <summary>
 		/// An item in the queue to be executed
 		/// </summary>
-		[DebuggerDisplay("{Job.Id}:{Batch.Id} ({PoolId})")]
+		[DebuggerDisplay("{_job.Id}:{Batch.Id} ({_poolId})")]
 		internal class QueueItem
 		{
 			/// <summary>
 			/// The stream for this job
 			/// </summary>
-			public IStream Stream;
+			public IStream _stream;
 
 			/// <summary>
 			/// The job instance
 			/// </summary>
-			public IJob Job;
+			public IJob _job;
 
 			/// <summary>
 			/// Index of the batch within this job to be executed
 			/// </summary>
-			public int BatchIdx;
+			public int _batchIdx;
 
 			/// <summary>
 			/// The pool of machines to allocate from
 			/// </summary>
-			public PoolId PoolId;
+			public PoolId _poolId;
 
 			/// <summary>
 			/// The type of workspace that this item should run in
 			/// </summary>
-			public AgentWorkspace Workspace;
+			public AgentWorkspace _workspace;
 
 			/// <summary>
 			/// Whether or not to use the AutoSDK.
 			/// </summary>
-			public bool UseAutoSDK;
+			public bool _useAutoSdk;
 
 			/// <summary>
 			/// Task for creating a lease and assigning to a waiter
 			/// </summary>
-			public Task? AssignTask;
+			public Task? _assignTask;
 
 			/// <summary>
 			/// Accessor for the batch referenced by this item
 			/// </summary>
-			public IJobStepBatch Batch
-			{
-				get { return Job.Batches[BatchIdx]; }
-			}
+			public IJobStepBatch Batch => _job.Batches[_batchIdx];
 
 			/// <summary>
 			/// Returns an identifier describing this unique batch
 			/// </summary>
-			public (JobId, SubResourceId) Id
-			{
-				get { return (Job.Id, Batch.Id); }
-			}
+			public (JobId, SubResourceId) Id => (_job.Id, Batch.Id);
 
 			/// <summary>
 			/// Constructor
 			/// </summary>
-			/// <param name="Stream">The stream containing this job</param>
-			/// <param name="Job">The job instance</param>
-			/// <param name="BatchIdx">The batch index to execute</param>
-			/// <param name="PoolId">Unique id of the pool of machines to allocate from</param>
-			/// <param name="Workspace">The workspace that this job should run in</param>
-			/// <param name="UseAutoSDK">Whether or not to use the AutoSDK</param>
-			public QueueItem(IStream Stream, IJob Job, int BatchIdx, PoolId PoolId, AgentWorkspace Workspace, bool UseAutoSDK)
+			/// <param name="stream">The stream containing this job</param>
+			/// <param name="job">The job instance</param>
+			/// <param name="batchIdx">The batch index to execute</param>
+			/// <param name="poolId">Unique id of the pool of machines to allocate from</param>
+			/// <param name="workspace">The workspace that this job should run in</param>
+			/// <param name="useAutoSdk">Whether or not to use the AutoSDK</param>
+			public QueueItem(IStream stream, IJob job, int batchIdx, PoolId poolId, AgentWorkspace workspace, bool useAutoSdk)
 			{
-				this.Stream = Stream;
-				this.Job = Job;
-				this.BatchIdx = BatchIdx;
-				this.PoolId = PoolId;
-				this.Workspace = Workspace;
-				this.UseAutoSDK = UseAutoSDK;
+				_stream = stream;
+				_job = job;
+				_batchIdx = batchIdx;
+				_poolId = poolId;
+				_workspace = workspace;
+				_useAutoSdk = useAutoSdk;
 			}
 		}
 
@@ -135,30 +122,30 @@ namespace Horde.Build.Tasks.Impl
 			/// <summary>
 			/// Compare two items
 			/// </summary>
-			/// <param name="X">First item to compare</param>
-			/// <param name="Y">Second item to compare</param>
+			/// <param name="x">First item to compare</param>
+			/// <param name="y">Second item to compare</param>
 			/// <returns>Negative value if X is a higher priority than Y</returns>
-			public int Compare([AllowNull] QueueItem X, [AllowNull] QueueItem Y)
+			public int Compare([AllowNull] QueueItem x, [AllowNull] QueueItem y)
 			{
-				if (X == null)
+				if (x == null)
 				{
-					return (Y == null) ? 0 : -1;
+					return (y == null) ? 0 : -1;
 				}
-				else if (Y == null)
+				else if (y == null)
 				{
 					return 1;
 				}
 
-				int Delta = Y.Batch.SchedulePriority - X.Batch.SchedulePriority;
-				if (Delta == 0)
+				int delta = y.Batch.SchedulePriority - x.Batch.SchedulePriority;
+				if (delta == 0)
 				{
-					Delta = X.Job.Id.CompareTo(Y.Job.Id);
-					if (Delta == 0)
+					delta = x._job.Id.CompareTo(y._job.Id);
+					if (delta == 0)
 					{
-						Delta = (int)X.Batch.Id.Value - (int)Y.Batch.Id.Value;
+						delta = (int)x.Batch.Id.Value - (int)y.Batch.Id.Value;
 					}
 				}
-				return Delta;
+				return delta;
 			}
 		}
 
@@ -178,60 +165,60 @@ namespace Horde.Build.Tasks.Impl
 			public Task<AgentLease?> Task => LeaseSource.Task;
 
 			/// <summary>
-			/// Completion source for the waiting agent. If a new queue item becomes available, the result will be passed through this.
+			/// Completion source for the waiting agent. If a new queue item becomes available, the result will be passed through 
 			/// </summary>
 			public TaskCompletionSource<AgentLease?> LeaseSource { get; } = new TaskCompletionSource<AgentLease?>();
 
 			/// <summary>
 			/// Constructor
 			/// </summary>
-			/// <param name="Agent">The agent waiting for a task</param>
-			public QueueWaiter(IAgent Agent)
+			/// <param name="agent">The agent waiting for a task</param>
+			public QueueWaiter(IAgent agent)
 			{
-				this.Agent = Agent;
+				Agent = agent;
 			}
 		}
 
-		DatabaseService DatabaseService;
-		StreamService StreamService;
-		ILogFileService LogFileService;
-		IAgentCollection AgentsCollection;
-		IJobCollection Jobs;
-		IJobStepRefCollection JobStepRefs;
-		IGraphCollection Graphs;
-		IPoolCollection PoolCollection;
-		IUgsMetadataCollection UgsMetadataCollection;
-		PerforceLoadBalancer PerforceLoadBalancer;
-		IHostApplicationLifetime ApplicationLifetime;
-		IOptionsMonitor<ServerSettings> Settings;
-		ILogger<JobTaskSource> Logger;
-		ITicker Ticker;
+		readonly DatabaseService _databaseService;
+		readonly StreamService _streamService;
+		readonly ILogFileService _logFileService;
+		readonly IAgentCollection _agentsCollection;
+		readonly IJobCollection _jobs;
+		readonly IJobStepRefCollection _jobStepRefs;
+		readonly IGraphCollection _graphs;
+		readonly IPoolCollection _poolCollection;
+		readonly IUgsMetadataCollection _ugsMetadataCollection;
+		readonly PerforceLoadBalancer _perforceLoadBalancer;
+		readonly IHostApplicationLifetime _applicationLifetime;
+		readonly IOptionsMonitor<ServerSettings> _settings;
+		readonly ILogger<JobTaskSource> _logger;
+		readonly ITicker _ticker;
 
 		// Object used for ensuring mutual exclusion to the queues
-		object LockObject = new object();
+		readonly object _lockObject = new object();
 
 		// List of items waiting to be executed
-		SortedSet<QueueItem> Queue = new SortedSet<QueueItem>(new QueueItemComparer());
+		SortedSet<QueueItem> _queue = new SortedSet<QueueItem>(new QueueItemComparer());
 
 		// Map from batch id to the corresponding queue item
-		Dictionary<(JobId, SubResourceId), QueueItem> BatchIdToQueueItem = new Dictionary<(JobId, SubResourceId), QueueItem>();
+		Dictionary<(JobId, SubResourceId), QueueItem> _batchIdToQueueItem = new Dictionary<(JobId, SubResourceId), QueueItem>();
 
 		// Set of long-poll tasks waiting to be satisfied 
-		HashSet<QueueWaiter> Waiters = new HashSet<QueueWaiter>();
+		readonly HashSet<QueueWaiter> _waiters = new HashSet<QueueWaiter>();
 
 		// During a background queue refresh operation, any updated batches are added to this dictionary for merging into the updated queue.
-		List<QueueItem>? NewQueueItemsDuringUpdate;
+		List<QueueItem>? _newQueueItemsDuringUpdate;
 
 		// Cache of pools
-		Dictionary<PoolId, IPool> CachedPoolIdToInstance = new Dictionary<PoolId, IPool>();
+		Dictionary<PoolId, IPool> _cachedPoolIdToInstance = new Dictionary<PoolId, IPool>();
 
 		// Cache of stream objects. Used to resolve agent types.
-		private Dictionary<StreamId, IStream> Streams = new Dictionary<StreamId, IStream>();
+		private Dictionary<StreamId, IStream> _streams = new Dictionary<StreamId, IStream>();
 
 		/// <summary>
 		/// Delegate for job schedule events
 		/// </summary>
-		public delegate void JobScheduleEvent(IPool Pool, bool HasAgentsOnline, IJob Job, IGraph Graph, SubResourceId BatchId);
+		public delegate void JobScheduleEvent(IPool pool, bool hasAgentsOnline, IJob job, IGraph graph, SubResourceId batchId);
 		
 		/// <summary>
 		/// Event triggered when a job is scheduled
@@ -239,39 +226,39 @@ namespace Horde.Build.Tasks.Impl
 		public event JobScheduleEvent? OnJobScheduled;
 
 		// Interval between querying the database for jobs to execute
-		static readonly TimeSpan RefreshInterval = TimeSpan.FromSeconds(5.0);
+		static readonly TimeSpan s_refreshInterval = TimeSpan.FromSeconds(5.0);
 
 		/// <summary>
 		/// Constructor
 		/// </summary>
-		public JobTaskSource(DatabaseService DatabaseService, IAgentCollection Agents, IJobCollection Jobs, IJobStepRefCollection JobStepRefs, IGraphCollection Graphs, IPoolCollection Pools, IUgsMetadataCollection UgsMetadataCollection, StreamService StreamService, ILogFileService LogFileService, PerforceLoadBalancer PerforceLoadBalancer, IHostApplicationLifetime ApplicationLifetime, IClock Clock, IOptionsMonitor<ServerSettings> Settings, ILogger<JobTaskSource> Logger)
+		public JobTaskSource(DatabaseService databaseService, IAgentCollection agents, IJobCollection jobs, IJobStepRefCollection jobStepRefs, IGraphCollection graphs, IPoolCollection pools, IUgsMetadataCollection ugsMetadataCollection, StreamService streamService, ILogFileService logFileService, PerforceLoadBalancer perforceLoadBalancer, IHostApplicationLifetime applicationLifetime, IClock clock, IOptionsMonitor<ServerSettings> settings, ILogger<JobTaskSource> logger)
 		{
-			this.DatabaseService = DatabaseService;
-			this.AgentsCollection = Agents;
-			this.Jobs = Jobs;
-			this.JobStepRefs = JobStepRefs;
-			this.Graphs = Graphs;
-			this.PoolCollection = Pools;
-			this.UgsMetadataCollection = UgsMetadataCollection;
-			this.StreamService = StreamService;
-			this.LogFileService = LogFileService;
-			this.PerforceLoadBalancer = PerforceLoadBalancer;
-			this.ApplicationLifetime = ApplicationLifetime;
-			this.Ticker = Clock.AddTicker(RefreshInterval, TickAsync, Logger);
-			this.Settings = Settings;
-			this.Logger = Logger;
+			_databaseService = databaseService;
+			_agentsCollection = agents;
+			_jobs = jobs;
+			_jobStepRefs = jobStepRefs;
+			_graphs = graphs;
+			_poolCollection = pools;
+			_ugsMetadataCollection = ugsMetadataCollection;
+			_streamService = streamService;
+			_logFileService = logFileService;
+			_perforceLoadBalancer = perforceLoadBalancer;
+			_applicationLifetime = applicationLifetime;
+			_ticker = clock.AddTicker<JobTaskSource>(s_refreshInterval, TickAsync, logger);
+			_settings = settings;
+			_logger = logger;
 
 			OnLeaseStartedProperties.Add(nameof(ExecuteJobTask.JobId), x => new JobId(x.JobId)).Add(nameof(ExecuteJobTask.BatchId), x => SubResourceId.Parse(x.BatchId)).Add(nameof(ExecuteJobTask.LogId), x => new LogId(x.LogId));
 		}
 
 		/// <inheritdoc/>
-		public Task StartAsync(CancellationToken cancellationToken) => Ticker.StartAsync();
+		public Task StartAsync(CancellationToken cancellationToken) => _ticker.StartAsync();
 
 		/// <inheritdoc/>
-		public Task StopAsync(CancellationToken cancellationToken) => Ticker.StopAsync();
+		public Task StopAsync(CancellationToken cancellationToken) => _ticker.StopAsync();
 
 		/// <inheritdoc/>
-		public void Dispose() => Ticker.Dispose();
+		public void Dispose() => _ticker.Dispose();
 
 		/// <summary>
 		/// Gets an object containing the stats of the queue for diagnostic purposes.
@@ -279,239 +266,239 @@ namespace Horde.Build.Tasks.Impl
 		/// <returns>Status object</returns>
 		public object GetStatus()
 		{
-			lock (LockObject)
+			lock (_lockObject)
 			{
-				List<object> OutputItems = new List<object>();
-				foreach (QueueItem QueueItem in Queue)
+				List<object> outputItems = new List<object>();
+				foreach (QueueItem queueItem in _queue)
 				{
-					OutputItems.Add(new { JobId = QueueItem.Job.Id.ToString(), BatchId = QueueItem.Batch.Id.ToString(), PoolId = QueueItem.PoolId.ToString(), Workspace = QueueItem.Workspace });
+					outputItems.Add(new { JobId = queueItem._job.Id.ToString(), BatchId = queueItem.Batch.Id.ToString(), PoolId = queueItem._poolId.ToString(), Workspace = queueItem._workspace });
 				}
 
-				List<object> OutputWaiters = new List<object>();
-				foreach (QueueWaiter Waiter in Waiters)
+				List<object> outputWaiters = new List<object>();
+				foreach (QueueWaiter waiter in _waiters)
 				{
-					OutputWaiters.Add(new { Id = Waiter.Agent.Id.ToString(), Pools = Waiter.Agent.GetPools().Select(x => x.ToString()).ToList(), Workspaces = Waiter.Agent.Workspaces });
+					outputWaiters.Add(new { Id = waiter.Agent.Id.ToString(), Pools = waiter.Agent.GetPools().Select(x => x.ToString()).ToList(), waiter.Agent.Workspaces });
 				}
 
-				return new { Items = OutputItems, Waiters = OutputWaiters };
+				return new { Items = outputItems, Waiters = outputWaiters };
 			}
 		}
 
 		/// <summary>
 		/// Cancel any pending wait for an agent, allowing it to cycle its session state immediately
 		/// </summary>
-		/// <param name="AgentId">The agent id</param>
-		public void CancelLongPollForAgent(AgentId AgentId)
+		/// <param name="agentId">The agent id</param>
+		public void CancelLongPollForAgent(AgentId agentId)
 		{
-			QueueWaiter? Waiter;
-			lock (LockObject)
+			QueueWaiter? waiter;
+			lock (_lockObject)
 			{
-				Waiter = Waiters.FirstOrDefault(x => x.Agent.Id == AgentId);
+				waiter = _waiters.FirstOrDefault(x => x.Agent.Id == agentId);
 			}
-			if(Waiter != null)
+			if(waiter != null)
 			{
-				Waiter.LeaseSource.TrySetResult(null);
+				waiter.LeaseSource.TrySetResult(null);
 			}
 		}
 
 		/// <summary>
 		/// Background task
 		/// </summary>
-		/// <param name="StoppingToken">Token that indicates that the service should shut down</param>
+		/// <param name="stoppingToken">Token that indicates that the service should shut down</param>
 		/// <returns>Async task</returns>
-		internal async ValueTask TickAsync(CancellationToken StoppingToken)
+		internal async ValueTask TickAsync(CancellationToken stoppingToken)
 		{
 			// Set the NewBatchIdToQueueItem member, so we capture any updated jobs during the DB query.
-			lock (LockObject)
+			lock (_lockObject)
 			{
-				NewQueueItemsDuringUpdate = new List<QueueItem>();
+				_newQueueItemsDuringUpdate = new List<QueueItem>();
 			}
 
 			// Query all the current streams
-			List<IStream> StreamsList = await StreamService.GetStreamsAsync();
-			Streams = StreamsList.ToDictionary(x => x.Id, x => x);
+			List<IStream> streamsList = await _streamService.GetStreamsAsync();
+			_streams = streamsList.ToDictionary(x => x.Id, x => x);
 
 			// Find all the pools which are valid (ie. have at least one online agent)
-			DateTime UtcNow = DateTime.UtcNow;
-			List<IAgent> Agents = await AgentsCollection.FindAsync();
-			List<IPool> Pools = await PoolCollection.GetAsync();
+			DateTime utcNow = DateTime.UtcNow;
+			List<IAgent> agents = await _agentsCollection.FindAsync();
+			List<IPool> pools = await _poolCollection.GetAsync();
 
 			// Find all the pools which are currently online
-			HashSet<PoolId> OnlinePools = new HashSet<PoolId>(Agents.Where(x => x.IsSessionValid(UtcNow) && x.Enabled).SelectMany(x => x.ExplicitPools));
-			foreach (IPool Pool in Pools)
+			HashSet<PoolId> onlinePools = new HashSet<PoolId>(agents.Where(x => x.IsSessionValid(utcNow) && x.Enabled).SelectMany(x => x.ExplicitPools));
+			foreach (IPool pool in pools)
 			{
-				if (Pool.Condition != null && !OnlinePools.Contains(Pool.Id) && Agents.Any(x => x.IsSessionValid(UtcNow) && x.SatisfiesCondition(Pool.Condition) && x.Enabled))
+				if (pool.Condition != null && !onlinePools.Contains(pool.Id) && agents.Any(x => x.IsSessionValid(utcNow) && x.SatisfiesCondition(pool.Condition) && x.Enabled))
 				{
-					OnlinePools.Add(Pool.Id);
+					onlinePools.Add(pool.Id);
 				}
 			}
 
 			// Find lists of valid pools and online pools
-			HashSet<PoolId> ValidPools = new HashSet<PoolId>(OnlinePools.Union(Agents.Where(x => !x.IsSessionValid(UtcNow) || !x.Enabled).SelectMany(x => x.ExplicitPools)));
-			foreach (IPool Pool in Pools)
+			HashSet<PoolId> validPools = new HashSet<PoolId>(onlinePools.Union(agents.Where(x => !x.IsSessionValid(utcNow) || !x.Enabled).SelectMany(x => x.ExplicitPools)));
+			foreach (IPool pool in pools)
 			{
-				if (Pool.Condition != null && !ValidPools.Contains(Pool.Id) && Agents.Any(x => !x.IsSessionValid(UtcNow) && x.SatisfiesCondition(Pool.Condition) && x.Enabled))
+				if (pool.Condition != null && !validPools.Contains(pool.Id) && agents.Any(x => !x.IsSessionValid(utcNow) && x.SatisfiesCondition(pool.Condition) && x.Enabled))
 				{
-					ValidPools.Add(Pool.Id);
+					validPools.Add(pool.Id);
 				}
 			}
 
 			// Query all the current pools
-			CachedPoolIdToInstance = Pools.ToDictionary(x => x.Id, x => x);
+			_cachedPoolIdToInstance = pools.ToDictionary(x => x.Id, x => x);
 
 			// New list of queue items
-			SortedSet<QueueItem> NewQueue = new SortedSet<QueueItem>(Queue.Comparer);
-			Dictionary<(JobId, SubResourceId), QueueItem> NewBatchIdToQueueItem = new Dictionary<(JobId, SubResourceId), QueueItem>();
+			SortedSet<QueueItem> newQueue = new SortedSet<QueueItem>(_queue.Comparer);
+			Dictionary<(JobId, SubResourceId), QueueItem> newBatchIdToQueueItem = new Dictionary<(JobId, SubResourceId), QueueItem>();
 			
 			// Returns true if agents are online and available for scheduling for a pool
-			bool IsPoolOnline(PoolId PoolId)
+			bool IsPoolOnline(PoolId poolId)
 			{
-				return OnlinePools.Contains(PoolId);
+				return onlinePools.Contains(poolId);
 			}
 			
 			// Returns true if a pool can be auto-scaled
-			bool IsPoolAutoScaled(PoolId PoolId)
+			bool IsPoolAutoScaled(PoolId poolId)
 			{
-				IPool? Pool = Pools.Find(p => p.Id == PoolId);
-				return ValidPools.Contains(PoolId) && Pool != null && Pool.EnableAutoscaling;
+				IPool? pool = pools.Find(p => p.Id == poolId);
+				return validPools.Contains(poolId) && pool != null && pool.EnableAutoscaling;
 			}
 			
-			bool HasAgentsOnlineOrIsAutoScaled(PoolId PoolId)
+			bool HasAgentsOnlineOrIsAutoScaled(PoolId poolId)
 			{
 				// If pool is auto-scaled, it will be considered online even if it has no agents online
-				return IsPoolOnline(PoolId) || IsPoolAutoScaled(PoolId);
+				return IsPoolOnline(poolId) || IsPoolAutoScaled(poolId);
 			}
 
 			// Query for a new list of jobs for the queue
-			List<IJob> NewJobs = await Jobs.GetDispatchQueueAsync();
-			for (int Idx = 0; Idx < NewJobs.Count; Idx++)
+			List<IJob> newJobs = await _jobs.GetDispatchQueueAsync();
+			for (int idx = 0; idx < newJobs.Count; idx++)
 			{
-				IJob? NewJob = NewJobs[Idx];
+				IJob? newJob = newJobs[idx];
 
-				if (NewJob.GraphHash == null)
+				if (newJob.GraphHash == null)
 				{
-					Logger.LogError("Job {JobId} has a null graph hash and can't be started.", NewJob.Id);
-					await Jobs.TryRemoveFromDispatchQueueAsync(NewJob);
+					_logger.LogError("Job {JobId} has a null graph hash and can't be started.", newJob.Id);
+					await _jobs.TryRemoveFromDispatchQueueAsync(newJob);
 					continue;
 				}
-				if (NewJob.AbortedByUserId != null)
+				if (newJob.AbortedByUserId != null)
 				{
-					Logger.LogError("Job {JobId} was aborted but not removed from dispatch queue", NewJob.Id);
-					await Jobs.TryRemoveFromDispatchQueueAsync(NewJob);
+					_logger.LogError("Job {JobId} was aborted but not removed from dispatch queue", newJob.Id);
+					await _jobs.TryRemoveFromDispatchQueueAsync(newJob);
 					continue;
 				}
 
 				// Get the graph for this job
-				IGraph Graph = await Graphs.GetAsync(NewJob.GraphHash);
+				IGraph graph = await _graphs.GetAsync(newJob.GraphHash);
 
 				// Get the stream. If it fails, skip the whole job.
-				IStream? Stream;
-				if (!Streams.TryGetValue(NewJob.StreamId, out Stream))
+				IStream? stream;
+				if (!_streams.TryGetValue(newJob.StreamId, out stream))
 				{
-					NewJob = await Jobs.SkipAllBatchesAsync(NewJob, Graph, JobStepBatchError.UnknownStream);
+					newJob = await _jobs.SkipAllBatchesAsync(newJob, graph, JobStepBatchError.UnknownStream);
 					continue;
 				}
 
 				// Update all the batches
-				HashSet<SubResourceId> CheckedBatchIds = new HashSet<SubResourceId>();
-				while (NewJob != null)
+				HashSet<SubResourceId> checkedBatchIds = new HashSet<SubResourceId>();
+				while (newJob != null)
 				{
 					// Find the next batch within this job that is ready
-					int BatchIdx = NewJob.Batches.FindIndex(x => x.State == JobStepBatchState.Ready && CheckedBatchIds.Add(x.Id));
-					if (BatchIdx == -1)
+					int batchIdx = newJob.Batches.FindIndex(x => x.State == JobStepBatchState.Ready && checkedBatchIds.Add(x.Id));
+					if (batchIdx == -1)
 					{
 						break;
 					}
 
 					// Validate the agent type and workspace settings
-					IJobStepBatch Batch = NewJob.Batches[BatchIdx];
-					if (!Stream.AgentTypes.TryGetValue(Graph.Groups[Batch.GroupIdx].AgentType, out AgentType? AgentType))
+					IJobStepBatch batch = newJob.Batches[batchIdx];
+					if (!stream.AgentTypes.TryGetValue(graph.Groups[batch.GroupIdx].AgentType, out AgentType? agentType))
 					{
-						NewJob = await SkipBatchAsync(NewJob, Batch.Id, Graph, JobStepBatchError.UnknownAgentType);
+						newJob = await SkipBatchAsync(newJob, batch.Id, graph, JobStepBatchError.UnknownAgentType);
 					}
-					else if (!CachedPoolIdToInstance.TryGetValue(AgentType.Pool, out IPool? Pool))
+					else if (!_cachedPoolIdToInstance.TryGetValue(agentType.Pool, out IPool? pool))
 					{
-						NewJob = await SkipBatchAsync(NewJob, Batch.Id, Graph, JobStepBatchError.UnknownPool);
+						newJob = await SkipBatchAsync(newJob, batch.Id, graph, JobStepBatchError.UnknownPool);
 					}
-					else if (!ValidPools.Contains(AgentType.Pool))
+					else if (!validPools.Contains(agentType.Pool))
 					{
-						NewJob = await SkipBatchAsync(NewJob, Batch.Id, Graph, JobStepBatchError.NoAgentsInPool);
+						newJob = await SkipBatchAsync(newJob, batch.Id, graph, JobStepBatchError.NoAgentsInPool);
 					}
-					else if (!HasAgentsOnlineOrIsAutoScaled(AgentType.Pool))
+					else if (!HasAgentsOnlineOrIsAutoScaled(agentType.Pool))
 					{
-						NewJob = await SkipBatchAsync(NewJob, Batch.Id, Graph, JobStepBatchError.NoAgentsOnline);
+						newJob = await SkipBatchAsync(newJob, batch.Id, graph, JobStepBatchError.NoAgentsOnline);
 					}
-					else if (!Stream.TryGetAgentWorkspace(AgentType, out (AgentWorkspace, bool)? WorkspaceResult))
+					else if (!stream.TryGetAgentWorkspace(agentType, out (AgentWorkspace, bool)? workspaceResult))
 					{
-						NewJob = await SkipBatchAsync(NewJob, Batch.Id, Graph, JobStepBatchError.UnknownWorkspace);
+						newJob = await SkipBatchAsync(newJob, batch.Id, graph, JobStepBatchError.UnknownWorkspace);
 					}
 					else
 					{
-						(AgentWorkspace Workspace, bool UseAutoSDK) = WorkspaceResult.Value;
-						QueueItem NewQueueItem = new QueueItem(Stream, NewJob, BatchIdx, AgentType.Pool, Workspace, UseAutoSDK);
-						NewQueue.Add(NewQueueItem);
-						NewBatchIdToQueueItem[(NewJob.Id, Batch.Id)] = NewQueueItem;
+						(AgentWorkspace workspace, bool useAutoSdk) = workspaceResult.Value;
+						QueueItem newQueueItem = new QueueItem(stream, newJob, batchIdx, agentType.Pool, workspace, useAutoSdk);
+						newQueue.Add(newQueueItem);
+						newBatchIdToQueueItem[(newJob.Id, batch.Id)] = newQueueItem;
 						
-						IPool? NewJobPool = Pools.Find(p => p.Id == AgentType.Pool);
-						if (NewJobPool != null)
+						IPool? newJobPool = pools.Find(p => p.Id == agentType.Pool);
+						if (newJobPool != null)
 						{
-							OnJobScheduled?.Invoke(NewJobPool, IsPoolOnline(AgentType.Pool), NewJob, Graph, Batch.Id);
+							OnJobScheduled?.Invoke(newJobPool, IsPoolOnline(agentType.Pool), newJob, graph, batch.Id);
 						}
 					}
 				}
 
-				if (NewJob != null)
+				if (newJob != null)
 				{
-					if (!NewJob.Batches.Any(Batch => Batch.State == JobStepBatchState.Ready || Batch.State == JobStepBatchState.Starting || Batch.State == JobStepBatchState.Running || Batch.State == JobStepBatchState.Stopping))
+					if (!newJob.Batches.Any(batch => batch.State == JobStepBatchState.Ready || batch.State == JobStepBatchState.Starting || batch.State == JobStepBatchState.Running || batch.State == JobStepBatchState.Stopping))
 					{
-						Logger.LogError("Job {JobId} is in dispatch queue but not currently executing", NewJob.Id);
-						await Jobs.TryRemoveFromDispatchQueueAsync(NewJob);
+						_logger.LogError("Job {JobId} is in dispatch queue but not currently executing", newJob.Id);
+						await _jobs.TryRemoveFromDispatchQueueAsync(newJob);
 					}
 				}
 			}
 
 			// Update the queue
-			lock (LockObject)
+			lock (_lockObject)
 			{
-				Queue = NewQueue;
-				BatchIdToQueueItem = NewBatchIdToQueueItem;
+				_queue = newQueue;
+				_batchIdToQueueItem = newBatchIdToQueueItem;
 
 				// Merge the new queue items with the queue
-				foreach (QueueItem NewQueueItem in NewQueueItemsDuringUpdate)
+				foreach (QueueItem newQueueItem in _newQueueItemsDuringUpdate)
 				{
-					QueueItem? ExistingQueueItem;
-					if (!NewBatchIdToQueueItem.TryGetValue((NewQueueItem.Job.Id, NewQueueItem.Batch.Id), out ExistingQueueItem))
+					QueueItem? existingQueueItem;
+					if (!newBatchIdToQueueItem.TryGetValue((newQueueItem._job.Id, newQueueItem.Batch.Id), out existingQueueItem))
 					{
 						// Always just add this item
-						Queue.Add(NewQueueItem);
-						BatchIdToQueueItem[NewQueueItem.Id] = NewQueueItem;
+						_queue.Add(newQueueItem);
+						_batchIdToQueueItem[newQueueItem.Id] = newQueueItem;
 					}
-					else if (NewQueueItem.Job.UpdateIndex > ExistingQueueItem.Job.UpdateIndex)
+					else if (newQueueItem._job.UpdateIndex > existingQueueItem._job.UpdateIndex)
 					{
 						// Replace the existing item
-						Queue.Remove(ExistingQueueItem);
-						Queue.Add(NewQueueItem);
-						BatchIdToQueueItem[NewQueueItem.Id] = NewQueueItem;
+						_queue.Remove(existingQueueItem);
+						_queue.Add(newQueueItem);
+						_batchIdToQueueItem[newQueueItem.Id] = newQueueItem;
 					}
 				}
 
 				// Clear out the list to capture queue items during an update
-				NewQueueItemsDuringUpdate = null;
+				_newQueueItemsDuringUpdate = null;
 			}
 		}
 
-	private async Task<IJob?> SkipBatchAsync(IJob Job, SubResourceId BatchId, IGraph Graph, JobStepBatchError Reason)
+	private async Task<IJob?> SkipBatchAsync(IJob job, SubResourceId batchId, IGraph graph, JobStepBatchError reason)
 		{
-			Logger.LogInformation("Skipping batch {BatchId} for job {JobId} (reason: {Reason})", BatchId, Job.Id, Reason);
+			_logger.LogInformation("Skipping batch {BatchId} for job {JobId} (reason: {Reason})", batchId, job.Id, reason);
 
-			IReadOnlyList<(LabelState, LabelOutcome)> OldLabelStates = Job.GetLabelStates(Graph);
-			IJob? NewJob = await Jobs.SkipBatchAsync(Job, BatchId, Graph, Reason);
-			if(NewJob != null)
+			IReadOnlyList<(LabelState, LabelOutcome)> oldLabelStates = job.GetLabelStates(graph);
+			IJob? newJob = await _jobs.SkipBatchAsync(job, batchId, graph, reason);
+			if(newJob != null)
 			{
-				IReadOnlyList<(LabelState, LabelOutcome)> NewLabelStates = NewJob.GetLabelStates(Graph);
-				await UpdateUgsBadges(NewJob, Graph, OldLabelStates, NewLabelStates);
+				IReadOnlyList<(LabelState, LabelOutcome)> newLabelStates = newJob.GetLabelStates(graph);
+				await UpdateUgsBadges(newJob, graph, oldLabelStates, newLabelStates);
 			}
-			return NewJob;
+			return newJob;
 		}
 
 		/// <summary>
@@ -520,32 +507,32 @@ namespace Horde.Build.Tasks.Impl
 		/// <returns>A copy of the queue</returns>
 		internal SortedSet<QueueItem> GetQueueForTesting()
 		{
-			lock (LockObject)
+			lock (_lockObject)
 			{
-				return new SortedSet<QueueItem>(Queue);
+				return new SortedSet<QueueItem>(_queue);
 			}
 		}
 
 		/// <summary>
 		/// Updates the current state of a job
 		/// </summary>
-		/// <param name="Job">The job that has been updated</param>
-		/// <param name="Graph">Graph for the job</param>
+		/// <param name="job">The job that has been updated</param>
+		/// <param name="graph">Graph for the job</param>
 		/// <returns>Async task</returns>
-		public void UpdateQueuedJob(IJob Job, IGraph Graph)
+		public void UpdateQueuedJob(IJob job, IGraph graph)
 		{
-			IStream? Stream;
-			Streams.TryGetValue(Job.StreamId, out Stream);
-			UpdateQueuedJob(Job, Graph, Stream);
+			IStream? stream;
+			_streams.TryGetValue(job.StreamId, out stream);
+			UpdateQueuedJob(job, graph, stream);
 		}
 
-		void AssignAnyQueueItemToWaiter(QueueWaiter Waiter)
+		void AssignAnyQueueItemToWaiter(QueueWaiter waiter)
 		{
-			lock (Waiters)
+			lock (_waiters)
 			{
-				foreach(QueueItem Item in BatchIdToQueueItem.Values)
+				foreach(QueueItem item in _batchIdToQueueItem.Values)
 				{
-					if (TryAssignItemToWaiter(Item, Waiter))
+					if (TryAssignItemToWaiter(item, waiter))
 					{
 						break;
 					}
@@ -556,20 +543,20 @@ namespace Horde.Build.Tasks.Impl
 		/// <summary>
 		/// Attempt to find a waiter that can handle the given queue item
 		/// </summary>
-		/// <param name="Item">The queue item</param>
+		/// <param name="item">The queue item</param>
 		/// <returns></returns>
 		[SuppressMessage("Maintainability", "CA1508:Avoid dead conditional code", Justification = "<Pending>")]
-		void AssignQueueItemToAnyWaiter(QueueItem Item)
+		void AssignQueueItemToAnyWaiter(QueueItem item)
 		{
-			if (Item.AssignTask == null && Item.Batch.SessionId == null)
+			if (item._assignTask == null && item.Batch.SessionId == null)
 			{
-				lock (Waiters)
+				lock (_waiters)
 				{
-					if (Item.AssignTask == null && Item.Batch.SessionId == null)
+					if (item._assignTask == null && item.Batch.SessionId == null)
 					{
-						foreach (QueueWaiter Waiter in Waiters)
+						foreach (QueueWaiter waiter in _waiters)
 						{
-							if (TryAssignItemToWaiter(Item, Waiter))
+							if (TryAssignItemToWaiter(item, waiter))
 							{
 								break;
 							}
@@ -582,19 +569,19 @@ namespace Horde.Build.Tasks.Impl
 		/// <summary>
 		/// Attempts to assign a queue item to an agent waiting for work
 		/// </summary>
-		/// <param name="Item"></param>
-		/// <param name="Waiter"></param>
+		/// <param name="item"></param>
+		/// <param name="waiter"></param>
 		/// <returns></returns>
-		bool TryAssignItemToWaiter(QueueItem Item, QueueWaiter Waiter)
+		bool TryAssignItemToWaiter(QueueItem item, QueueWaiter waiter)
 		{
-			if (Item.AssignTask == null && Item.Batch.SessionId == null && Waiter.Agent.Enabled && Waiter.Agent.Leases.Count == 0 && Waiter.Agent.IsInPool(Item.PoolId))
+			if (item._assignTask == null && item.Batch.SessionId == null && waiter.Agent.Enabled && waiter.Agent.Leases.Count == 0 && waiter.Agent.IsInPool(item._poolId))
 			{
-				Task StartTask = new Task<Task>(() => TryCreateLeaseAsync(Item, Waiter));
-				Task ExecuteTask = StartTask.ContinueWith(Task => Task, TaskScheduler.Default);
+				Task startTask = new Task<Task>(() => TryCreateLeaseAsync(item, waiter));
+				Task executeTask = startTask.ContinueWith(task => task, TaskScheduler.Default);
 
-				if (Interlocked.CompareExchange(ref Item.AssignTask, ExecuteTask, null) == null)
+				if (Interlocked.CompareExchange(ref item._assignTask, executeTask, null) == null)
 				{
-					StartTask.Start(TaskScheduler.Default);
+					startTask.Start(TaskScheduler.Default);
 					return true;
 				}
 			}
@@ -604,64 +591,64 @@ namespace Horde.Build.Tasks.Impl
 		/// <summary>
 		/// Updates the current state of a job
 		/// </summary>
-		/// <param name="Job">The job that has been updated</param>
-		/// <param name="Graph">Graph for the job</param>
-		/// <param name="Stream">The stream containing the job</param>
-		public void UpdateQueuedJob(IJob Job, IGraph Graph, IStream? Stream)
+		/// <param name="job">The job that has been updated</param>
+		/// <param name="graph">Graph for the job</param>
+		/// <param name="stream">The stream containing the job</param>
+		public void UpdateQueuedJob(IJob job, IGraph graph, IStream? stream)
 		{
-			List<TaskCompletionSource<bool>> CompleteWaiters = new List<TaskCompletionSource<bool>>();
-			lock (LockObject)
+			List<TaskCompletionSource<bool>> completeWaiters = new List<TaskCompletionSource<bool>>();
+			lock (_lockObject)
 			{
-				for (int BatchIdx = 0; BatchIdx < Job.Batches.Count; BatchIdx++)
+				for (int batchIdx = 0; batchIdx < job.Batches.Count; batchIdx++)
 				{
-					IJobStepBatch Batch = Job.Batches[BatchIdx];
-					if (Batch.State == JobStepBatchState.Ready && Stream != null && Batch.AgentId == null)
+					IJobStepBatch batch = job.Batches[batchIdx];
+					if (batch.State == JobStepBatchState.Ready && stream != null && batch.AgentId == null)
 					{
 						// Check if this item is already in the list.
-						QueueItem? ExistingItem;
-						if (BatchIdToQueueItem.TryGetValue((Job.Id, Batch.Id), out ExistingItem))
+						QueueItem? existingItem;
+						if (_batchIdToQueueItem.TryGetValue((job.Id, batch.Id), out existingItem))
 						{
 							// Make sure this is a newer version of the job. There's no guarantee that this is the latest revision.
-							if (Job.UpdateIndex > ExistingItem.Job.UpdateIndex)
+							if (job.UpdateIndex > existingItem._job.UpdateIndex)
 							{
-								if (Batch.SchedulePriority == ExistingItem.Batch.SchedulePriority)
+								if (batch.SchedulePriority == existingItem.Batch.SchedulePriority)
 								{
-									ExistingItem.Job = Job;
-									ExistingItem.BatchIdx = BatchIdx;
+									existingItem._job = job;
+									existingItem._batchIdx = batchIdx;
 								}
 								else
 								{
-									RemoveQueueItem(ExistingItem);
-									InsertQueueItem(Stream, Job, BatchIdx, ExistingItem.PoolId, ExistingItem.Workspace, ExistingItem.UseAutoSDK);
+									RemoveQueueItem(existingItem);
+									InsertQueueItem(stream, job, batchIdx, existingItem._poolId, existingItem._workspace, existingItem._useAutoSdk);
 								}
 							}
 							continue;
 						}
 
 						// Get the group being executed by this batch
-						INodeGroup Group = Graph.Groups[Batch.GroupIdx];
+						INodeGroup group = graph.Groups[batch.GroupIdx];
 
 						// Get the requirements for the new queue item
-						AgentType? AgentType;
-						if (Stream.AgentTypes.TryGetValue(Group.AgentType, out AgentType))
+						AgentType? agentType;
+						if (stream.AgentTypes.TryGetValue(group.AgentType, out agentType))
 						{
-							(AgentWorkspace, bool)? Result;
-							if (Stream.TryGetAgentWorkspace(AgentType, out Result))
+							(AgentWorkspace, bool)? result;
+							if (stream.TryGetAgentWorkspace(agentType, out result))
 							{
-								(AgentWorkspace AgentWorkspace, bool UseAutoSdk) = Result.Value;
-								InsertQueueItem(Stream, Job, BatchIdx, AgentType.Pool, AgentWorkspace, UseAutoSdk);
+								(AgentWorkspace agentWorkspace, bool useAutoSdk) = result.Value;
+								InsertQueueItem(stream, job, batchIdx, agentType.Pool, agentWorkspace, useAutoSdk);
 							}
 						}
 					}
 					else
 					{
 						// Check if this item is already in the list. Remove it if it is.
-						QueueItem? ExistingItem;
-						if (BatchIdToQueueItem.TryGetValue((Job.Id, Batch.Id), out ExistingItem))
+						QueueItem? existingItem;
+						if (_batchIdToQueueItem.TryGetValue((job.Id, batch.Id), out existingItem))
 						{
-							if (Job.UpdateIndex > ExistingItem.Job.UpdateIndex)
+							if (job.UpdateIndex > existingItem._job.UpdateIndex)
 							{
-								RemoveQueueItem(ExistingItem);
+								RemoveQueueItem(existingItem);
 							}
 						}
 					}
@@ -669,350 +656,350 @@ namespace Horde.Build.Tasks.Impl
 			}
 
 			// Awake all the threads that have been assigned new work items. Has do be done outside the lock to prevent continuations running within it (see Waiter.CompletionSource for more info).
-			foreach (TaskCompletionSource<bool> CompleteWaiter in CompleteWaiters)
+			foreach (TaskCompletionSource<bool> completeWaiter in completeWaiters)
 			{
-				CompleteWaiter.TrySetResult(true);
+				completeWaiter.TrySetResult(true);
 			}
 		}
 
 		/// <inheritdoc/>
-		public override async Task<AgentLease?> AssignLeaseAsync(IAgent Agent, CancellationToken CancellationToken)
+		public override async Task<AgentLease?> AssignLeaseAsync(IAgent agent, CancellationToken cancellationToken)
 		{
-			QueueWaiter Waiter = new QueueWaiter(Agent);
-			lock (LockObject)
+			QueueWaiter waiter = new QueueWaiter(agent);
+			lock (_lockObject)
 			{
-				AssignAnyQueueItemToWaiter(Waiter);
-				if (Waiter.LeaseSource.Task.TryGetResult(out AgentLease? Result))
+				AssignAnyQueueItemToWaiter(waiter);
+				if (waiter.LeaseSource.Task.TryGetResult(out AgentLease? result))
 				{
-					return Result;
+					return result;
 				}
-				Waiters.Add(Waiter);
+				_waiters.Add(waiter);
 			}
 
-			AgentLease? Lease;
-			using (CancellationToken.Register(() => Waiter.LeaseSource.TrySetResult(null)))
+			AgentLease? lease;
+			using (cancellationToken.Register(() => waiter.LeaseSource.TrySetResult(null)))
 			{
-				Lease = await Waiter.Task;
+				lease = await waiter.Task;
 			}
 
-			lock (LockObject)
+			lock (_lockObject)
 			{
-				Waiters.Remove(Waiter);
+				_waiters.Remove(waiter);
 			}
 
-			return Lease;
+			return lease;
 		}
 
 		/// <inheritdoc/>
-		public override Task CancelLeaseAsync(IAgent Agent, LeaseId LeaseId, ExecuteJobTask Task)
+		public override Task CancelLeaseAsync(IAgent agent, LeaseId leaseId, ExecuteJobTask task)
 		{
-			return CancelLeaseAsync(Agent, new JobId(Task.JobId), Task.BatchId.ToSubResourceId());
+			return CancelLeaseAsync(agent, new JobId(task.JobId), task.BatchId.ToSubResourceId());
 		}
 
 		/// <summary>
 		/// Assign a new batch to be executed by the given agent
 		/// </summary>
-		/// <param name="Item">The item to create a lease for</param>
-		/// <param name="Waiter">The agent waiting for work</param>
+		/// <param name="item">The item to create a lease for</param>
+		/// <param name="waiter">The agent waiting for work</param>
 		/// <returns>New work to execute</returns>
-		private async Task<AgentLease?> TryCreateLeaseAsync(QueueItem Item, QueueWaiter Waiter)
+		private async Task<AgentLease?> TryCreateLeaseAsync(QueueItem item, QueueWaiter waiter)
 		{
-			IJob Job = Item.Job;
-			IJobStepBatch Batch = Item.Batch;
-			IAgent Agent = Waiter.Agent;
-			Logger.LogDebug("Assigning job {JobId}, batch {BatchId} to waiter (agent {AgentId})", Job.Id, Batch.Id, Agent.Id);
+			IJob job = item._job;
+			IJobStepBatch batch = item.Batch;
+			IAgent agent = waiter.Agent;
+			_logger.LogDebug("Assigning job {JobId}, batch {BatchId} to waiter (agent {AgentId})", job.Id, batch.Id, agent.Id);
 
 			// Generate a new unique id for the lease
-			LeaseId LeaseId = LeaseId.GenerateNewId();
+			LeaseId leaseId = LeaseId.GenerateNewId();
 
 			// The next time to try assigning to another agent
-			DateTime BackOffTime = DateTime.UtcNow + TimeSpan.FromMinutes(1.0);
+			DateTime backOffTime = DateTime.UtcNow + TimeSpan.FromMinutes(1.0);
 
 			// Try to update the job with this agent id
-			LogId LogId = (await LogFileService.CreateLogFileAsync(Job.Id, Agent.SessionId, LogType.Json)).Id;
-			IJob? NewJob = await Jobs.TryAssignLeaseAsync(Item.Job, Item.BatchIdx, Item.PoolId, Agent.Id, Agent.SessionId!.Value, LeaseId, LogId);
-			if (NewJob != null)
+			LogId logId = (await _logFileService.CreateLogFileAsync(job.Id, agent.SessionId, LogType.Json)).Id;
+			IJob? newJob = await _jobs.TryAssignLeaseAsync(item._job, item._batchIdx, item._poolId, agent.Id, agent.SessionId!.Value, leaseId, logId);
+			if (newJob != null)
 			{
-				Job = NewJob;
+				job = newJob;
 
 				// Get the lease name
-				StringBuilder LeaseName = new StringBuilder($"{Item.Stream.Name} - ");
-				if (Job.PreflightChange > 0)
+				StringBuilder leaseName = new StringBuilder($"{item._stream.Name} - ");
+				if (job.PreflightChange > 0)
 				{
-					LeaseName.Append((Job.Change > 0) ? $"Preflight CL {Job.PreflightChange} against CL {Job.Change}" : $"Preflight CL {Job.PreflightChange} against latest");
+					leaseName.Append((job.Change > 0) ? $"Preflight CL {job.PreflightChange} against CL {job.Change}" : $"Preflight CL {job.PreflightChange} against latest");
 				}
 				else
 				{
-					LeaseName.Append((Job.Change > 0) ? $"CL {Job.Change}" : "Latest CL");
+					leaseName.Append((job.Change > 0) ? $"CL {job.Change}" : "Latest CL");
 				}
-				LeaseName.Append(CultureInfo.InvariantCulture, $" - {Job.Name}");
+				leaseName.Append(CultureInfo.InvariantCulture, $" - {job.Name}");
 
 				// Get the global settings
-				Globals Globals = await DatabaseService.GetGlobalsAsync();
+				Globals globals = await _databaseService.GetGlobalsAsync();
 
 				// Encode the payload
-				ExecuteJobTask? Task = await CreateExecuteJobTaskAsync(Item.Stream, Job, Batch, Agent, Item.Workspace, Item.UseAutoSDK, LogId);
-				if (Task != null)
+				ExecuteJobTask? task = await CreateExecuteJobTaskAsync(item._stream, job, batch, agent, item._workspace, item._useAutoSdk, logId);
+				if (task != null)
 				{
-					byte[] Payload = Any.Pack(Task).ToByteArray();
+					byte[] payload = Any.Pack(task).ToByteArray();
 
 					// Create the lease and try to set it on the waiter. If this fails, the waiter has already moved on, and the lease can be cancelled.
-					AgentLease Lease = new AgentLease(LeaseId, LeaseName.ToString(), Job.StreamId, Item.PoolId, LogId, LeaseState.Pending, null, true, Payload);
-					if (Waiter.LeaseSource.TrySetResult(Lease))
+					AgentLease lease = new AgentLease(leaseId, leaseName.ToString(), job.StreamId, item._poolId, logId, LeaseState.Pending, null, true, payload);
+					if (waiter.LeaseSource.TrySetResult(lease))
 					{
-						Logger.LogDebug("Assigned lease {LeaseId} to agent {AgentId}", LeaseId, Agent.Id);
-						return Lease;
+						_logger.LogDebug("Assigned lease {LeaseId} to agent {AgentId}", leaseId, agent.Id);
+						return lease;
 					}
 				}
 
 				// Cancel the lease
-				Logger.LogDebug("Unable to assign lease {LeaseId} to agent {AgentId}, cancelling", LeaseId, Agent.Id);
-				await CancelLeaseAsync(Waiter.Agent, Job.Id, Batch.Id);
+				_logger.LogDebug("Unable to assign lease {LeaseId} to agent {AgentId}, cancelling", leaseId, agent.Id);
+				await CancelLeaseAsync(waiter.Agent, job.Id, batch.Id);
 			}
 			else
 			{
 				// Unable to assign job
-				Logger.LogDebug("Failed to assign job {JobId}, batch {BatchId} to agent {AgentId}. Refreshing queue entries.", Job.Id, Batch.Id, Agent.Id);
+				_logger.LogDebug("Failed to assign job {JobId}, batch {BatchId} to agent {AgentId}. Refreshing queue entries.", job.Id, batch.Id, agent.Id);
 
 				// Get the new copy of the job
-				NewJob = await Jobs.GetAsync(Job.Id);
-				if (NewJob == null)
+				newJob = await _jobs.GetAsync(job.Id);
+				if (newJob == null)
 				{
-					lock (LockObject)
+					lock (_lockObject)
 					{
-						List<QueueItem> RemoveItems = Queue.Where(x => x.Job == Job).ToList();
-						foreach (QueueItem RemoveItem in RemoveItems)
+						List<QueueItem> removeItems = _queue.Where(x => x._job == job).ToList();
+						foreach (QueueItem removeItem in removeItems)
 						{
-							RemoveQueueItem(RemoveItem);
+							RemoveQueueItem(removeItem);
 						}
 					}
 				}
 				else
 				{
-					IGraph Graph = await Graphs.GetAsync(NewJob.GraphHash);
-					UpdateQueuedJob(NewJob, Graph);
+					IGraph graph = await _graphs.GetAsync(newJob.GraphHash);
+					UpdateQueuedJob(newJob, graph);
 				}
 			}
 
 			// Clear out the assignment for this item, and try to reassign it
-			Item.AssignTask = null;
-			AssignQueueItemToAnyWaiter(Item);
+			item._assignTask = null;
+			AssignQueueItemToAnyWaiter(item);
 			return null;
 		}
 
-		async Task<ExecuteJobTask?> CreateExecuteJobTaskAsync(IStream Stream, IJob Job, IJobStepBatch Batch, IAgent Agent, AgentWorkspace Workspace, bool bUseAutoSDK, LogId LogId)
+		async Task<ExecuteJobTask?> CreateExecuteJobTaskAsync(IStream stream, IJob job, IJobStepBatch batch, IAgent agent, AgentWorkspace workspace, bool bUseAutoSdk, LogId logId)
 		{
 			// Get the lease name
-			StringBuilder LeaseName = new StringBuilder($"{Stream.Name} - ");
-			if (Job.PreflightChange > 0)
+			StringBuilder leaseName = new StringBuilder($"{stream.Name} - ");
+			if (job.PreflightChange > 0)
 			{
-				LeaseName.Append((Job.Change > 0) ? $"Preflight CL {Job.PreflightChange} against CL {Job.Change}" : $"Preflight CL {Job.PreflightChange} against latest");
+				leaseName.Append((job.Change > 0) ? $"Preflight CL {job.PreflightChange} against CL {job.Change}" : $"Preflight CL {job.PreflightChange} against latest");
 			}
 			else
 			{
-				LeaseName.Append((Job.Change > 0) ? $"CL {Job.Change}" : "Latest CL");
+				leaseName.Append((job.Change > 0) ? $"CL {job.Change}" : "Latest CL");
 			}
-			LeaseName.Append(CultureInfo.InvariantCulture, $" - {Job.Name}");
+			leaseName.Append(CultureInfo.InvariantCulture, $" - {job.Name}");
 
 			// Get the global settings
-			Globals Globals = await DatabaseService.GetGlobalsAsync();
+			Globals globals = await _databaseService.GetGlobalsAsync();
 
 			// Encode the payload
-			ExecuteJobTask Task = new ExecuteJobTask();
-			Task.JobId = Job.Id.ToString();
-			Task.BatchId = Batch.Id.ToString();
-			Task.LogId = LogId.ToString();
-			Task.JobName = LeaseName.ToString();
+			ExecuteJobTask task = new ExecuteJobTask();
+			task.JobId = job.Id.ToString();
+			task.BatchId = batch.Id.ToString();
+			task.LogId = logId.ToString();
+			task.JobName = leaseName.ToString();
 
-			List<HordeCommon.Rpc.Messages.AgentWorkspace> Workspaces = new List<HordeCommon.Rpc.Messages.AgentWorkspace>();
+			List<HordeCommon.Rpc.Messages.AgentWorkspace> workspaces = new List<HordeCommon.Rpc.Messages.AgentWorkspace>();
 
-			PerforceCluster? Cluster = Globals.FindPerforceCluster(Workspace.Cluster);
-			if (Cluster == null)
+			PerforceCluster? cluster = globals.FindPerforceCluster(workspace.Cluster);
+			if (cluster == null)
 			{
 				return null;
 			}
 
-			AgentWorkspace? AutoSdkWorkspace = bUseAutoSDK ? Agent.GetAutoSdkWorkspace(Cluster) : null;
-			if (AutoSdkWorkspace != null)
+			AgentWorkspace? autoSdkWorkspace = bUseAutoSdk ? agent.GetAutoSdkWorkspace(cluster) : null;
+			if (autoSdkWorkspace != null)
 			{
-				if (!await Agent.TryAddWorkspaceMessage(AutoSdkWorkspace, Cluster, PerforceLoadBalancer, Workspaces))
+				if (!await agent.TryAddWorkspaceMessage(autoSdkWorkspace, cluster, _perforceLoadBalancer, workspaces))
 				{
 					return null;
 				}
-				Task.AutoSdkWorkspace = Workspaces.Last();
+				task.AutoSdkWorkspace = workspaces.Last();
 			}
 
-			if (!await Agent.TryAddWorkspaceMessage(Workspace, Cluster, PerforceLoadBalancer, Workspaces))
+			if (!await agent.TryAddWorkspaceMessage(workspace, cluster, _perforceLoadBalancer, workspaces))
 			{
 				return null;
 			}
 
-			Task.Workspace = Workspaces.Last();
-			return Task;
+			task.Workspace = workspaces.Last();
+			return task;
 		}
 
 		/// <summary>
 		/// Send any badge updates for this job
 		/// </summary>
-		/// <param name="Job">The job being updated</param>
-		/// <param name="Graph">Graph for the job</param>
-		/// <param name="OldLabelStates">Previous badge states for the job</param>
+		/// <param name="job">The job being updated</param>
+		/// <param name="graph">Graph for the job</param>
+		/// <param name="oldLabelStates">Previous badge states for the job</param>
 		/// <returns>Async task</returns>
-		public async Task UpdateUgsBadges(IJob Job, IGraph Graph, IReadOnlyList<(LabelState, LabelOutcome)> OldLabelStates)
+		public async Task UpdateUgsBadges(IJob job, IGraph graph, IReadOnlyList<(LabelState, LabelOutcome)> oldLabelStates)
 		{
-			await UpdateUgsBadges(Job, Graph, OldLabelStates, Job.GetLabelStates(Graph));
+			await UpdateUgsBadges(job, graph, oldLabelStates, job.GetLabelStates(graph));
 		}
 
 		/// <summary>
 		/// Send any badge updates for this job
 		/// </summary>
-		/// <param name="Job">The job being updated</param>
-		/// <param name="Graph">Graph for the job</param>
-		/// <param name="OldLabelStates">Previous badge states for the job</param>
-		/// <param name="NewLabelStates">The new badge states for the job</param>
+		/// <param name="job">The job being updated</param>
+		/// <param name="graph">Graph for the job</param>
+		/// <param name="oldLabelStates">Previous badge states for the job</param>
+		/// <param name="newLabelStates">The new badge states for the job</param>
 		/// <returns>Async task</returns>
-		public async Task UpdateUgsBadges(IJob Job, IGraph Graph, IReadOnlyList<(LabelState, LabelOutcome)> OldLabelStates, IReadOnlyList<(LabelState, LabelOutcome)> NewLabelStates)
+		public async Task UpdateUgsBadges(IJob job, IGraph graph, IReadOnlyList<(LabelState, LabelOutcome)> oldLabelStates, IReadOnlyList<(LabelState, LabelOutcome)> newLabelStates)
 		{
-			if (!Job.ShowUgsBadges || Job.PreflightChange != 0)
+			if (!job.ShowUgsBadges || job.PreflightChange != 0)
 			{
 				return;
 			}
 
-			IReadOnlyDictionary<int, UgsBadgeState> OldStates = Job.GetUgsBadgeStates(Graph, OldLabelStates);
-			IReadOnlyDictionary<int, UgsBadgeState> NewStates = Job.GetUgsBadgeStates(Graph, NewLabelStates);
+			IReadOnlyDictionary<int, UgsBadgeState> oldStates = job.GetUgsBadgeStates(graph, oldLabelStates);
+			IReadOnlyDictionary<int, UgsBadgeState> newStates = job.GetUgsBadgeStates(graph, newLabelStates);
 
 			// Figure out a list of all the badges that have been modified
-			List<int> UpdateLabels = new List<int>();
-			foreach (KeyValuePair<int, UgsBadgeState> Pair in OldStates)
+			List<int> updateLabels = new List<int>();
+			foreach (KeyValuePair<int, UgsBadgeState> pair in oldStates)
 			{
-				if (!NewStates.ContainsKey(Pair.Key))
+				if (!newStates.ContainsKey(pair.Key))
 				{
-					UpdateLabels.Add(Pair.Key);
+					updateLabels.Add(pair.Key);
 				}
 			}
-			foreach (KeyValuePair<int, UgsBadgeState> Pair in NewStates)
+			foreach (KeyValuePair<int, UgsBadgeState> pair in newStates)
 			{
-				if (!OldStates.TryGetValue(Pair.Key, out UgsBadgeState PrevState) || PrevState != Pair.Value)
+				if (!oldStates.TryGetValue(pair.Key, out UgsBadgeState prevState) || prevState != pair.Value)
 				{
-					UpdateLabels.Add(Pair.Key);
+					updateLabels.Add(pair.Key);
 				}
 			}
 
 			// Cached stream for this job
-			IStream? Stream = null;
+			IStream? stream = null;
 
 			// Send all the updates
-			Dictionary<int, IUgsMetadata> MetadataCache = new Dictionary<int, IUgsMetadata>();
-			foreach (int LabelIdx in UpdateLabels)
+			Dictionary<int, IUgsMetadata> metadataCache = new Dictionary<int, IUgsMetadata>();
+			foreach (int labelIdx in updateLabels)
 			{
-				ILabel Label = Graph.Labels[LabelIdx];
+				ILabel label = graph.Labels[labelIdx];
 
 				// Skip if this label has no UGS name.
-				if (Label.UgsName == null)
+				if (label.UgsName == null)
 				{
 					continue;
 				}
 
 				// Get the new state
-				if (!NewStates.TryGetValue(LabelIdx, out UgsBadgeState NewState))
+				if (!newStates.TryGetValue(labelIdx, out UgsBadgeState newState))
 				{
-					NewState = UgsBadgeState.Skipped;
+					newState = UgsBadgeState.Skipped;
 				}
 
 				// Get the stream
-				if (Stream == null)
+				if (stream == null)
 				{
-					Stream = await StreamService.GetStreamAsync(Job.StreamId);
-					if (Stream == null)
+					stream = await _streamService.GetStreamAsync(job.StreamId);
+					if (stream == null)
 					{
-						Logger.LogError("Unable to fetch definition for stream {StreamId}", Job.StreamId);
+						_logger.LogError("Unable to fetch definition for stream {StreamId}", job.StreamId);
 						break;
 					}
 				}
 
 				// The changelist number to display the badge for
-				int Change;
-				if (Label.Change == LabelChange.Code)
+				int change;
+				if (label.Change == LabelChange.Code)
 				{
-					Change = Job.CodeChange;
+					change = job.CodeChange;
 				}
 				else
 				{
-					Change = Job.Change;
+					change = job.Change;
 				}
 
 				// Get the current metadata state
-				IUgsMetadata? Metadata;
-				if (!MetadataCache.TryGetValue(Change, out Metadata))
+				IUgsMetadata? metadata;
+				if (!metadataCache.TryGetValue(change, out metadata))
 				{
-					Metadata = await UgsMetadataCollection.FindOrAddAsync(Stream.Name, Change, Label.UgsProject);
-					MetadataCache[Change] = Metadata;
+					metadata = await _ugsMetadataCollection.FindOrAddAsync(stream.Name, change, label.UgsProject);
+					metadataCache[change] = metadata;
 				}
 
 				// Apply the update
-				Uri LabelUrl = new Uri(Settings.CurrentValue.DashboardUrl, $"job/{Job.Id}?label={LabelIdx}");
-				Logger.LogInformation("Updating state of badge {BadgeName} at {Change} to {NewState} ({LabelUrl})", Label.UgsName, Change, NewState, LabelUrl);
-				Metadata = await UgsMetadataCollection.UpdateBadgeAsync(Metadata, Label.UgsName!, LabelUrl, NewState);
-				MetadataCache[Change] = Metadata;
+				Uri labelUrl = new Uri(_settings.CurrentValue.DashboardUrl, $"job/{job.Id}?label={labelIdx}");
+				_logger.LogInformation("Updating state of badge {BadgeName} at {Change} to {NewState} ({LabelUrl})", label.UgsName, change, newState, labelUrl);
+				metadata = await _ugsMetadataCollection.UpdateBadgeAsync(metadata, label.UgsName!, labelUrl, newState);
+				metadataCache[change] = metadata;
 			}
 		}
 
 		/// <inheritdoc/>
-		public override async Task OnLeaseFinishedAsync(IAgent Agent, LeaseId LeaseId, ExecuteJobTask Task, LeaseOutcome Outcome, ReadOnlyMemory<byte> Output, ILogger Logger)
+		public override async Task OnLeaseFinishedAsync(IAgent agent, LeaseId leaseId, ExecuteJobTask task, LeaseOutcome outcome, ReadOnlyMemory<byte> output, ILogger logger)
 		{
-			await base.OnLeaseFinishedAsync(Agent, LeaseId, Task, Outcome, Output, Logger);
+			await base.OnLeaseFinishedAsync(agent, leaseId, task, outcome, output, logger);
 
-			if (Outcome != LeaseOutcome.Success)
+			if (outcome != LeaseOutcome.Success)
 			{
-				AgentId AgentId = Agent.Id;
-				JobId JobId = new JobId(Task.JobId);
-				SubResourceId BatchId = Task.BatchId.ToSubResourceId();
+				AgentId agentId = agent.Id;
+				JobId jobId = new JobId(task.JobId);
+				SubResourceId batchId = task.BatchId.ToSubResourceId();
 
 				// Update the batch
 				for (; ; )
 				{
-					IJob? Job = await Jobs.GetAsync(JobId);
-					if (Job == null)
+					IJob? job = await _jobs.GetAsync(jobId);
+					if (job == null)
 					{
 						break;
 					}
 
-					int BatchIdx = Job.Batches.FindIndex(x => x.Id == BatchId);
-					if (BatchIdx == -1)
+					int batchIdx = job.Batches.FindIndex(x => x.Id == batchId);
+					if (batchIdx == -1)
 					{
 						break;
 					}
 
-					IJobStepBatch Batch = Job.Batches[BatchIdx];
-					if (Batch.AgentId != AgentId)
+					IJobStepBatch batch = job.Batches[batchIdx];
+					if (batch.AgentId != agentId)
 					{
 						break;
 					}
 
-					int RunningStepIdx = Batch.Steps.FindIndex(x => x.State == JobStepState.Running);
+					int runningStepIdx = batch.Steps.FindIndex(x => x.State == JobStepState.Running);
 
-					JobStepBatchError Error;
-					if (Outcome == LeaseOutcome.Cancelled)
+					JobStepBatchError error;
+					if (outcome == LeaseOutcome.Cancelled)
 					{
-						Error = JobStepBatchError.Cancelled;
+						error = JobStepBatchError.Cancelled;
 					}
 					else
 					{
-						Error = JobStepBatchError.ExecutionError;
+						error = JobStepBatchError.ExecutionError;
 					}
 
-					IGraph Graph = await Graphs.GetAsync(Job.GraphHash);
-					Job = await Jobs.TryFailBatchAsync(Job, BatchIdx, Graph, Error);
+					IGraph graph = await _graphs.GetAsync(job.GraphHash);
+					job = await _jobs.TryFailBatchAsync(job, batchIdx, graph, error);
 
-					if (Job != null)
+					if (job != null)
 					{
-						if (Batch.Error != JobStepBatchError.None)
+						if (batch.Error != JobStepBatchError.None)
 						{
-							Logger.LogInformation("Failed job {JobId}, batch {BatchId} with error {Error}", Job.Id, Batch.Id, Batch.Error);
+							logger.LogInformation("Failed job {JobId}, batch {BatchId} with error {Error}", job.Id, batch.Id, batch.Error);
 						}
-						if (RunningStepIdx != -1)
+						if (runningStepIdx != -1)
 						{
-							await JobStepRefs.UpdateAsync(Job, Batch, Batch.Steps[RunningStepIdx], Graph);
+							await _jobStepRefs.UpdateAsync(job, batch, batch.Steps[runningStepIdx], graph);
 						}
 						break;
 					}
@@ -1023,37 +1010,37 @@ namespace Horde.Build.Tasks.Impl
 		/// <summary>
 		/// 
 		/// </summary>
-		/// <param name="Agent"></param>
-		/// <param name="JobId"></param>
-		/// <param name="BatchId"></param>
+		/// <param name="agent"></param>
+		/// <param name="jobId"></param>
+		/// <param name="batchId"></param>
 		/// <returns></returns>
-		async Task CancelLeaseAsync(IAgent Agent, JobId JobId, SubResourceId BatchId)
+		async Task CancelLeaseAsync(IAgent agent, JobId jobId, SubResourceId batchId)
 		{
-			Logger.LogDebug("Cancelling lease for job {JobId}, batch {BatchId}", JobId, BatchId);
+			_logger.LogDebug("Cancelling lease for job {JobId}, batch {BatchId}", jobId, batchId);
 
 			// Update the batch
 			for (; ; )
 			{
-				IJob? Job = await Jobs.GetAsync(JobId);
-				if (Job == null)
+				IJob? job = await _jobs.GetAsync(jobId);
+				if (job == null)
 				{
 					break;
 				}
 
-				int BatchIdx = Job.Batches.FindIndex(x => x.Id == BatchId);
-				if (BatchIdx == -1)
+				int batchIdx = job.Batches.FindIndex(x => x.Id == batchId);
+				if (batchIdx == -1)
 				{
 					break;
 				}
 
-				IJobStepBatch Batch = Job.Batches[BatchIdx];
-				if (Batch.AgentId != Agent.Id)
+				IJobStepBatch batch = job.Batches[batchIdx];
+				if (batch.AgentId != agent.Id)
 				{
 					break;
 				}
 
-				IJob? NewJob = await Jobs.TryCancelLeaseAsync(Job, BatchIdx);
-				if (NewJob != null)
+				IJob? newJob = await _jobs.TryCancelLeaseAsync(job, batchIdx);
+				if (newJob != null)
 				{
 					break;
 				}
@@ -1063,34 +1050,34 @@ namespace Horde.Build.Tasks.Impl
 		/// <summary>
 		/// Inserts an item into the queue
 		/// </summary>
-		/// <param name="Stream">The stream containing the job</param>
-		/// <param name="Job"></param>
-		/// <param name="BatchIdx"></param>
-		/// <param name="PoolId">The pool to use</param>
-		/// <param name="Workspace">The workspace for this item to run in</param>
-		/// <param name="UseAutoSDK">Whether or not to use the AutoSDK</param>
+		/// <param name="stream">The stream containing the job</param>
+		/// <param name="job"></param>
+		/// <param name="batchIdx"></param>
+		/// <param name="poolId">The pool to use</param>
+		/// <param name="workspace">The workspace for this item to run in</param>
+		/// <param name="useAutoSdk">Whether or not to use the AutoSDK</param>
 		/// <returns></returns>
-		void InsertQueueItem(IStream Stream, IJob Job, int BatchIdx, PoolId PoolId, AgentWorkspace Workspace, bool UseAutoSDK)
+		void InsertQueueItem(IStream stream, IJob job, int batchIdx, PoolId poolId, AgentWorkspace workspace, bool useAutoSdk)
 		{
-			Logger.LogDebug("Adding queued job {JobId}, batch {BatchId} [Pool: {Pool}, Workspace: {Workspace}]", Job.Id, Job.Batches[BatchIdx].Id, PoolId, Workspace.Identifier);
+			_logger.LogDebug("Adding queued job {JobId}, batch {BatchId} [Pool: {Pool}, Workspace: {Workspace}]", job.Id, job.Batches[batchIdx].Id, poolId, workspace.Identifier);
 
-			QueueItem NewItem = new QueueItem(Stream, Job, BatchIdx, PoolId, Workspace, UseAutoSDK);
-			BatchIdToQueueItem[NewItem.Id] = NewItem;
-			Queue.Add(NewItem);
+			QueueItem newItem = new QueueItem(stream, job, batchIdx, poolId, workspace, useAutoSdk);
+			_batchIdToQueueItem[newItem.Id] = newItem;
+			_queue.Add(newItem);
 
-			AssignQueueItemToAnyWaiter(NewItem);
+			AssignQueueItemToAnyWaiter(newItem);
 		}
 
 		/// <summary>
 		/// Removes an item from the queue
 		/// </summary>
-		/// <param name="Item">Item to remove</param>
-		void RemoveQueueItem(QueueItem Item)
+		/// <param name="item">Item to remove</param>
+		void RemoveQueueItem(QueueItem item)
 		{
-			Logger.LogDebug("Removing queued job {JobId}, batch {BatchId}", Item.Job.Id, Item.Batch.Id);
+			_logger.LogDebug("Removing queued job {JobId}, batch {BatchId}", item._job.Id, item.Batch.Id);
 
-			Queue.Remove(Item);
-			BatchIdToQueueItem.Remove(Item.Id);
+			_queue.Remove(item);
+			_batchIdToQueueItem.Remove(item.Id);
 		}
 	}
 }
diff --git a/Engine/Source/Programs/Horde/Horde.Build/Jobs/Impl/JobTimingCollection.cs b/Engine/Source/Programs/Horde/Horde.Build/Jobs/Impl/JobTimingCollection.cs
--- a/Engine/Source/Programs/Horde/Horde.Build/Jobs/Impl/JobTimingCollection.cs
+++ b/Engine/Source/Programs/Horde/Horde.Build/Jobs/Impl/JobTimingCollection.cs
@@ -1,18 +1,14 @@
 // Copyright Epic Games, Inc. All Rights Reserved.
 
-using Amazon.Runtime.Internal.Util;
-using Amazon.S3.Model.Internal.MarshallTransformations;
+using System.Collections.Generic;
+using System.Diagnostics.CodeAnalysis;
+using System.Linq;
+using System.Threading.Tasks;
 using Horde.Build.Models;
 using Horde.Build.Services;
 using Horde.Build.Utilities;
-using MongoDB.Bson;
 using MongoDB.Bson.Serialization.Attributes;
 using MongoDB.Driver;
-using System;
-using System.Collections.Generic;
-using System.Diagnostics.CodeAnalysis;
-using System.Linq;
-using System.Threading.Tasks;
 
 namespace Horde.Build.Collections.Impl
 {
@@ -36,15 +32,15 @@ namespace Horde.Build.Collections.Impl
 			[BsonConstructor]
 			private JobStepTimingDocument()
 			{
-				this.Name = null!;
+				Name = null!;
 			}
 
-			public JobStepTimingDocument(string Name, float? AverageWaitTime, float? AverageInitTime, float? AverageDuration)
+			public JobStepTimingDocument(string name, float? averageWaitTime, float? averageInitTime, float? averageDuration)
 			{
-				this.Name = Name;
-				this.AverageWaitTime = AverageWaitTime;
-				this.AverageInitTime = AverageInitTime;
-				this.AverageDuration = AverageDuration;
+				Name = name;
+				AverageWaitTime = averageWaitTime;
+				AverageInitTime = averageInitTime;
+				AverageDuration = averageDuration;
 			}
 		}
 
@@ -58,56 +54,56 @@ namespace Horde.Build.Collections.Impl
 			public int UpdateIndex { get; set; }
 
 			[BsonIgnore]
-			public Dictionary<string, IJobStepTiming>? NameToStep;
+			public Dictionary<string, IJobStepTiming>? _nameToStep;
 
-			public bool TryGetStepTiming(string Name, [NotNullWhen(true)] out IJobStepTiming? Timing)
+			public bool TryGetStepTiming(string name, [NotNullWhen(true)] out IJobStepTiming? timing)
 			{
-				if (NameToStep == null)
+				if (_nameToStep == null)
 				{
-					NameToStep = new Dictionary<string, IJobStepTiming>();
-					foreach (JobStepTimingDocument Step in Steps)
+					_nameToStep = new Dictionary<string, IJobStepTiming>();
+					foreach (JobStepTimingDocument step in Steps)
 					{
-						if (NameToStep.ContainsKey(Step.Name))
+						if (_nameToStep.ContainsKey(step.Name))
 						{
-							Serilog.Log.Logger.Warning("Step {Name} appears twice in job timing document {Id}", Step.Name, Id);
+							Serilog.Log.Logger.Warning("Step {Name} appears twice in job timing document {Id}", step.Name, Id);
 						}
-						NameToStep[Step.Name] = Step;
+						_nameToStep[step.Name] = step;
 					}
 				}
-				return NameToStep.TryGetValue(Name, out Timing);
+				return _nameToStep.TryGetValue(name, out timing);
 			}
 		}
 
 		/// <summary>
 		/// Collection of timing documents
 		/// </summary>
-		IMongoCollection<JobTimingDocument> Collection;
+		readonly IMongoCollection<JobTimingDocument> _collection;
 
 		/// <summary>
 		/// Constructor
 		/// </summary>
-		/// <param name="DatabaseService">The database service singleton</param>
-		public JobTimingCollection(DatabaseService DatabaseService)
+		/// <param name="databaseService">The database service singleton</param>
+		public JobTimingCollection(DatabaseService databaseService)
 		{
-			Collection = DatabaseService.GetCollection<JobTimingDocument>("JobTiming");
+			_collection = databaseService.GetCollection<JobTimingDocument>("JobTiming");
 		}
 
 		/// <inheritdoc/>
-		public async Task<IJobTiming?> TryAddAsync(JobId JobId, List<JobStepTimingData> Steps)
+		public async Task<IJobTiming?> TryAddAsync(JobId jobId, List<JobStepTimingData> steps)
 		{
-			JobTimingDocument JobTiming = new JobTimingDocument();
-			JobTiming.Id = JobId;
-			JobTiming.Steps.AddRange(Steps.Select(x => new JobStepTimingDocument(x.Name, x.AverageWaitTime, x.AverageInitTime, x.AverageDuration)));
-			JobTiming.UpdateIndex = 1;
+			JobTimingDocument jobTiming = new JobTimingDocument();
+			jobTiming.Id = jobId;
+			jobTiming.Steps.AddRange(steps.Select(x => new JobStepTimingDocument(x.Name, x.AverageWaitTime, x.AverageInitTime, x.AverageDuration)));
+			jobTiming.UpdateIndex = 1;
 
 			try
 			{
-				await Collection.InsertOneAsync(JobTiming);
-				return JobTiming;
+				await _collection.InsertOneAsync(jobTiming);
+				return jobTiming;
 			}
-			catch (MongoWriteException Ex)
+			catch (MongoWriteException ex)
 			{
-				if (Ex.WriteError.Category == ServerErrorCategory.DuplicateKey)
+				if (ex.WriteError.Category == ServerErrorCategory.DuplicateKey)
 				{
 					return null;
 				}
@@ -119,27 +115,27 @@ namespace Horde.Build.Collections.Impl
 		}
 
 		/// <inheritdoc/>
-		public async Task<IJobTiming?> TryGetAsync(JobId JobId)
+		public async Task<IJobTiming?> TryGetAsync(JobId jobId)
 		{
-			return await Collection.Find(x => x.Id == JobId).FirstOrDefaultAsync();
+			return await _collection.Find(x => x.Id == jobId).FirstOrDefaultAsync();
 		}
 
 		/// <inheritdoc/>
-		public async Task<IJobTiming?> TryAddStepsAsync(IJobTiming JobTiming, List<JobStepTimingData> Steps)
+		public async Task<IJobTiming?> TryAddStepsAsync(IJobTiming jobTiming, List<JobStepTimingData> steps)
 		{
-			JobTimingDocument JobTimingDocument = (JobTimingDocument)JobTiming;
-			List<JobStepTimingDocument> StepDocuments = Steps.ConvertAll(x => new JobStepTimingDocument(x.Name, x.AverageWaitTime, x.AverageInitTime, x.AverageDuration));
+			JobTimingDocument jobTimingDocument = (JobTimingDocument)jobTiming;
+			List<JobStepTimingDocument> stepDocuments = steps.ConvertAll(x => new JobStepTimingDocument(x.Name, x.AverageWaitTime, x.AverageInitTime, x.AverageDuration));
 
-			FilterDefinition<JobTimingDocument> Filter = Builders<JobTimingDocument>.Filter.Expr(x => x.Id == JobTimingDocument.Id && x.UpdateIndex == JobTimingDocument.UpdateIndex);
-			UpdateDefinition<JobTimingDocument> Update = Builders<JobTimingDocument>.Update.Set(x => x.UpdateIndex, JobTimingDocument.UpdateIndex + 1).PushEach(x => x.Steps, StepDocuments);
+			FilterDefinition<JobTimingDocument> filter = Builders<JobTimingDocument>.Filter.Expr(x => x.Id == jobTimingDocument.Id && x.UpdateIndex == jobTimingDocument.UpdateIndex);
+			UpdateDefinition<JobTimingDocument> update = Builders<JobTimingDocument>.Update.Set(x => x.UpdateIndex, jobTimingDocument.UpdateIndex + 1).PushEach(x => x.Steps, stepDocuments);
 
-			UpdateResult Result = await Collection.UpdateOneAsync(Filter, Update);
-			if (Result.ModifiedCount > 0)
+			UpdateResult result = await _collection.UpdateOneAsync(filter, update);
+			if (result.ModifiedCount > 0)
 			{
-				JobTimingDocument.NameToStep = null;
-				JobTimingDocument.Steps.AddRange(StepDocuments);
-				JobTimingDocument.UpdateIndex++;
-				return JobTimingDocument;
+				jobTimingDocument._nameToStep = null;
+				jobTimingDocument.Steps.AddRange(stepDocuments);
+				jobTimingDocument.UpdateIndex++;
+				return jobTimingDocument;
 			}
 			return null;
 		}
diff --git a/Engine/Source/Programs/Horde/Horde.Build/Jobs/Impl/TemplateCollection.cs b/Engine/Source/Programs/Horde/Horde.Build/Jobs/Impl/TemplateCollection.cs
--- a/Engine/Source/Programs/Horde/Horde.Build/Jobs/Impl/TemplateCollection.cs
+++ b/Engine/Source/Programs/Horde/Horde.Build/Jobs/Impl/TemplateCollection.cs
@@ -1,20 +1,19 @@
 // Copyright Epic Games, Inc. All Rights Reserved.
 
+using System;
+using System.Collections.Generic;
+using System.Linq;
+using System.Threading.Tasks;
 using EpicGames.Core;
-using HordeCommon;
-using Horde.Build.Api;
 using Horde.Build.Models;
 using Horde.Build.Services;
 using Horde.Build.Utilities;
+using HordeCommon;
 using Microsoft.Extensions.Caching.Memory;
 using MongoDB.Bson;
 using MongoDB.Bson.Serialization;
 using MongoDB.Bson.Serialization.Attributes;
 using MongoDB.Driver;
-using System;
-using System.Collections.Generic;
-using System.Linq;
-using System.Threading.Tasks;
 
 namespace Horde.Build.Collections.Impl
 {
@@ -58,39 +57,39 @@ namespace Horde.Build.Collections.Impl
 				Name = null!;
 			}
 
-			public TemplateDocument(string Name, Priority? Priority, bool bAllowPreflights, bool UpdateIssues, bool PromoteIssuesByDefault, string? InitialAgentType, string? SubmitNewChange, string? SubmitDescription, List<string>? Arguments, List<Parameter>? Parameters)
+			public TemplateDocument(string name, Priority? priority, bool bAllowPreflights, bool updateIssues, bool promoteIssuesByDefault, string? initialAgentType, string? submitNewChange, string? submitDescription, List<string>? arguments, List<Parameter>? parameters)
 			{
-				this.Name = Name;
-				this.Priority = Priority;
-				this.AllowPreflights = bAllowPreflights;
-				this.UpdateIssues = UpdateIssues;
-				this.PromoteIssuesByDefault = PromoteIssuesByDefault;
-				this.InitialAgentType = InitialAgentType;
-				this.SubmitNewChange = SubmitNewChange;
-				this.SubmitDescription = SubmitDescription;
-				this.Arguments = Arguments ?? new List<string>();
-				this.Parameters = Parameters ?? new List<Parameter>();
+				Name = name;
+				Priority = priority;
+				AllowPreflights = bAllowPreflights;
+				UpdateIssues = updateIssues;
+				PromoteIssuesByDefault = promoteIssuesByDefault;
+				InitialAgentType = initialAgentType;
+				SubmitNewChange = submitNewChange;
+				SubmitDescription = submitDescription;
+				Arguments = arguments ?? new List<string>();
+				Parameters = parameters ?? new List<Parameter>();
 
 				// Compute the hash once all other fields have been set
-				this.Id = ContentHash.SHA1(BsonExtensionMethods.ToBson(this));
+				Id = ContentHash.SHA1(BsonExtensionMethods.ToBson(this));
 			}
 		}
 
 		/// <summary>
 		/// Template documents
 		/// </summary>
-		IMongoCollection<TemplateDocument> Templates;
+		readonly IMongoCollection<TemplateDocument> _templates;
 
 		/// <summary>
 		/// Cache of template documents
 		/// </summary>
-		MemoryCache TemplateCache;
+		readonly MemoryCache _templateCache;
 
 		/// <summary>
 		/// Constructor
 		/// </summary>
-		/// <param name="DatabaseService">The database service singleton</param>
-		public TemplateCollection(DatabaseService DatabaseService)
+		/// <param name="databaseService">The database service singleton</param>
+		public TemplateCollection(DatabaseService databaseService)
 		{
 			// Ensure discriminator cannot be registered twice (throws exception). Can otherwise happen during unit tests.
 			if (BsonSerializer.LookupDiscriminatorConvention(typeof(JobsTabColumn)) == null)
@@ -98,56 +97,56 @@ namespace Horde.Build.Collections.Impl
 				BsonSerializer.RegisterDiscriminatorConvention(typeof(JobsTabColumn), new DefaultDiscriminatorConvention(typeof(JobsTabColumn), typeof(JobsTabLabelColumn)));	
 			}
 			
-			Templates = DatabaseService.GetCollection<TemplateDocument>("Templates");
+			_templates = databaseService.GetCollection<TemplateDocument>("Templates");
 
-			MemoryCacheOptions Options = new MemoryCacheOptions();
-			TemplateCache = new MemoryCache(Options);
+			MemoryCacheOptions options = new MemoryCacheOptions();
+			_templateCache = new MemoryCache(options);
 		}
 
 		/// <inheritdoc/>
 		public void Dispose()
 		{
-			TemplateCache.Dispose();
+			_templateCache.Dispose();
 		}
 
 		/// <inheritdoc/>
-		public async Task<ITemplate> AddAsync(string Name, Priority? Priority, bool bAllowPreflights, bool bUpdateIssues, bool bPromoteIssuesByDefault, string? InitialAgentType, string? SubmitNewChange, string? SubmitDescription, List<string>? Arguments, List<Parameter>? Parameters)
+		public async Task<ITemplate> AddAsync(string name, Priority? priority, bool bAllowPreflights, bool bUpdateIssues, bool bPromoteIssuesByDefault, string? initialAgentType, string? submitNewChange, string? submitDescription, List<string>? arguments, List<Parameter>? parameters)
 		{
-			TemplateDocument Template = new TemplateDocument(Name, Priority, bAllowPreflights, bUpdateIssues, bPromoteIssuesByDefault, InitialAgentType, SubmitNewChange, SubmitDescription, Arguments, Parameters);
-			if (await GetAsync(Template.Id) == null)
+			TemplateDocument template = new TemplateDocument(name, priority, bAllowPreflights, bUpdateIssues, bPromoteIssuesByDefault, initialAgentType, submitNewChange, submitDescription, arguments, parameters);
+			if (await GetAsync(template.Id) == null)
 			{
-				await Templates.ReplaceOneAsync(x => x.Id == Template.Id, Template, new ReplaceOptions { IsUpsert = true });
+				await _templates.ReplaceOneAsync(x => x.Id == template.Id, template, new ReplaceOptions { IsUpsert = true });
 			}
-			return Template;
+			return template;
 		}
 
 		/// <inheritdoc/>
 		public async Task<List<ITemplate>> FindAllAsync()
 		{
-			List<TemplateDocument> Results = await Templates.Find(FilterDefinition<TemplateDocument>.Empty).ToListAsync();
-			return Results.ConvertAll<ITemplate>(x => x);
+			List<TemplateDocument> results = await _templates.Find(FilterDefinition<TemplateDocument>.Empty).ToListAsync();
+			return results.ConvertAll<ITemplate>(x => x);
 		}
 
 		/// <inheritdoc/>
-		public async Task<ITemplate?> GetAsync(ContentHash TemplateId)
+		public async Task<ITemplate?> GetAsync(ContentHash templateId)
 		{
-			object? Result;
-			if (TemplateCache.TryGetValue(TemplateId, out Result))
+			object? result;
+			if (_templateCache.TryGetValue(templateId, out result))
 			{
-				return (ITemplate?)Result;
+				return (ITemplate?)result;
 			}
 
-			ITemplate? Template = await Templates.Find<TemplateDocument>(x => x.Id == TemplateId).FirstOrDefaultAsync();
-			if (Template != null)
+			ITemplate? template = await _templates.Find<TemplateDocument>(x => x.Id == templateId).FirstOrDefaultAsync();
+			if (template != null)
 			{
-				using (ICacheEntry Entry = TemplateCache.CreateEntry(TemplateId))
+				using (ICacheEntry entry = _templateCache.CreateEntry(templateId))
 				{
-					Entry.SetSlidingExpiration(TimeSpan.FromMinutes(30.0));
-					Entry.SetValue(Template);
+					entry.SetSlidingExpiration(TimeSpan.FromMinutes(30.0));
+					entry.SetValue(template);
 				}
 			}
 
-			return Template;
+			return template;
 		}
 	}
 }
diff --git a/Engine/Source/Programs/Horde/Horde.Build/Jobs/Impl/TestDataCollection.cs b/Engine/Source/Programs/Horde/Horde.Build/Jobs/Impl/TestDataCollection.cs
--- a/Engine/Source/Programs/Horde/Horde.Build/Jobs/Impl/TestDataCollection.cs
+++ b/Engine/Source/Programs/Horde/Horde.Build/Jobs/Impl/TestDataCollection.cs
@@ -1,21 +1,18 @@
 // Copyright Epic Games, Inc. All Rights Reserved.
 
+using System;
+using System.Collections.Generic;
+using System.Linq;
+using System.Threading.Tasks;
 using Horde.Build.Models;
 using Horde.Build.Services;
 using Horde.Build.Utilities;
 using MongoDB.Bson;
-using MongoDB.Bson.Serialization.Attributes;
 using MongoDB.Driver;
-using System;
-using System.Collections.Generic;
-using System.Diagnostics.CodeAnalysis;
-using System.Linq;
-using System.Threading.Tasks;
 
 namespace Horde.Build.Collections.Impl
 {
 	using JobId = ObjectId<IJob>;
-	using ProjectId = StringId<IProject>;
 	using StreamId = StringId<IStream>;
 	using TemplateRefId = StringId<TemplateRef>;
 
@@ -40,96 +37,96 @@ namespace Horde.Build.Collections.Impl
 
 			private TestDataDocument()
 			{
-				this.Key = String.Empty;
-				this.Data = new BsonDocument();
+				Key = String.Empty;
+				Data = new BsonDocument();
 			}
 
-			public TestDataDocument(IJob Job, IJobStep JobStep, string Key, BsonDocument Value)
+			public TestDataDocument(IJob job, IJobStep jobStep, string key, BsonDocument value)
 			{
-				this.Id = ObjectId.GenerateNewId();
-				this.StreamId = Job.StreamId;
-				this.TemplateRefId = Job.TemplateId;
-				this.JobId = Job.Id;
-				this.StepId = JobStep.Id;
-				this.Change = Job.Change;
-				this.Key = Key;
-				this.Data = Value;
+				Id = ObjectId.GenerateNewId();
+				StreamId = job.StreamId;
+				TemplateRefId = job.TemplateId;
+				JobId = job.Id;
+				StepId = jobStep.Id;
+				Change = job.Change;
+				Key = key;
+				Data = value;
 			}
 		}
 
 		/// <summary>
 		/// The stream collection
 		/// </summary>
-		IMongoCollection<TestDataDocument> TestDataDocuments;
+		readonly IMongoCollection<TestDataDocument> _testDataDocuments;
 
 		/// <summary>
 		/// Constructor
 		/// </summary>
-		/// <param name="DatabaseService">The database service instance</param>
-		public TestDataCollection(DatabaseService DatabaseService)
+		/// <param name="databaseService">The database service instance</param>
+		public TestDataCollection(DatabaseService databaseService)
 		{
-			TestDataDocuments = DatabaseService.GetCollection<TestDataDocument>("TestData");
+			_testDataDocuments = databaseService.GetCollection<TestDataDocument>("TestData");
 
-			if (!DatabaseService.ReadOnlyMode)
+			if (!databaseService.ReadOnlyMode)
 			{
-				TestDataDocuments.Indexes.CreateOne(new CreateIndexModel<TestDataDocument>(Builders<TestDataDocument>.IndexKeys.Ascending(x => x.StreamId).Ascending(x => x.Change).Ascending(x => x.Key)));
-				TestDataDocuments.Indexes.CreateOne(new CreateIndexModel<TestDataDocument>(Builders<TestDataDocument>.IndexKeys.Ascending(x => x.JobId).Ascending(x => x.StepId).Ascending(x => x.Key), new CreateIndexOptions { Unique = true }));
+				_testDataDocuments.Indexes.CreateOne(new CreateIndexModel<TestDataDocument>(Builders<TestDataDocument>.IndexKeys.Ascending(x => x.StreamId).Ascending(x => x.Change).Ascending(x => x.Key)));
+				_testDataDocuments.Indexes.CreateOne(new CreateIndexModel<TestDataDocument>(Builders<TestDataDocument>.IndexKeys.Ascending(x => x.JobId).Ascending(x => x.StepId).Ascending(x => x.Key), new CreateIndexOptions { Unique = true }));
 			}
 		}
 
 		/// <inheritdoc/>
-		public async Task<ITestData> AddAsync(IJob Job, IJobStep Step, string Key, BsonDocument Value)
+		public async Task<ITestData> AddAsync(IJob job, IJobStep step, string key, BsonDocument value)
 		{
-			TestDataDocument NewTestData = new TestDataDocument(Job, Step, Key, Value);
-			await TestDataDocuments.InsertOneAsync(NewTestData);
-			return NewTestData;
+			TestDataDocument newTestData = new TestDataDocument(job, step, key, value);
+			await _testDataDocuments.InsertOneAsync(newTestData);
+			return newTestData;
 		}
 
 		/// <inheritdoc/>
-		public async Task<ITestData?> GetAsync(ObjectId Id)
+		public async Task<ITestData?> GetAsync(ObjectId id)
 		{
-			return await TestDataDocuments.Find<TestDataDocument>(x => x.Id == Id).FirstOrDefaultAsync();
+			return await _testDataDocuments.Find<TestDataDocument>(x => x.Id == id).FirstOrDefaultAsync();
 		}
 
 		/// <inheritdoc/>
-		public async Task<List<ITestData>> FindAsync(StreamId? StreamId, int? MinChange, int? MaxChange, JobId? JobId, SubResourceId? StepId, string? Key = null, int Index = 0, int Count = 10)
+		public async Task<List<ITestData>> FindAsync(StreamId? streamId, int? minChange, int? maxChange, JobId? jobId, SubResourceId? stepId, string? key = null, int index = 0, int count = 10)
 		{
-			FilterDefinition<TestDataDocument> Filter = FilterDefinition<TestDataDocument>.Empty;
-			if (StreamId != null)
+			FilterDefinition<TestDataDocument> filter = FilterDefinition<TestDataDocument>.Empty;
+			if (streamId != null)
 			{
-				Filter &= Builders<TestDataDocument>.Filter.Eq(x => x.StreamId, StreamId.Value);
-				if (MinChange != null)
+				filter &= Builders<TestDataDocument>.Filter.Eq(x => x.StreamId, streamId.Value);
+				if (minChange != null)
 				{
-					Filter &= Builders<TestDataDocument>.Filter.Gte(x => x.Change, MinChange.Value);
+					filter &= Builders<TestDataDocument>.Filter.Gte(x => x.Change, minChange.Value);
 				}
-				if (MaxChange != null)
+				if (maxChange != null)
 				{
-					Filter &= Builders<TestDataDocument>.Filter.Lte(x => x.Change, MaxChange.Value);
+					filter &= Builders<TestDataDocument>.Filter.Lte(x => x.Change, maxChange.Value);
 				}
 			}
-			if (JobId != null)
+			if (jobId != null)
 			{
-				Filter &= Builders<TestDataDocument>.Filter.Eq(x => x.JobId, JobId.Value);
-				if (StepId != null)
+				filter &= Builders<TestDataDocument>.Filter.Eq(x => x.JobId, jobId.Value);
+				if (stepId != null)
 				{
-					Filter &= Builders<TestDataDocument>.Filter.Eq(x => x.StepId, StepId.Value);
+					filter &= Builders<TestDataDocument>.Filter.Eq(x => x.StepId, stepId.Value);
 				}
 			}
-			if (Key != null)
+			if (key != null)
 			{
-				Filter &= Builders<TestDataDocument>.Filter.Eq(x => x.Key, Key);
+				filter &= Builders<TestDataDocument>.Filter.Eq(x => x.Key, key);
 			}
 
-			SortDefinition<TestDataDocument> Sort = Builders<TestDataDocument>.Sort.Ascending(x => x.StreamId).Descending(x => x.Change);
+			SortDefinition<TestDataDocument> sort = Builders<TestDataDocument>.Sort.Ascending(x => x.StreamId).Descending(x => x.Change);
 
-			List<TestDataDocument> Results = await TestDataDocuments.Find(Filter).Sort(Sort).Skip(Index).Limit(Count).ToListAsync();
-			return Results.ConvertAll<ITestData>(x => x);
+			List<TestDataDocument> results = await _testDataDocuments.Find(filter).Sort(sort).Skip(index).Limit(count).ToListAsync();
+			return results.ConvertAll<ITestData>(x => x);
 		}
 
 		/// <inheritdoc/>
-		public async Task DeleteAsync(ObjectId Id)
+		public async Task DeleteAsync(ObjectId id)
 		{
-			await TestDataDocuments.DeleteOneAsync<TestDataDocument>(x => x.Id == Id);
+			await _testDataDocuments.DeleteOneAsync<TestDataDocument>(x => x.Id == id);
 		}
 	}
 }
