Horde: Fix coding conventions within Horde.Build.

diff --git a/Engine/Source/Programs/Horde/Horde.Build/Jobs/SubResourceId.cs b/Engine/Source/Programs/Horde/Horde.Build/Jobs/SubResourceId.cs
--- a/Engine/Source/Programs/Horde/Horde.Build/Jobs/SubResourceId.cs
+++ b/Engine/Source/Programs/Horde/Horde.Build/Jobs/SubResourceId.cs
@@ -1,14 +1,11 @@
 // Copyright Epic Games, Inc. All Rights Reserved.
 
-using MongoDB.Bson.Serialization;
-using MongoDB.Bson.Serialization.Attributes;
 using System;
-using System.Collections.Generic;
 using System.ComponentModel;
 using System.Diagnostics;
 using System.Globalization;
-using System.Linq;
-using System.Threading.Tasks;
+using MongoDB.Bson.Serialization;
+using MongoDB.Bson.Serialization.Attributes;
 
 namespace Horde.Build.Models
 {
@@ -30,10 +27,10 @@ namespace Horde.Build.Models
 		/// <summary>
 		/// Constructor
 		/// </summary>
-		/// <param name="Value">New identifier for this subresource</param>
-		public SubResourceId(ushort Value)
+		/// <param name="value">New identifier for this subresource</param>
+		public SubResourceId(ushort value)
 		{
-			this.Value = Value;
+			Value = value;
 		}
 
 		/// <summary>
@@ -58,11 +55,11 @@ namespace Horde.Build.Models
 		/// <summary>
 		/// Parse a subresource id from a string
 		/// </summary>
-		/// <param name="Text">Text to parse</param>
+		/// <param name="text">Text to parse</param>
 		/// <returns>New subresource id</returns>
-		public static SubResourceId Parse(string Text)
+		public static SubResourceId Parse(string text)
 		{
-			return new SubResourceId(ushort.Parse(Text, NumberStyles.HexNumber, CultureInfo.InvariantCulture));
+			return new SubResourceId(UInt16.Parse(text, NumberStyles.HexNumber, CultureInfo.InvariantCulture));
 		}
 
 		/// <summary>
@@ -75,15 +72,15 @@ namespace Horde.Build.Models
 		}
 
 		/// <inheritdoc/>
-		public override bool Equals(object? Other)
+		public override bool Equals(object? other)
 		{
-			return Other is SubResourceId OtherId && Equals(OtherId);
+			return other is SubResourceId otherId && Equals(otherId);
 		}
 
 		/// <inheritdoc/>
-		public bool Equals(SubResourceId Other)
+		public bool Equals(SubResourceId other)
 		{
-			return Value == Other.Value;
+			return Value == other.Value;
 		}
 
 		/// <inheritdoc/>
@@ -95,23 +92,23 @@ namespace Horde.Build.Models
 		/// <summary>
 		/// Equality operator for identifiers
 		/// </summary>
-		/// <param name="Left">First identifier to compare</param>
-		/// <param name="Right">Second identifier to compare</param>
+		/// <param name="left">First identifier to compare</param>
+		/// <param name="right">Second identifier to compare</param>
 		/// <returns>True if the identifiers are equal</returns>
-		public static bool operator ==(SubResourceId Left, SubResourceId Right)
+		public static bool operator ==(SubResourceId left, SubResourceId right)
 		{
-			return Left.Value == Right.Value;
+			return left.Value == right.Value;
 		}
 
 		/// <summary>
 		/// Inequality operator for identifiers
 		/// </summary>
-		/// <param name="Left">First identifier to compare</param>
-		/// <param name="Right">Second identifier to compare</param>
+		/// <param name="left">First identifier to compare</param>
+		/// <param name="right">Second identifier to compare</param>
 		/// <returns>True if the identifiers are equal</returns>
-		public static bool operator !=(SubResourceId Left, SubResourceId Right)
+		public static bool operator !=(SubResourceId left, SubResourceId right)
 		{
-			return Left.Value != Right.Value;
+			return left.Value != right.Value;
 		}
 	}
 
@@ -123,11 +120,11 @@ namespace Horde.Build.Models
 		/// <summary>
 		/// Parse a string as a subresource identifier
 		/// </summary>
-		/// <param name="Text">Text to parse</param>
+		/// <param name="text">Text to parse</param>
 		/// <returns>The new subresource identifier</returns>
-		public static SubResourceId ToSubResourceId(this string Text)
+		public static SubResourceId ToSubResourceId(this string text)
 		{
-			return new SubResourceId(ushort.Parse(Text, NumberStyles.HexNumber, CultureInfo.InvariantCulture));
+			return new SubResourceId(UInt16.Parse(text, NumberStyles.HexNumber, CultureInfo.InvariantCulture));
 		}
 	}
 
@@ -140,28 +137,28 @@ namespace Horde.Build.Models
 		public Type ValueType => typeof(SubResourceId);
 
 		/// <inheritdoc/>
-		public object Deserialize(BsonDeserializationContext Context, BsonDeserializationArgs Args)
+		public object Deserialize(BsonDeserializationContext context, BsonDeserializationArgs args)
 		{
-			return new SubResourceId((ushort)Context.Reader.ReadInt32());
+			return new SubResourceId((ushort)context.Reader.ReadInt32());
 		}
 
 		/// <inheritdoc/>
-		public void Serialize(BsonSerializationContext Context, BsonSerializationArgs Args, object Value)
+		public void Serialize(BsonSerializationContext context, BsonSerializationArgs args, object value)
 		{
-			SubResourceId Id = (SubResourceId)Value;
-			Context.Writer.WriteInt32((int)Id.Value);
+			SubResourceId id = (SubResourceId)value;
+			context.Writer.WriteInt32((int)id.Value);
 		}
 
 		/// <inheritdoc/>
-		public void Serialize(BsonSerializationContext Context, BsonSerializationArgs Args, SubResourceId Id)
+		public void Serialize(BsonSerializationContext context, BsonSerializationArgs args, SubResourceId id)
 		{
-			Context.Writer.WriteInt32((int)Id.Value);
+			context.Writer.WriteInt32((int)id.Value);
 		}
 
 		/// <inheritdoc/>
-		SubResourceId IBsonSerializer<SubResourceId>.Deserialize(BsonDeserializationContext Context, BsonDeserializationArgs Vrgs)
+		SubResourceId IBsonSerializer<SubResourceId>.Deserialize(BsonDeserializationContext context, BsonDeserializationArgs vrgs)
 		{
-			return new SubResourceId((ushort)Context.Reader.ReadInt32());
+			return new SubResourceId((ushort)context.Reader.ReadInt32());
 		}
 	}
 
@@ -171,29 +168,29 @@ namespace Horde.Build.Models
 	sealed class SubResourceIdTypeConverter : TypeConverter
 	{
 		/// <inheritdoc/>
-		public override bool CanConvertFrom(ITypeDescriptorContext? Context, Type SourceType)
+		public override bool CanConvertFrom(ITypeDescriptorContext? context, Type sourceType)
 		{
-			return SourceType == typeof(string);
+			return sourceType == typeof(string);
 		}
 
 		/// <inheritdoc/>
-		public override object ConvertFrom(ITypeDescriptorContext? Context, CultureInfo? Culture, object Value)
+		public override object ConvertFrom(ITypeDescriptorContext? context, CultureInfo? culture, object value)
 		{
-			return SubResourceId.Parse((string)Value);
+			return SubResourceId.Parse((string)value);
 		}
 
 		/// <inheritdoc/>
-		public override bool CanConvertTo(ITypeDescriptorContext? Context, Type? DestinationType)
+		public override bool CanConvertTo(ITypeDescriptorContext? context, Type? destinationType)
 		{
-			return DestinationType == typeof(string);
+			return destinationType == typeof(string);
 		}
 
 		/// <inheritdoc/>
-		public override object? ConvertTo(ITypeDescriptorContext? Context, CultureInfo? Culture, object? Value, Type DestinationType)
+		public override object? ConvertTo(ITypeDescriptorContext? context, CultureInfo? culture, object? value, Type destinationType)
 		{
-			if (DestinationType == typeof(string))
+			if (destinationType == typeof(string))
 			{
-				return Value?.ToString();
+				return value?.ToString();
 			}
 			else
 			{
diff --git a/Engine/Source/Programs/Horde/Horde.Build/Logs/Builder/ILogBuilder.cs b/Engine/Source/Programs/Horde/Horde.Build/Logs/Builder/ILogBuilder.cs
--- a/Engine/Source/Programs/Horde/Horde.Build/Logs/Builder/ILogBuilder.cs
+++ b/Engine/Source/Programs/Horde/Horde.Build/Logs/Builder/ILogBuilder.cs
@@ -1,13 +1,11 @@
 // Copyright Epic Games, Inc. All Rights Reserved.
 
-using Horde.Build.Api;
-using Horde.Build.Models;
-using Horde.Build.Utilities;
-using MongoDB.Bson;
 using System;
 using System.Collections.Generic;
-using System.Linq;
 using System.Threading.Tasks;
+using Horde.Build.Api;
+using Horde.Build.Models;
+using Horde.Build.Utilities;
 
 namespace Horde.Build.Logs
 {
@@ -26,54 +24,54 @@ namespace Horde.Build.Logs
 		/// <summary>
 		/// Append data to a key
 		/// </summary>
-		/// <param name="LogId">The log file id</param>
-		/// <param name="ChunkOffset">Offset of the chunk within the log file</param>
-		/// <param name="WriteOffset">Offset to write to</param>
-		/// <param name="WriteLineIndex">Line index being written</param>
-		/// <param name="WriteLineCount">Line count being written</param>
-		/// <param name="Data">Data to be appended</param>
-		/// <param name="Type">Type of data stored in this log file</param>
+		/// <param name="logId">The log file id</param>
+		/// <param name="chunkOffset">Offset of the chunk within the log file</param>
+		/// <param name="writeOffset">Offset to write to</param>
+		/// <param name="writeLineIndex">Line index being written</param>
+		/// <param name="writeLineCount">Line count being written</param>
+		/// <param name="data">Data to be appended</param>
+		/// <param name="type">Type of data stored in this log file</param>
 		/// <returns>True if the data was appended to the given chunk. False if the chunk has been completed.</returns>
-		Task<bool> AppendAsync(LogId LogId, long ChunkOffset, long WriteOffset, int WriteLineIndex, int WriteLineCount, ReadOnlyMemory<byte> Data, LogType Type);
+		Task<bool> AppendAsync(LogId logId, long chunkOffset, long writeOffset, int writeLineIndex, int writeLineCount, ReadOnlyMemory<byte> data, LogType type);
 
 		/// <summary>
 		/// Finish the current sub chunk
 		/// </summary>
-		/// <param name="LogId">The log file id</param>
-		/// <param name="Offset">Offset of the chunk within the log file</param>
+		/// <param name="logId">The log file id</param>
+		/// <param name="offset">Offset of the chunk within the log file</param>
 		/// <returns>Async task</returns>
-		Task CompleteSubChunkAsync(LogId LogId, long Offset);
+		Task CompleteSubChunkAsync(LogId logId, long offset);
 
 		/// <summary>
 		/// Finish the current chunk
 		/// </summary>
-		/// <param name="LogId">The log file id</param>
-		/// <param name="Offset">Offset of the chunk within the log file</param>
+		/// <param name="logId">The log file id</param>
+		/// <param name="offset">Offset of the chunk within the log file</param>
 		/// <returns>Async task</returns>
-		Task CompleteChunkAsync(LogId LogId, long Offset);
+		Task CompleteChunkAsync(LogId logId, long offset);
 
 		/// <summary>
 		/// Remove a complete chunk from the builder
 		/// </summary>
-		/// <param name="LogId">The log file id</param>
-		/// <param name="Offset">Offset of the chunk within the log file</param>
+		/// <param name="logId">The log file id</param>
+		/// <param name="offset">Offset of the chunk within the log file</param>
 		/// <returns>Async task</returns>
-		Task RemoveChunkAsync(LogId LogId, long Offset);
+		Task RemoveChunkAsync(LogId logId, long offset);
 
 		/// <summary>
 		/// Gets the current chunk for the given log file
 		/// </summary>
-		/// <param name="LogId">The log file id</param>
-		/// <param name="Offset">Offset of the chunk within the log file</param>
-		/// <param name="LineIndex">Line index of the chunk within the log file</param>
+		/// <param name="logId">The log file id</param>
+		/// <param name="offset">Offset of the chunk within the log file</param>
+		/// <param name="lineIndex">Line index of the chunk within the log file</param>
 		/// <returns></returns>
-		Task<LogChunkData?> GetChunkAsync(LogId LogId, long Offset, int LineIndex);
+		Task<LogChunkData?> GetChunkAsync(LogId logId, long offset, int lineIndex);
 
 		/// <summary>
 		/// Touches the timestamps of all the chunks after the given age, and returns them. Used for flushing the builder.
 		/// </summary>
-		/// <param name="MinAge">Minimum age of the chunks to enumerate. If specified, only chunks last modified longer than this period will be returned.</param>
+		/// <param name="minAge">Minimum age of the chunks to enumerate. If specified, only chunks last modified longer than this period will be returned.</param>
 		/// <returns>List of chunks, identified by log id and chunk index</returns>
-		Task<List<(LogId, long)>> TouchChunksAsync(TimeSpan MinAge);
+		Task<List<(LogId, long)>> TouchChunksAsync(TimeSpan minAge);
 	}
 }
diff --git a/Engine/Source/Programs/Horde/Horde.Build/Logs/Builder/Impl/LocalLogBuilder.cs b/Engine/Source/Programs/Horde/Horde.Build/Logs/Builder/Impl/LocalLogBuilder.cs
--- a/Engine/Source/Programs/Horde/Horde.Build/Logs/Builder/Impl/LocalLogBuilder.cs
+++ b/Engine/Source/Programs/Horde/Horde.Build/Logs/Builder/Impl/LocalLogBuilder.cs
@@ -1,15 +1,13 @@
 // Copyright Epic Games, Inc. All Rights Reserved.
 
-using Horde.Build.Api;
-using Horde.Build.Models;
-using Horde.Build.Utilities;
-using MongoDB.Bson;
 using System;
 using System.Collections.Concurrent;
 using System.Collections.Generic;
 using System.IO;
-using System.Linq;
 using System.Threading.Tasks;
+using Horde.Build.Api;
+using Horde.Build.Models;
+using Horde.Build.Utilities;
 
 namespace Horde.Build.Logs.Builder
 {
@@ -38,21 +36,21 @@ namespace Horde.Build.Logs.Builder
 			/// <summary>
 			/// Constructor
 			/// </summary>
-			/// <param name="Type">Type of data stored in this subchunk</param>
-			public PendingSubChunk(LogType Type)
+			/// <param name="type">Type of data stored in this subchunk</param>
+			public PendingSubChunk(LogType type)
 			{
-				this.Type = Type;
+				Type = type;
 			}
 
 			/// <summary>
 			/// Converts this pending sub-chunk to a concrete LogSubChunkData object
 			/// </summary>
-			/// <param name="Offset">Offset of this sub-chunk within the file</param>
-			/// <param name="LineIndex">The line index</param>
+			/// <param name="offset">Offset of this sub-chunk within the file</param>
+			/// <param name="lineIndex">The line index</param>
 			/// <returns>New sub-chunk data object</returns>
-			public LogSubChunkData ToSubChunkData(long Offset, int LineIndex)
+			public LogSubChunkData ToSubChunkData(long offset, int lineIndex)
 			{
-				return new LogSubChunkData(Type, Offset, LineIndex, new ReadOnlyLogText(Stream.ToArray()));
+				return new LogSubChunkData(Type, offset, lineIndex, new ReadOnlyLogText(Stream.ToArray()));
 			}
 		}
 
@@ -69,30 +67,30 @@ namespace Horde.Build.Logs.Builder
 			/// <summary>
 			/// Array of sub chunks for this chunk
 			/// </summary>
-			public List<PendingSubChunk> SubChunks = new List<PendingSubChunk>();
+			public List<PendingSubChunk> _subChunks = new List<PendingSubChunk>();
 
 			/// <summary>
 			/// The next sub-chunk
 			/// </summary>
-			public PendingSubChunk NextSubChunk;
+			public PendingSubChunk _nextSubChunk;
 
 			/// <summary>
 			/// Total length of this chunk
 			/// </summary>
-			public long Length;
+			public long _length;
 
 			/// <summary>
 			/// Whether the chunk is complete
 			/// </summary>
-			public bool Complete;
+			public bool _complete;
 
 			/// <summary>
 			/// Constructor
 			/// </summary>
-			/// <param name="Type">Type of data to store in this file</param>
-			public PendingChunk(LogType Type)
+			/// <param name="type">Type of data to store in this file</param>
+			public PendingChunk(LogType type)
 			{
-				NextSubChunk = new PendingSubChunk(Type);
+				_nextSubChunk = new PendingSubChunk(type);
 			}
 
 			/// <summary>
@@ -100,10 +98,10 @@ namespace Horde.Build.Logs.Builder
 			/// </summary>
 			public void CompleteSubChunk()
 			{
-				if (NextSubChunk.Stream.Length > 0)
+				if (_nextSubChunk.Stream.Length > 0)
 				{
-					SubChunks.Add(NextSubChunk);
-					NextSubChunk = new PendingSubChunk(NextSubChunk.Type);
+					_subChunks.Add(_nextSubChunk);
+					_nextSubChunk = new PendingSubChunk(_nextSubChunk.Type);
 				}
 			}
 		}
@@ -111,36 +109,36 @@ namespace Horde.Build.Logs.Builder
 		/// <summary>
 		/// Current log chunk state
 		/// </summary>
-		ConcurrentDictionary<(LogId, long), PendingChunk> PendingChunks = new ConcurrentDictionary<(LogId, long), PendingChunk>();
+		readonly ConcurrentDictionary<(LogId, long), PendingChunk> _pendingChunks = new ConcurrentDictionary<(LogId, long), PendingChunk>();
 
 		/// <inheritdoc/>
 		public bool FlushOnShutdown => true;
 
 		/// <inheritdoc/>
-		public Task<bool> AppendAsync(LogId LogId, long ChunkOffset, long WriteOffset, int WriteLineIndex, int WriteLineCount, ReadOnlyMemory<byte> Data, LogType Type)
+		public Task<bool> AppendAsync(LogId logId, long chunkOffset, long writeOffset, int writeLineIndex, int writeLineCount, ReadOnlyMemory<byte> data, LogType type)
 		{
-			PendingChunk? PendingChunk;
-			while (!PendingChunks.TryGetValue((LogId, ChunkOffset), out PendingChunk))
+			PendingChunk? pendingChunk;
+			while (!_pendingChunks.TryGetValue((logId, chunkOffset), out pendingChunk))
 			{
-				if(WriteOffset != ChunkOffset)
+				if(writeOffset != chunkOffset)
 				{
 					return Task.FromResult(false);
 				}
 
-				PendingChunk = new PendingChunk(Type);
+				pendingChunk = new PendingChunk(type);
 
-				if (PendingChunks.TryAdd((LogId, ChunkOffset), PendingChunk))
+				if (_pendingChunks.TryAdd((logId, chunkOffset), pendingChunk))
 				{
 					break;
 				}
 			}
 
-			lock (PendingChunk)
+			lock (pendingChunk)
 			{
-				if (!PendingChunk.Complete && ChunkOffset + PendingChunk.Length == WriteOffset)
+				if (!pendingChunk._complete && chunkOffset + pendingChunk._length == writeOffset)
 				{
-					PendingChunk.NextSubChunk.Stream.Write(Data.Span);
-					PendingChunk.Length += Data.Span.Length;
+					pendingChunk._nextSubChunk.Stream.Write(data.Span);
+					pendingChunk._length += data.Span.Length;
 					return Task.FromResult(true);
 				}
 			}
@@ -148,31 +146,31 @@ namespace Horde.Build.Logs.Builder
 		}
 
 		/// <inheritdoc/>
-		public Task CompleteSubChunkAsync(LogId LogId, long Offset)
+		public Task CompleteSubChunkAsync(LogId logId, long offset)
 		{
-			PendingChunk? PendingChunk;
-			if (PendingChunks.TryGetValue((LogId, Offset), out PendingChunk))
+			PendingChunk? pendingChunk;
+			if (_pendingChunks.TryGetValue((logId, offset), out pendingChunk))
 			{
-				lock (PendingChunk)
+				lock (pendingChunk)
 				{
-					PendingChunk.CompleteSubChunk();
+					pendingChunk.CompleteSubChunk();
 				}
 			}
 			return Task.CompletedTask;
 		}
 
 		/// <inheritdoc/>
-		public Task CompleteChunkAsync(LogId LogId, long Offset)
+		public Task CompleteChunkAsync(LogId logId, long offset)
 		{
-			PendingChunk? PendingChunk;
-			if (PendingChunks.TryGetValue((LogId, Offset), out PendingChunk))
+			PendingChunk? pendingChunk;
+			if (_pendingChunks.TryGetValue((logId, offset), out pendingChunk))
 			{
-				lock (PendingChunk)
+				lock (pendingChunk)
 				{
-					if (!PendingChunk.Complete)
+					if (!pendingChunk._complete)
 					{
-						PendingChunk.Complete = true;
-						PendingChunk.CompleteSubChunk();
+						pendingChunk._complete = true;
+						pendingChunk.CompleteSubChunk();
 					}
 				}
 			}
@@ -180,61 +178,61 @@ namespace Horde.Build.Logs.Builder
 		}
 
 		/// <inheritdoc/>
-		public Task RemoveChunkAsync(LogId LogId, long Offset)
+		public Task RemoveChunkAsync(LogId logId, long offset)
 		{
-			PendingChunks.TryRemove((LogId, Offset), out _);
+			_pendingChunks.TryRemove((logId, offset), out _);
 			return Task.CompletedTask;
 		}
 
 		/// <inheritdoc/>
-		public Task<LogChunkData?> GetChunkAsync(LogId LogId, long Offset, int LineIndex)
+		public Task<LogChunkData?> GetChunkAsync(LogId logId, long offset, int lineIndex)
 		{
-			PendingChunk? PendingChunk;
-			if (PendingChunks.TryGetValue((LogId, Offset), out PendingChunk))
+			PendingChunk? pendingChunk;
+			if (_pendingChunks.TryGetValue((logId, offset), out pendingChunk))
 			{
-				long SubChunkOffset = Offset;
-				int SubChunkLineIndex = LineIndex;
+				long subChunkOffset = offset;
+				int subChunkLineIndex = lineIndex;
 
-				List<LogSubChunkData> SubChunks = new List<LogSubChunkData>();
-				foreach(PendingSubChunk PendingSubChunk in PendingChunk.SubChunks)
+				List<LogSubChunkData> subChunks = new List<LogSubChunkData>();
+				foreach(PendingSubChunk pendingSubChunk in pendingChunk._subChunks)
 				{
-					LogSubChunkData SubChunk = PendingSubChunk.ToSubChunkData(SubChunkOffset, SubChunkLineIndex);
-					SubChunkOffset += SubChunk.Length;
-					SubChunkLineIndex += SubChunk.LineCount;
-					SubChunks.Add(SubChunk);
+					LogSubChunkData subChunk = pendingSubChunk.ToSubChunkData(subChunkOffset, subChunkLineIndex);
+					subChunkOffset += subChunk.Length;
+					subChunkLineIndex += subChunk.LineCount;
+					subChunks.Add(subChunk);
 				}
 
-				PendingSubChunk NextSubChunk = PendingChunk.NextSubChunk;
-				if (NextSubChunk.Stream.Length > 0)
+				PendingSubChunk nextSubChunk = pendingChunk._nextSubChunk;
+				if (nextSubChunk.Stream.Length > 0)
 				{
-					LogSubChunkData SubChunkData = NextSubChunk.ToSubChunkData(SubChunkOffset, SubChunkLineIndex);
-					SubChunks.Add(SubChunkData);
+					LogSubChunkData subChunkData = nextSubChunk.ToSubChunkData(subChunkOffset, subChunkLineIndex);
+					subChunks.Add(subChunkData);
 				}
 
-				return Task.FromResult<LogChunkData?>(new LogChunkData(Offset, LineIndex, SubChunks));
+				return Task.FromResult<LogChunkData?>(new LogChunkData(offset, lineIndex, subChunks));
 			}
 
 			return Task.FromResult<LogChunkData?>(null);
 		}
 
 		/// <inheritdoc/>
-		public Task<List<(LogId, long)>> TouchChunksAsync(TimeSpan MinAge)
+		public Task<List<(LogId, long)>> TouchChunksAsync(TimeSpan minAge)
 		{
-			DateTime UtcNow = DateTime.UtcNow;
+			DateTime utcNow = DateTime.UtcNow;
 
-			List<(LogId, long)> Chunks = new List<(LogId, long)>();
-			foreach (KeyValuePair<(LogId, long), PendingChunk> PendingChunk in PendingChunks.ToArray())
+			List<(LogId, long)> chunks = new List<(LogId, long)>();
+			foreach (KeyValuePair<(LogId, long), PendingChunk> pendingChunk in _pendingChunks.ToArray())
 			{
-				lock (PendingChunk.Value)
+				lock (pendingChunk.Value)
 				{
-					if (PendingChunk.Value.CreateTimeUtc < UtcNow - MinAge)
+					if (pendingChunk.Value.CreateTimeUtc < utcNow - minAge)
 					{
-						Chunks.Add(PendingChunk.Key);
-						PendingChunk.Value.CreateTimeUtc = UtcNow;
+						chunks.Add(pendingChunk.Key);
+						pendingChunk.Value.CreateTimeUtc = utcNow;
 					}
 				}
 			}
-			return Task.FromResult(Chunks);
+			return Task.FromResult(chunks);
 		}
 	}
 }
diff --git a/Engine/Source/Programs/Horde/Horde.Build/Logs/Builder/Impl/RedisLogBuilder.cs b/Engine/Source/Programs/Horde/Horde.Build/Logs/Builder/Impl/RedisLogBuilder.cs
--- a/Engine/Source/Programs/Horde/Horde.Build/Logs/Builder/Impl/RedisLogBuilder.cs
+++ b/Engine/Source/Programs/Horde/Horde.Build/Logs/Builder/Impl/RedisLogBuilder.cs
@@ -1,26 +1,23 @@
 // Copyright Epic Games, Inc. All Rights Reserved.
 
-using EpicGames.Core;
-using Horde.Build.Api;
-using Horde.Build.Utilities;
-using Microsoft.Extensions.Logging;
-using Microsoft.Extensions.Primitives;
-using MongoDB.Bson;
-using StackExchange.Redis;
 using System;
 using System.Collections.Generic;
 using System.Globalization;
-using System.Linq;
 using System.Text.RegularExpressions;
 using System.Threading.Tasks;
+using EpicGames.Core;
+using Horde.Build.Api;
+using Horde.Build.Models;
+using Horde.Build.Utilities;
+using Microsoft.Extensions.Logging;
 using OpenTracing;
 using OpenTracing.Util;
-using Horde.Build.Models;
+using StackExchange.Redis;
 
 namespace Horde.Build.Logs.Builder
 {
-	using LogId = ObjectId<ILogFile>;
 	using Condition = StackExchange.Redis.Condition;
+	using LogId = ObjectId<ILogFile>;
 
 	/// <summary>
 	/// Redis-based cache for log file chunks
@@ -39,18 +36,18 @@ namespace Horde.Build.Logs.Builder
 			public string SubChunkData => $"{Prefix}-subchunk";
 			public string Complete => $"{Prefix}-complete";
 
-			public ChunkKeys(LogId LogId, long Offset)
+			public ChunkKeys(LogId logId, long offset)
 			{
-				Prefix = $"log-{LogId}-chunk-{Offset}-builder";
+				Prefix = $"log-{logId}-chunk-{offset}-builder";
 			}
 
-			public static bool TryParse(string Prefix, out LogId LogId, out long Offset)
+			public static bool TryParse(string prefix, out LogId logId, out long offset)
 			{
-				Match Match = Regex.Match(Prefix, "log-([^-]+)-chunk-([^-]+)-builder");
-				if (!Match.Success || !LogId.TryParse(Match.Groups[1].Value, out LogId) || !long.TryParse(Match.Groups[2].Value, out Offset))
+				Match match = Regex.Match(prefix, "log-([^-]+)-chunk-([^-]+)-builder");
+				if (!match.Success || !LogId.TryParse(match.Groups[1].Value, out logId) || !Int64.TryParse(match.Groups[2].Value, out offset))
 				{
-					LogId = LogId.Empty;
-					Offset = 0;
+					logId = LogId.Empty;
+					offset = 0;
 					return false;
 				}
 				return true;
@@ -63,68 +60,68 @@ namespace Horde.Build.Logs.Builder
 		/// <summary>
 		/// The Redis database connection pool
 		/// </summary>
-		RedisConnectionPool RedisConnectionPool;
+		readonly RedisConnectionPool _redisConnectionPool;
 
 		/// <summary>
 		/// Logger for debug output
 		/// </summary>
-		ILogger Logger;
+		readonly ILogger _logger;
 
 		/// <summary>
 		/// Constructor
 		/// </summary>
-		/// <param name="RedisConnectionPool">The redis database singleton</param>
-		/// <param name="Logger">Logger for debug output</param>
-		public RedisLogBuilder(RedisConnectionPool RedisConnectionPool, ILogger Logger)
+		/// <param name="redisConnectionPool">The redis database singleton</param>
+		/// <param name="logger">Logger for debug output</param>
+		public RedisLogBuilder(RedisConnectionPool redisConnectionPool, ILogger logger)
 		{
-			this.RedisConnectionPool = RedisConnectionPool;
-			this.Logger = Logger;
+			_redisConnectionPool = redisConnectionPool;
+			_logger = logger;
 		}
 
 		/// <inheritdoc/>
-		public async Task<bool> AppendAsync(LogId LogId, long ChunkOffset, long WriteOffset, int WriteLineIndex, int WriteLineCount, ReadOnlyMemory<byte> Data, LogType Type)
+		public async Task<bool> AppendAsync(LogId logId, long chunkOffset, long writeOffset, int writeLineIndex, int writeLineCount, ReadOnlyMemory<byte> data, LogType type)
 		{
-			IDatabase RedisDb = RedisConnectionPool.GetDatabase();
-			ChunkKeys Keys = new ChunkKeys(LogId, ChunkOffset);
+			IDatabase redisDb = _redisConnectionPool.GetDatabase();
+			ChunkKeys keys = new ChunkKeys(logId, chunkOffset);
 
-			if(ChunkOffset == WriteOffset)
+			if(chunkOffset == writeOffset)
 			{
-				using IScope Scope = GlobalTracer.Instance.BuildSpan("Redis.CreateChunk").StartActive();
-				Scope.Span.SetTag("LogId", LogId.ToString());
-				Scope.Span.SetTag("Offset", ChunkOffset.ToString(CultureInfo.InvariantCulture));
-				Scope.Span.SetTag("WriteOffset", WriteOffset.ToString(CultureInfo.InvariantCulture));
-
-				ITransaction? CreateTransaction = RedisDb.CreateTransaction();
-				CreateTransaction.AddCondition(Condition.SortedSetNotContains(ItemsKey, Keys.Prefix));
-				_ = CreateTransaction.SortedSetAddAsync(ItemsKey, Keys.Prefix, DateTime.UtcNow.Ticks);
-				_ = CreateTransaction.StringSetAsync(Keys.Type, (int)Type);
-				_ = CreateTransaction.StringSetAsync(Keys.Length, Data.Length);
-				_ = CreateTransaction.StringSetAsync(Keys.LineIndex, WriteLineIndex + WriteLineCount);
-				_ = CreateTransaction.StringSetAsync(Keys.SubChunkData, Data);
-				if (await CreateTransaction.ExecuteAsync())
+				using IScope scope = GlobalTracer.Instance.BuildSpan("Redis.CreateChunk").StartActive();
+				scope.Span.SetTag("LogId", logId.ToString());
+				scope.Span.SetTag("Offset", chunkOffset.ToString(CultureInfo.InvariantCulture));
+				scope.Span.SetTag("WriteOffset", writeOffset.ToString(CultureInfo.InvariantCulture));
+
+				ITransaction? createTransaction = redisDb.CreateTransaction();
+				createTransaction.AddCondition(Condition.SortedSetNotContains(ItemsKey, keys.Prefix));
+				_ = createTransaction.SortedSetAddAsync(ItemsKey, keys.Prefix, DateTime.UtcNow.Ticks);
+				_ = createTransaction.StringSetAsync(keys.Type, (int)type);
+				_ = createTransaction.StringSetAsync(keys.Length, data.Length);
+				_ = createTransaction.StringSetAsync(keys.LineIndex, writeLineIndex + writeLineCount);
+				_ = createTransaction.StringSetAsync(keys.SubChunkData, data);
+				if (await createTransaction.ExecuteAsync())
 				{
-					Logger.LogTrace("Created {Size} bytes in {Key}", Data.Length, Keys.SubChunkData);
+					_logger.LogTrace("Created {Size} bytes in {Key}", data.Length, keys.SubChunkData);
 					return true;
 				}
 			}
 
 			{
-				using IScope Scope = GlobalTracer.Instance.BuildSpan("Redis.AppendChunk").StartActive();
-				Scope.Span.SetTag("LogId", LogId.ToString());
-				Scope.Span.SetTag("Offset", ChunkOffset.ToString(CultureInfo.InvariantCulture));
-				Scope.Span.SetTag("WriteOffset", WriteOffset.ToString(CultureInfo.InvariantCulture));
-
-				ITransaction? AppendTransaction = RedisDb.CreateTransaction();
-				AppendTransaction.AddCondition(Condition.SortedSetContains(ItemsKey, Keys.Prefix));
-				AppendTransaction.AddCondition(Condition.KeyNotExists(Keys.Complete));
-				AppendTransaction.AddCondition(Condition.StringEqual(Keys.Type, (int)Type));
-				AppendTransaction.AddCondition(Condition.StringEqual(Keys.Length, (int)(WriteOffset - ChunkOffset)));
-				_ = AppendTransaction.StringAppendAsync(Keys.SubChunkData, Data);
-				_ = AppendTransaction.StringSetAsync(Keys.Length, (int)(WriteOffset - ChunkOffset) + Data.Length);
-				_ = AppendTransaction.StringSetAsync(Keys.LineIndex, WriteLineIndex + WriteLineCount);
-				if (await AppendTransaction.ExecuteAsync())
+				using IScope scope = GlobalTracer.Instance.BuildSpan("Redis.AppendChunk").StartActive();
+				scope.Span.SetTag("LogId", logId.ToString());
+				scope.Span.SetTag("Offset", chunkOffset.ToString(CultureInfo.InvariantCulture));
+				scope.Span.SetTag("WriteOffset", writeOffset.ToString(CultureInfo.InvariantCulture));
+
+				ITransaction? appendTransaction = redisDb.CreateTransaction();
+				appendTransaction.AddCondition(Condition.SortedSetContains(ItemsKey, keys.Prefix));
+				appendTransaction.AddCondition(Condition.KeyNotExists(keys.Complete));
+				appendTransaction.AddCondition(Condition.StringEqual(keys.Type, (int)type));
+				appendTransaction.AddCondition(Condition.StringEqual(keys.Length, (int)(writeOffset - chunkOffset)));
+				_ = appendTransaction.StringAppendAsync(keys.SubChunkData, data);
+				_ = appendTransaction.StringSetAsync(keys.Length, (int)(writeOffset - chunkOffset) + data.Length);
+				_ = appendTransaction.StringSetAsync(keys.LineIndex, writeLineIndex + writeLineCount);
+				if (await appendTransaction.ExecuteAsync())
 				{
-					Logger.LogTrace("Appended {Size} to {Key}", Data.Length, Keys.SubChunkData);
+					_logger.LogTrace("Appended {Size} to {Key}", data.Length, keys.SubChunkData);
 					return true;
 				}
 			}
@@ -133,176 +130,176 @@ namespace Horde.Build.Logs.Builder
 		}
 
 		/// <inheritdoc/>
-		public async Task CompleteSubChunkAsync(LogId LogId, long Offset)
+		public async Task CompleteSubChunkAsync(LogId logId, long offset)
 		{
-			IDatabase RedisDb = RedisConnectionPool.GetDatabase();
-			ChunkKeys Keys = new ChunkKeys(LogId, Offset);
+			IDatabase redisDb = _redisConnectionPool.GetDatabase();
+			ChunkKeys keys = new ChunkKeys(logId, offset);
 			for (; ; )
 			{
-				using IScope Scope = GlobalTracer.Instance.BuildSpan("Redis.CompleteSubChunk").StartActive();
-				Scope.Span.SetTag("LogId", LogId.ToString());
-				Scope.Span.SetTag("Offset", Offset.ToString(CultureInfo.InvariantCulture));
+				using IScope scope = GlobalTracer.Instance.BuildSpan("Redis.CompleteSubChunk").StartActive();
+				scope.Span.SetTag("LogId", logId.ToString());
+				scope.Span.SetTag("Offset", offset.ToString(CultureInfo.InvariantCulture));
 
-				LogType Type = (LogType)(int)await RedisDb.StringGetAsync(Keys.Type);
-				int Length = (int)await RedisDb.StringGetAsync(Keys.Length);
-				int LineIndex = (int)await RedisDb.StringGetAsync(Keys.LineIndex);
-				long ChunkDataLength = await RedisDb.StringLengthAsync(Keys.ChunkData);
+				LogType type = (LogType)(int)await redisDb.StringGetAsync(keys.Type);
+				int length = (int)await redisDb.StringGetAsync(keys.Length);
+				int lineIndex = (int)await redisDb.StringGetAsync(keys.LineIndex);
+				long chunkDataLength = await redisDb.StringLengthAsync(keys.ChunkData);
 
-				RedisValue SubChunkTextValue = await RedisDb.StringGetAsync(Keys.SubChunkData);
-				if (SubChunkTextValue.IsNullOrEmpty)
+				RedisValue subChunkTextValue = await redisDb.StringGetAsync(keys.SubChunkData);
+				if (subChunkTextValue.IsNullOrEmpty)
 				{
 					break;
 				}
 
-				ReadOnlyLogText SubChunkText = new ReadOnlyLogText(SubChunkTextValue);
-				LogSubChunkData SubChunkData = new LogSubChunkData(Type, Offset + Length - SubChunkText.Data.Length, LineIndex - SubChunkText.LineCount, SubChunkText);
-				byte[] SubChunkDataBytes = SubChunkData.ToByteArray();
+				ReadOnlyLogText subChunkText = new ReadOnlyLogText(subChunkTextValue);
+				LogSubChunkData subChunkData = new LogSubChunkData(type, offset + length - subChunkText.Data.Length, lineIndex - subChunkText.LineCount, subChunkText);
+				byte[] subChunkDataBytes = subChunkData.ToByteArray();
 
-				ITransaction WriteTransaction = RedisDb.CreateTransaction();
+				ITransaction writeTransaction = redisDb.CreateTransaction();
 
-				WriteTransaction.AddCondition(Condition.StringLengthEqual(Keys.ChunkData, ChunkDataLength));
-				WriteTransaction.AddCondition(Condition.StringLengthEqual(Keys.SubChunkData, SubChunkText.Length));
-				WriteTransaction.AddCondition(Condition.StringEqual(Keys.Length, Length));
-				WriteTransaction.AddCondition(Condition.StringEqual(Keys.LineIndex, LineIndex));
-				Task<long> NewLength = WriteTransaction.StringAppendAsync(Keys.ChunkData, SubChunkDataBytes);
-				_ = WriteTransaction.KeyDeleteAsync(Keys.SubChunkData);
+				writeTransaction.AddCondition(Condition.StringLengthEqual(keys.ChunkData, chunkDataLength));
+				writeTransaction.AddCondition(Condition.StringLengthEqual(keys.SubChunkData, subChunkText.Length));
+				writeTransaction.AddCondition(Condition.StringEqual(keys.Length, length));
+				writeTransaction.AddCondition(Condition.StringEqual(keys.LineIndex, lineIndex));
+				Task<long> newLength = writeTransaction.StringAppendAsync(keys.ChunkData, subChunkDataBytes);
+				_ = writeTransaction.KeyDeleteAsync(keys.SubChunkData);
 
-				if (await WriteTransaction.ExecuteAsync())
+				if (await writeTransaction.ExecuteAsync())
 				{
-					Logger.LogDebug("Completed sub-chunk for log {LogId} chunk offset {Offset} -> sub-chunk size {SubChunkSize}, chunk size {ChunkSize}", LogId, Offset, SubChunkDataBytes.Length, await NewLength);
+					_logger.LogDebug("Completed sub-chunk for log {LogId} chunk offset {Offset} -> sub-chunk size {SubChunkSize}, chunk size {ChunkSize}", logId, offset, subChunkDataBytes.Length, await newLength);
 					break;
 				}
 			}
 		}
 
 		/// <inheritdoc/>
-		public async Task CompleteChunkAsync(LogId LogId, long Offset)
+		public async Task CompleteChunkAsync(LogId logId, long offset)
 		{
-			IDatabase RedisDb = RedisConnectionPool.GetDatabase();
-			using IScope Scope = GlobalTracer.Instance.BuildSpan("Redis.CompleteChunk").StartActive();
-			Scope.Span.SetTag("LogId", LogId.ToString());
-			Scope.Span.SetTag("Offset", Offset.ToString(CultureInfo.InvariantCulture));
+			IDatabase redisDb = _redisConnectionPool.GetDatabase();
+			using IScope scope = GlobalTracer.Instance.BuildSpan("Redis.CompleteChunk").StartActive();
+			scope.Span.SetTag("LogId", logId.ToString());
+			scope.Span.SetTag("Offset", offset.ToString(CultureInfo.InvariantCulture));
 
-			ChunkKeys Keys = new ChunkKeys(LogId, Offset);
+			ChunkKeys keys = new ChunkKeys(logId, offset);
 
-			ITransaction Transaction = RedisDb.CreateTransaction();
-			Transaction.AddCondition(Condition.SortedSetContains(ItemsKey, Keys.Prefix));
-			_ = Transaction.StringSetAsync(Keys.Complete, true);
-			if(!await Transaction.ExecuteAsync())
+			ITransaction transaction = redisDb.CreateTransaction();
+			transaction.AddCondition(Condition.SortedSetContains(ItemsKey, keys.Prefix));
+			_ = transaction.StringSetAsync(keys.Complete, true);
+			if(!await transaction.ExecuteAsync())
 			{
-				Logger.LogDebug("Log {LogId} chunk offset {Offset} is not in Redis builder", LogId, Offset);
+				_logger.LogDebug("Log {LogId} chunk offset {Offset} is not in Redis builder", logId, offset);
 				return;
 			}
 
-			await CompleteSubChunkAsync(LogId, Offset);
+			await CompleteSubChunkAsync(logId, offset);
 		}
 
 		/// <inheritdoc/>
-		public async Task RemoveChunkAsync(LogId LogId, long Offset)
+		public async Task RemoveChunkAsync(LogId logId, long offset)
 		{
-			IDatabase RedisDb = RedisConnectionPool.GetDatabase();
-			using IScope Scope = GlobalTracer.Instance.BuildSpan("Redis.RemoveChunk").StartActive();
-			Scope.Span.SetTag("LogId", LogId.ToString());
-			Scope.Span.SetTag("Offset", Offset.ToString(CultureInfo.InvariantCulture));
-
-			ChunkKeys Keys = new ChunkKeys(LogId, Offset);
-
-			ITransaction Transaction = RedisDb.CreateTransaction();
-			_ = Transaction.KeyDeleteAsync(Keys.Type);
-			_ = Transaction.KeyDeleteAsync(Keys.LineIndex);
-			_ = Transaction.KeyDeleteAsync(Keys.Length);
-			_ = Transaction.KeyDeleteAsync(Keys.ChunkData);
-			_ = Transaction.KeyDeleteAsync(Keys.SubChunkData);
-			_ = Transaction.KeyDeleteAsync(Keys.Complete);
-
-			_ = Transaction.SortedSetRemoveAsync(ItemsKey, Keys.Prefix);
-			await Transaction.ExecuteAsync();
+			IDatabase redisDb = _redisConnectionPool.GetDatabase();
+			using IScope scope = GlobalTracer.Instance.BuildSpan("Redis.RemoveChunk").StartActive();
+			scope.Span.SetTag("LogId", logId.ToString());
+			scope.Span.SetTag("Offset", offset.ToString(CultureInfo.InvariantCulture));
+
+			ChunkKeys keys = new ChunkKeys(logId, offset);
+
+			ITransaction transaction = redisDb.CreateTransaction();
+			_ = transaction.KeyDeleteAsync(keys.Type);
+			_ = transaction.KeyDeleteAsync(keys.LineIndex);
+			_ = transaction.KeyDeleteAsync(keys.Length);
+			_ = transaction.KeyDeleteAsync(keys.ChunkData);
+			_ = transaction.KeyDeleteAsync(keys.SubChunkData);
+			_ = transaction.KeyDeleteAsync(keys.Complete);
+
+			_ = transaction.SortedSetRemoveAsync(ItemsKey, keys.Prefix);
+			await transaction.ExecuteAsync();
 		}
 
 		/// <inheritdoc/>
-		public async Task<LogChunkData?> GetChunkAsync(LogId LogId, long Offset, int LineIndex)
+		public async Task<LogChunkData?> GetChunkAsync(LogId logId, long offset, int lineIndex)
 		{
-			IDatabase RedisDb = RedisConnectionPool.GetDatabase();
-			ChunkKeys Keys = new ChunkKeys(LogId, Offset);
+			IDatabase redisDb = _redisConnectionPool.GetDatabase();
+			ChunkKeys keys = new ChunkKeys(logId, offset);
 			for (; ; )
 			{
-				using IScope Scope = GlobalTracer.Instance.BuildSpan("Redis.GetChunk").StartActive();
-				Scope.Span.SetTag("LogId", LogId.ToString());
-				Scope.Span.SetTag("Offset", Offset.ToString(CultureInfo.InvariantCulture));
+				using IScope scope = GlobalTracer.Instance.BuildSpan("Redis.GetChunk").StartActive();
+				scope.Span.SetTag("LogId", logId.ToString());
+				scope.Span.SetTag("Offset", offset.ToString(CultureInfo.InvariantCulture));
 
-				RedisValue ChunkDataValue = await RedisDb.StringGetAsync(Keys.ChunkData);
+				RedisValue chunkDataValue = await redisDb.StringGetAsync(keys.ChunkData);
 
-				ReadOnlyMemory<byte> ChunkData = ChunkDataValue;
+				ReadOnlyMemory<byte> chunkData = chunkDataValue;
 
-				ITransaction Transaction = RedisDb.CreateTransaction();
-				Transaction.AddCondition(Condition.StringLengthEqual(Keys.ChunkData, ChunkData.Length));
-				Task<RedisValue> TypeTask = Transaction.StringGetAsync(Keys.Type);
-				Task<RedisValue> LastSubChunkDataTask = Transaction.StringGetAsync(Keys.SubChunkData);
+				ITransaction transaction = redisDb.CreateTransaction();
+				transaction.AddCondition(Condition.StringLengthEqual(keys.ChunkData, chunkData.Length));
+				Task<RedisValue> typeTask = transaction.StringGetAsync(keys.Type);
+				Task<RedisValue> lastSubChunkDataTask = transaction.StringGetAsync(keys.SubChunkData);
 
-				if (await Transaction.ExecuteAsync())
+				if (await transaction.ExecuteAsync())
 				{
-					MemoryReader Reader = new MemoryReader(ChunkData);
+					MemoryReader reader = new MemoryReader(chunkData);
 
-					long SubChunkOffset = Offset;
-					int SubChunkLineIndex = LineIndex;
+					long subChunkOffset = offset;
+					int subChunkLineIndex = lineIndex;
 
-					List<LogSubChunkData> SubChunks = new List<LogSubChunkData>();
-					while (Reader.Offset < Reader.Length)
+					List<LogSubChunkData> subChunks = new List<LogSubChunkData>();
+					while (reader.Offset < reader.Length)
 					{
-						LogSubChunkData SubChunkData = Reader.ReadLogSubChunkData(SubChunkOffset, SubChunkLineIndex);
-						SubChunkOffset += SubChunkData.Length;
-						SubChunkLineIndex += SubChunkData.LineCount;
-						SubChunks.Add(SubChunkData);
+						LogSubChunkData subChunkData = reader.ReadLogSubChunkData(subChunkOffset, subChunkLineIndex);
+						subChunkOffset += subChunkData.Length;
+						subChunkLineIndex += subChunkData.LineCount;
+						subChunks.Add(subChunkData);
 					}
 
-					RedisValue Type = await TypeTask;
+					RedisValue type = await typeTask;
 
-					RedisValue LastSubChunkData = await LastSubChunkDataTask;
-					if (LastSubChunkData.Length() > 0)
+					RedisValue lastSubChunkData = await lastSubChunkDataTask;
+					if (lastSubChunkData.Length() > 0)
 					{
-						ReadOnlyLogText SubChunkDataText = new ReadOnlyLogText(LastSubChunkData);
-						SubChunks.Add(new LogSubChunkData((LogType)(int)Type, SubChunkOffset, SubChunkLineIndex, SubChunkDataText));
+						ReadOnlyLogText subChunkDataText = new ReadOnlyLogText(lastSubChunkData);
+						subChunks.Add(new LogSubChunkData((LogType)(int)type, subChunkOffset, subChunkLineIndex, subChunkDataText));
 					}
 
-					if (SubChunks.Count == 0)
+					if (subChunks.Count == 0)
 					{
 						break;
 					}
 
-					return new LogChunkData(Offset, LineIndex, SubChunks);
+					return new LogChunkData(offset, lineIndex, subChunks);
 				}
 			}
 			return null;
 		}
 
 		/// <inheritdoc/>
-		public async Task<List<(LogId, long)>> TouchChunksAsync(TimeSpan MinAge)
+		public async Task<List<(LogId, long)>> TouchChunksAsync(TimeSpan minAge)
 		{
-			IDatabase RedisDb = RedisConnectionPool.GetDatabase();
+			IDatabase redisDb = _redisConnectionPool.GetDatabase();
 			
 			// Find all the chunks that are suitable for expiry
-			DateTime UtcNow = DateTime.UtcNow;
-			SortedSetEntry[] Entries = await RedisDb.SortedSetRangeByScoreWithScoresAsync(ItemsKey, stop: (UtcNow - MinAge).Ticks);
+			DateTime utcNow = DateTime.UtcNow;
+			SortedSetEntry[] entries = await redisDb.SortedSetRangeByScoreWithScoresAsync(ItemsKey, stop: (utcNow - minAge).Ticks);
 
 			// Update the score for each element in a transaction. If it succeeds, we can write the chunk. Otherwise another pod has beat us to it.
-			List<(LogId, long)> Results = new List<(LogId, long)>();
-			foreach (SortedSetEntry Entry in Entries)
+			List<(LogId, long)> results = new List<(LogId, long)>();
+			foreach (SortedSetEntry entry in entries)
 			{
-				if (ChunkKeys.TryParse(Entry.Element.ToString(), out LogId LogId, out long Offset))
+				if (ChunkKeys.TryParse(entry.Element.ToString(), out LogId logId, out long offset))
 				{
-					ITransaction Transaction = RedisDb.CreateTransaction();
-					Transaction.AddCondition(Condition.SortedSetEqual(ItemsKey, Entry.Element, Entry.Score));
-					_ = Transaction.SortedSetAddAsync(ItemsKey, Entry.Element, UtcNow.Ticks);
+					ITransaction transaction = redisDb.CreateTransaction();
+					transaction.AddCondition(Condition.SortedSetEqual(ItemsKey, entry.Element, entry.Score));
+					_ = transaction.SortedSetAddAsync(ItemsKey, entry.Element, utcNow.Ticks);
 
-					if (!await Transaction.ExecuteAsync())
+					if (!await transaction.ExecuteAsync())
 					{
 						break;
 					}
 
-					Results.Add((LogId, Offset));
+					results.Add((logId, offset));
 				}
 			}
-			return Results;
+			return results;
 		}
 	}
 }
diff --git a/Engine/Source/Programs/Horde/Horde.Build/Logs/Data/LogChunkData.cs b/Engine/Source/Programs/Horde/Horde.Build/Logs/Data/LogChunkData.cs
--- a/Engine/Source/Programs/Horde/Horde.Build/Logs/Data/LogChunkData.cs
+++ b/Engine/Source/Programs/Horde/Horde.Build/Logs/Data/LogChunkData.cs
@@ -1,14 +1,10 @@
 // Copyright Epic Games, Inc. All Rights Reserved.
 
-using EpicGames.Core;
-using Horde.Build.Api;
-using Horde.Build.Utilities;
 using System;
 using System.Collections.Generic;
-using System.IO;
 using System.Linq;
-using System.Text;
-using System.Threading.Tasks;
+using EpicGames.Core;
+using Horde.Build.Api;
 
 namespace Horde.Build.Logs
 {
@@ -55,88 +51,88 @@ namespace Horde.Build.Logs
 		/// <summary>
 		/// Constructor
 		/// </summary>
-		/// <param name="Offset">Offset of this chunk</param>
-		/// <param name="LineIndex">First line index of this chunk</param>
-		/// <param name="SubChunks">Sub-chunks for this chunk</param>
-		public LogChunkData(long Offset, int LineIndex, IReadOnlyList<LogSubChunkData> SubChunks)
+		/// <param name="offset">Offset of this chunk</param>
+		/// <param name="lineIndex">First line index of this chunk</param>
+		/// <param name="subChunks">Sub-chunks for this chunk</param>
+		public LogChunkData(long offset, int lineIndex, IReadOnlyList<LogSubChunkData> subChunks)
 		{
-			this.Offset = Offset;
-			this.LineIndex = LineIndex;
-			this.SubChunks = SubChunks;
-			this.SubChunkOffset = CreateSumLookup(SubChunks, x => x.Length);
-			this.SubChunkLineIndex = CreateSumLookup(SubChunks, x => x.LineCount);
+			Offset = offset;
+			LineIndex = lineIndex;
+			SubChunks = subChunks;
+			SubChunkOffset = CreateSumLookup(subChunks, x => x.Length);
+			SubChunkLineIndex = CreateSumLookup(subChunks, x => x.LineCount);
 
-			if(SubChunks.Count > 0)
+			if(subChunks.Count > 0)
 			{
-				int LastSubChunkIdx = SubChunks.Count - 1;
-				Length = SubChunkOffset[LastSubChunkIdx] + SubChunks[LastSubChunkIdx].Length;
-				LineCount = SubChunkLineIndex[LastSubChunkIdx] + SubChunks[LastSubChunkIdx].LineCount;
+				int lastSubChunkIdx = subChunks.Count - 1;
+				Length = SubChunkOffset[lastSubChunkIdx] + subChunks[lastSubChunkIdx].Length;
+				LineCount = SubChunkLineIndex[lastSubChunkIdx] + subChunks[lastSubChunkIdx].LineCount;
 			}
 		}
 
 		/// <summary>
 		/// Creates a lookup table by summing a field across a list of subchunks
 		/// </summary>
-		/// <param name="SubChunks"></param>
-		/// <param name="Field"></param>
+		/// <param name="subChunks"></param>
+		/// <param name="field"></param>
 		/// <returns></returns>
-		static int[] CreateSumLookup(IReadOnlyList<LogSubChunkData> SubChunks, Func<LogSubChunkData, int> Field)
+		static int[] CreateSumLookup(IReadOnlyList<LogSubChunkData> subChunks, Func<LogSubChunkData, int> field)
 		{
-			int[] Total = new int[SubChunks.Count];
+			int[] total = new int[subChunks.Count];
 
-			int Value = 0;
-			for (int Idx = 0; Idx + 1 < SubChunks.Count; Idx++)
+			int value = 0;
+			for (int idx = 0; idx + 1 < subChunks.Count; idx++)
 			{
-				Value += Field(SubChunks[Idx]);
-				Total[Idx + 1] = Value;
+				value += field(subChunks[idx]);
+				total[idx + 1] = value;
 			}
 
-			return Total;
+			return total;
 		}
 
 		/// <summary>
 		/// Gets the offset of a line within the file
 		/// </summary>
-		/// <param name="LineIdx">The line index</param>
+		/// <param name="lineIdx">The line index</param>
 		/// <returns>Offset of the line within the file</returns>
-		public long GetLineOffsetWithinChunk(int LineIdx)
+		public long GetLineOffsetWithinChunk(int lineIdx)
 		{
-			int SubChunkIdx = SubChunkLineIndex.BinarySearch(LineIdx);
-			if (SubChunkIdx < 0)
+			int subChunkIdx = SubChunkLineIndex.BinarySearch(lineIdx);
+			if (subChunkIdx < 0)
 			{
-				SubChunkIdx = ~SubChunkIdx - 1;
+				subChunkIdx = ~subChunkIdx - 1;
 			}
-			return SubChunkOffset[SubChunkIdx] + SubChunks[SubChunkIdx].InflateText().LineOffsets[LineIdx - SubChunkLineIndex[SubChunkIdx]];
+			return SubChunkOffset[subChunkIdx] + SubChunks[subChunkIdx].InflateText().LineOffsets[lineIdx - SubChunkLineIndex[subChunkIdx]];
 		}
 
 		/// <summary>
 		/// Gets the sub chunk index for the given line
 		/// </summary>
-		/// <param name="ChunkLineIdx">Line index within the chunk</param>
+		/// <param name="chunkLineIdx">Line index within the chunk</param>
 		/// <returns>Subchunk line index</returns>
-		public int GetSubChunkForLine(int ChunkLineIdx)
+		public int GetSubChunkForLine(int chunkLineIdx)
 		{
-			int SubChunkIdx = SubChunkLineIndex.BinarySearch(ChunkLineIdx);
-			if (SubChunkIdx < 0)
+			int subChunkIdx = SubChunkLineIndex.BinarySearch(chunkLineIdx);
+			if (subChunkIdx < 0)
 			{
-				SubChunkIdx = ~SubChunkIdx - 1;
+				subChunkIdx = ~subChunkIdx - 1;
 			}
-			return SubChunkIdx;
+			return subChunkIdx;
 		}
 
 		/// <summary>
 		/// Gets the index of the sub-chunk containing the given offset
 		/// </summary>
-		/// <param name="Offset">Offset to search for</param>
+		/// <param name="offset">Offset to search for</param>
 		/// <returns>Index of the sub-chunk</returns>
-		public int GetSubChunkForOffsetWithinChunk(int Offset)
+		public int GetSubChunkForOffsetWithinChunk(int offset)
 		{
-			int SubChunkIdx = SubChunkOffset.BinarySearch(Offset);
-			if (SubChunkIdx < 0)
+			int subChunkIdx = SubChunkOffset.BinarySearch(offset);
+			if (subChunkIdx < 0)
 			{
-				SubChunkIdx = ~SubChunkIdx - 1;
+				subChunkIdx = ~subChunkIdx - 1;
 			}
-			return SubChunkIdx;
+			return subChunkIdx;
 		}
 
 		/// <summary>
@@ -147,85 +143,85 @@ namespace Horde.Build.Logs
 		/// <summary>
 		/// Read a log chunk from the given stream
 		/// </summary>
-		/// <param name="Reader">The reader to read from</param>
-		/// <param name="Offset">Offset of this chunk within the file</param>
-		/// <param name="LineIndex">Line index of this chunk</param>
+		/// <param name="reader">The reader to read from</param>
+		/// <param name="offset">Offset of this chunk within the file</param>
+		/// <param name="lineIndex">Line index of this chunk</param>
 		/// <returns>New log chunk data</returns>
-		public static LogChunkData Read(MemoryReader Reader, long Offset, int LineIndex)
+		public static LogChunkData Read(MemoryReader reader, long offset, int lineIndex)
 		{
-			int Signature = Reader.ReadInt32();
-			if ((Signature & 0xffffff) != CurrentSignature)
+			int signature = reader.ReadInt32();
+			if ((signature & 0xffffff) != CurrentSignature)
 			{
-				List<LogSubChunkData> SubChunks = new List<LogSubChunkData>();
-				SubChunks.Add(new LogSubChunkData(LogType.Json, Offset, LineIndex, new ReadOnlyLogText(Reader.Memory)));
-				Reader.Offset = Reader.Memory.Length;
-				return new LogChunkData(Offset, LineIndex, SubChunks);
+				List<LogSubChunkData> subChunks = new List<LogSubChunkData>();
+				subChunks.Add(new LogSubChunkData(LogType.Json, offset, lineIndex, new ReadOnlyLogText(reader.Memory)));
+				reader.Offset = reader.Memory.Length;
+				return new LogChunkData(offset, lineIndex, subChunks);
 			}
 
-			int Version = Signature >> 24;
-			if (Version == 0)
+			int version = signature >> 24;
+			if (version == 0)
 			{
-				List<LogSubChunkData> SubChunks = ReadSubChunkList(Reader, Offset, LineIndex);
-				Reader.ReadVariableLengthBytes();
-				return new LogChunkData(Offset, LineIndex, SubChunks);
+				List<LogSubChunkData> subChunks = ReadSubChunkList(reader, offset, lineIndex);
+				reader.ReadVariableLengthBytes();
+				return new LogChunkData(offset, lineIndex, subChunks);
 			}
 			else
 			{
-				List<LogSubChunkData> SubChunks = ReadSubChunkList(Reader, Offset, LineIndex);
-				return new LogChunkData(Offset, LineIndex, SubChunks);
+				List<LogSubChunkData> subChunks = ReadSubChunkList(reader, offset, lineIndex);
+				return new LogChunkData(offset, lineIndex, subChunks);
 			}
 		}
 
 		/// <summary>
 		/// Read a list of sub-chunks from the given stream
 		/// </summary>
-		/// <param name="Reader">The reader to read from</param>
-		/// <param name="SubChunkOffset">Offset of this chunk within the file</param>
-		/// <param name="SubChunkLineIndex">Line index of this chunk</param>
+		/// <param name="reader">The reader to read from</param>
+		/// <param name="subChunkOffset">Offset of this chunk within the file</param>
+		/// <param name="subChunkLineIndex">Line index of this chunk</param>
 		/// <returns>List of sub-chunks</returns>
-		static List<LogSubChunkData> ReadSubChunkList(MemoryReader Reader, long SubChunkOffset, int SubChunkLineIndex)
+		static List<LogSubChunkData> ReadSubChunkList(MemoryReader reader, long subChunkOffset, int subChunkLineIndex)
 		{
-			int NumSubChunks = Reader.ReadInt32();
+			int numSubChunks = reader.ReadInt32();
 
-			List<LogSubChunkData> SubChunks = new List<LogSubChunkData>();
-			for (int Idx = 0; Idx < NumSubChunks; Idx++)
+			List<LogSubChunkData> subChunks = new List<LogSubChunkData>();
+			for (int idx = 0; idx < numSubChunks; idx++)
 			{
-				LogSubChunkData SubChunkData = Reader.ReadLogSubChunkData(SubChunkOffset, SubChunkLineIndex);
-				SubChunkOffset += SubChunkData.Length;
-				SubChunkLineIndex += SubChunkData.LineCount;
-				SubChunks.Add(SubChunkData);
+				LogSubChunkData subChunkData = reader.ReadLogSubChunkData(subChunkOffset, subChunkLineIndex);
+				subChunkOffset += subChunkData.Length;
+				subChunkLineIndex += subChunkData.LineCount;
+				subChunks.Add(subChunkData);
 			}
 
-			return SubChunks;
+			return subChunks;
 		}
 
 		/// <summary>
 		/// Construct an object from flat memory buffer
 		/// </summary>
-		/// <param name="Memory">Memory buffer</param>
-		/// <param name="Offset">Offset of this chunk within the file</param>
-		/// <param name="LineIndex">Line index of this chunk</param>
+		/// <param name="memory">Memory buffer</param>
+		/// <param name="offset">Offset of this chunk within the file</param>
+		/// <param name="lineIndex">Line index of this chunk</param>
 		/// <returns>Log chunk data</returns>
-		public static LogChunkData FromMemory(ReadOnlyMemory<byte> Memory, long Offset, int LineIndex)
+		public static LogChunkData FromMemory(ReadOnlyMemory<byte> memory, long offset, int lineIndex)
 		{
-			MemoryReader Reader = new MemoryReader(Memory);
-			LogChunkData ChunkData = Read(Reader, Offset, LineIndex);
-			Reader.CheckOffset(Memory.Length);
-			return ChunkData;
+			MemoryReader reader = new MemoryReader(memory);
+			LogChunkData chunkData = Read(reader, offset, lineIndex);
+			reader.CheckOffset(memory.Length);
+			return chunkData;
 		}
 
 		/// <summary>
 		/// Write the chunk data to a stream
 		/// </summary>
 		/// <returns>Serialized data</returns>
-		public void Write(MemoryWriter Writer)
+		public void Write(MemoryWriter writer)
 		{
-			Writer.WriteInt32(CurrentSignature | (1 << 24));
+			writer.WriteInt32(CurrentSignature | (1 << 24));
 
-			Writer.WriteInt32(SubChunks.Count);
-			foreach (LogSubChunkData SubChunk in SubChunks)
+			writer.WriteInt32(SubChunks.Count);
+			foreach (LogSubChunkData subChunk in SubChunks)
 			{
-				Writer.WriteLogSubChunkData(SubChunk);
+				writer.WriteLogSubChunkData(subChunk);
 			}
 		}
 
@@ -235,11 +231,11 @@ namespace Horde.Build.Logs
 		/// <returns>Log chunk data</returns>
 		public byte[] ToByteArray()
 		{
-			byte[] Data = new byte[GetSerializedSize()];
-			MemoryWriter Writer = new MemoryWriter(Data);
-			Write(Writer);
-			Writer.CheckOffset(Data.Length);
-			return Data;
+			byte[] data = new byte[GetSerializedSize()];
+			MemoryWriter writer = new MemoryWriter(data);
+			Write(writer);
+			writer.CheckOffset(data.Length);
+			return data;
 		}
 
 		/// <summary>
@@ -259,22 +255,22 @@ namespace Horde.Build.Logs
 		/// <summary>
 		/// Read a log chunk from the given stream
 		/// </summary>
-		/// <param name="Reader">The reader to read from</param>
-		/// <param name="Offset">Offset of this chunk within the file</param>
-		/// <param name="LineIndex">Line index of this chunk</param>
+		/// <param name="reader">The reader to read from</param>
+		/// <param name="offset">Offset of this chunk within the file</param>
+		/// <param name="lineIndex">Line index of this chunk</param>
 		/// <returns>New log chunk data</returns>
-		public static LogChunkData ReadLogChunkData(this MemoryReader Reader, long Offset, int LineIndex)
+		public static LogChunkData ReadLogChunkData(this MemoryReader reader, long offset, int lineIndex)
 		{
-			return LogChunkData.Read(Reader, Offset, LineIndex);
+			return LogChunkData.Read(reader, offset, lineIndex);
 		}
 
 		/// <summary>
 		/// Write the chunk data to a stream
 		/// </summary>
 		/// <returns>Serialized data</returns>
-		public static void WriteLogChunkData(this MemoryWriter Writer, LogChunkData ChunkData)
+		public static void WriteLogChunkData(this MemoryWriter writer, LogChunkData chunkData)
 		{
-			ChunkData.Write(Writer);
+			chunkData.Write(writer);
 		}
 	}
 }
diff --git a/Engine/Source/Programs/Horde/Horde.Build/Logs/Data/LogIndexData.cs b/Engine/Source/Programs/Horde/Horde.Build/Logs/Data/LogIndexData.cs
--- a/Engine/Source/Programs/Horde/Horde.Build/Logs/Data/LogIndexData.cs
+++ b/Engine/Source/Programs/Horde/Horde.Build/Logs/Data/LogIndexData.cs
@@ -1,16 +1,13 @@
 // Copyright Epic Games, Inc. All Rights Reserved.
 
-using EpicGames.Core;
-using Horde.Build.Utilities;
-using OpenTracing;
-using OpenTracing.Util;
 using System;
 using System.Collections.Generic;
 using System.IO;
 using System.Linq;
-using System.Reflection;
-using System.Text;
-using System.Threading.Tasks;
+using EpicGames.Core;
+using Horde.Build.Utilities;
+using OpenTracing;
+using OpenTracing.Util;
 
 namespace Horde.Build.Logs
 {
@@ -42,16 +39,16 @@ namespace Horde.Build.Logs
 		/// <summary>
 		/// Constructor
 		/// </summary>
-		/// <param name="LineIndex">Index of the first line within this block</param>
-		/// <param name="LineCount">Number of lines in the block</param>
-		/// <param name="CachedPlainText">The decompressed plain text</param>
-		/// <param name="CompressedPlainText">The compressed text data</param>
-		public LogIndexBlock(int LineIndex, int LineCount, ILogText? CachedPlainText, ReadOnlyMemory<byte> CompressedPlainText)
+		/// <param name="lineIndex">Index of the first line within this block</param>
+		/// <param name="lineCount">Number of lines in the block</param>
+		/// <param name="cachedPlainText">The decompressed plain text</param>
+		/// <param name="compressedPlainText">The compressed text data</param>
+		public LogIndexBlock(int lineIndex, int lineCount, ILogText? cachedPlainText, ReadOnlyMemory<byte> compressedPlainText)
 		{
-			this.LineIndex = LineIndex;
-			this.LineCount = LineCount;
-			this.CachedPlainText = CachedPlainText;
-			this.CompressedPlainText = CompressedPlainText;
+			LineIndex = lineIndex;
+			LineCount = lineCount;
+			CachedPlainText = cachedPlainText;
+			CompressedPlainText = compressedPlainText;
 		}
 
 		/// <summary>
@@ -75,35 +72,35 @@ namespace Horde.Build.Logs
 		/// <summary>
 		/// Deserialize an index block
 		/// </summary>
-		/// <param name="Reader">Reader to deserialize from</param>
+		/// <param name="reader">Reader to deserialize from</param>
 		/// <returns>The new index block</returns>
-		public static LogIndexBlock ReadLogIndexBlock(this MemoryReader Reader)
+		public static LogIndexBlock ReadLogIndexBlock(this MemoryReader reader)
 		{
-			int LineIndex = Reader.ReadInt32();
-			int LineCount = Reader.ReadInt32();
-			ReadOnlyMemory<byte> CompressedPlainText = Reader.ReadVariableLengthBytes();
-			return new LogIndexBlock(LineIndex, LineCount, null, CompressedPlainText);
+			int lineIndex = reader.ReadInt32();
+			int lineCount = reader.ReadInt32();
+			ReadOnlyMemory<byte> compressedPlainText = reader.ReadVariableLengthBytes();
+			return new LogIndexBlock(lineIndex, lineCount, null, compressedPlainText);
 		}
 
 		/// <summary>
 		/// Serialize an index block
 		/// </summary>
-		/// <param name="Writer">Writer to serialize to</param>
-		/// <param name="Block">The block to serialize</param>
-		public static void WriteLogIndexBlock(this MemoryWriter Writer, LogIndexBlock Block)
+		/// <param name="writer">Writer to serialize to</param>
+		/// <param name="block">The block to serialize</param>
+		public static void WriteLogIndexBlock(this MemoryWriter writer, LogIndexBlock block)
 		{
-			Writer.WriteInt32(Block.LineIndex);
-			Writer.WriteInt32(Block.LineCount);
-			Writer.WriteVariableLengthBytes(Block.CompressedPlainText.Span);
+			writer.WriteInt32(block.LineIndex);
+			writer.WriteInt32(block.LineCount);
+			writer.WriteVariableLengthBytes(block.CompressedPlainText.Span);
 		}
 
 		/// <summary>
 		/// Gets the serialized size of a block
 		/// </summary>
-		/// <param name="Block">The block to serialize</param>
-		public static int GetSerializedSize(this LogIndexBlock Block)
+		/// <param name="block">The block to serialize</param>
+		public static int GetSerializedSize(this LogIndexBlock block)
 		{
-			return sizeof(int) + sizeof(int) + (sizeof(int) + Block.CompressedPlainText.Length);
+			return sizeof(int) + sizeof(int) + (sizeof(int) + block.CompressedPlainText.Length);
 		}
 	}
 
@@ -146,17 +143,17 @@ namespace Horde.Build.Logs
 		/// <summary>
 		/// Index for tokens into the block list
 		/// </summary>
-		ReadOnlyTrie? CachedTrie;
+		ReadOnlyTrie? _cachedTrie;
 
 		/// <summary>
 		/// Number of bits in the index devoted to the block index
 		/// </summary>
-		int NumBlockBits;
+		readonly int _numBlockBits;
 
 		/// <summary>
 		/// List of text blocks
 		/// </summary>
-		LogIndexBlock[] Blocks;
+		readonly LogIndexBlock[] _blocks;
 
 		/// <summary>
 		/// Empty index data
@@ -166,32 +163,32 @@ namespace Horde.Build.Logs
 		/// <summary>
 		/// Number of lines covered by the index
 		/// </summary>
-		public int LineCount => (Blocks.Length > 0) ? (Blocks[Blocks.Length - 1].LineIndex + Blocks[Blocks.Length - 1].LineCount) : 0;
+		public int LineCount => (_blocks.Length > 0) ? (_blocks[^1].LineIndex + _blocks[^1].LineCount) : 0;
 
 		/// <summary>
 		/// Constructor
 		/// </summary>
-		/// <param name="CachedTrie">Index into the text blocks</param>
-		/// <param name="NumBlockBits">Number of bits devoted to the block index</param>
-		/// <param name="Blocks">Bloom filters for this log file</param>
-		public LogIndexData(ReadOnlyTrie? CachedTrie, int NumBlockBits, LogIndexBlock[] Blocks)
+		/// <param name="cachedTrie">Index into the text blocks</param>
+		/// <param name="numBlockBits">Number of bits devoted to the block index</param>
+		/// <param name="blocks">Bloom filters for this log file</param>
+		public LogIndexData(ReadOnlyTrie? cachedTrie, int numBlockBits, LogIndexBlock[] blocks)
 		{
-			this.CachedTrie = CachedTrie;
-			this.NumBlockBits = NumBlockBits;
-			this.Blocks = Blocks;
+			_cachedTrie = cachedTrie;
+			_numBlockBits = numBlockBits;
+			_blocks = blocks;
 		}
 
 		/// <summary>
 		/// Updates the blocks in this index relative to a give line index
 		/// </summary>
-		/// <param name="LineIndex"></param>
-		public void SetBaseLineIndex(int LineIndex)
+		/// <param name="lineIndex"></param>
+		public void SetBaseLineIndex(int lineIndex)
 		{
-			for(int BlockIdx = 0; BlockIdx < Blocks.Length; BlockIdx++)
+			for(int blockIdx = 0; blockIdx < _blocks.Length; blockIdx++)
 			{
-				LogIndexBlock Block = Blocks[BlockIdx];
-				Blocks[BlockIdx] = new LogIndexBlock(LineIndex, Block.LineCount, Block.CachedPlainText, Block.CompressedPlainText);
-				LineIndex += Block.LineCount;
+				LogIndexBlock block = _blocks[blockIdx];
+				_blocks[blockIdx] = new LogIndexBlock(lineIndex, block.LineCount, block.CachedPlainText, block.CompressedPlainText);
+				lineIndex += block.LineCount;
 			}
 		}
 
@@ -201,185 +198,185 @@ namespace Horde.Build.Logs
 		/// <returns>The trie for this data</returns>
 		ReadOnlyTrie BuildTrie()
 		{
-			if(CachedTrie == null)
+			if(_cachedTrie == null)
 			{
-				ReadOnlyTrieBuilder Builder = new ReadOnlyTrieBuilder();
-				for (int BlockIdx = 0; BlockIdx < Blocks.Length; BlockIdx++)
+				ReadOnlyTrieBuilder builder = new ReadOnlyTrieBuilder();
+				for (int blockIdx = 0; blockIdx < _blocks.Length; blockIdx++)
 				{
-					LogIndexBlock Block = Blocks[BlockIdx];
-					LogToken.GetTokens(Block.InflatePlainText().Data.Span, Token => Builder.Add((Token << NumBlockBits) | (uint)BlockIdx));
+					LogIndexBlock block = _blocks[blockIdx];
+					LogToken.GetTokens(block.InflatePlainText().Data.Span, token => builder.Add((token << _numBlockBits) | (uint)blockIdx));
 				}
-				CachedTrie = Builder.Build();
+				_cachedTrie = builder.Build();
 			}
-			return CachedTrie;
+			return _cachedTrie;
 		}
 
 		/// <summary>
 		/// Create an index from an array of blocks
 		/// </summary>
-		/// <param name="Indexes">List of indexes to merge</param>
+		/// <param name="indexes">List of indexes to merge</param>
 		/// <returns>Index data</returns>
-		public static LogIndexData Merge(IEnumerable<LogIndexData> Indexes)
+		public static LogIndexData Merge(IEnumerable<LogIndexData> indexes)
 		{
-			using IScope Scope = GlobalTracer.Instance.BuildSpan("LogIndexData.Merge").StartActive();
+			using IScope scope = GlobalTracer.Instance.BuildSpan("LogIndexData.Merge").StartActive();
 
 			// Create the combined block list
-			LogIndexBlock[] NewBlocks = Indexes.SelectMany(x => x.Blocks).ToArray();
+			LogIndexBlock[] newBlocks = indexes.SelectMany(x => x._blocks).ToArray();
 
 			// Figure out how many bits to devote to the block size
-			int NewNumBlockBits = 0;
-			while(NewBlocks.Length > (1 << NewNumBlockBits))
+			int newNumBlockBits = 0;
+			while(newBlocks.Length > (1 << newNumBlockBits))
 			{
-				NewNumBlockBits += 4;
+				newNumBlockBits += 4;
 			}
 
 			// Add all the blocks into a combined index
-			ReadOnlyTrieBuilder NewTrieBuilder = new ReadOnlyTrieBuilder();
+			ReadOnlyTrieBuilder newTrieBuilder = new ReadOnlyTrieBuilder();
 
-			int BlockCount = 0;
-			foreach(LogIndexData Index in Indexes)
+			int blockCount = 0;
+			foreach(LogIndexData index in indexes)
 			{
-				ulong BlockMask = ((1UL << Index.NumBlockBits) - 1);
-				foreach (ulong Value in Index.BuildTrie())
+				ulong blockMask = ((1UL << index._numBlockBits) - 1);
+				foreach (ulong value in index.BuildTrie())
 				{
-					ulong Token = Value >> Index.NumBlockBits;
-					int BlockIdx = (int)(Value & BlockMask);
+					ulong token = value >> index._numBlockBits;
+					int blockIdx = (int)(value & blockMask);
 
-					ulong NewToken = (Token << NewNumBlockBits) | (ulong)(long)(BlockCount + BlockIdx);
-					NewTrieBuilder.Add(NewToken);
+					ulong newToken = (token << newNumBlockBits) | (ulong)(long)(blockCount + blockIdx);
+					newTrieBuilder.Add(newToken);
 				}
-				BlockCount += Index.Blocks.Length;
+				blockCount += index._blocks.Length;
 			}
 
 			// Construct the new index data
-			ReadOnlyTrie NewTrie = NewTrieBuilder.Build();
-			return new LogIndexData(NewTrie, NewNumBlockBits, NewBlocks);
+			ReadOnlyTrie newTrie = newTrieBuilder.Build();
+			return new LogIndexData(newTrie, newNumBlockBits, newBlocks);
 		}
 
 		/// <summary>
 		/// Search for the given text in the index
 		/// </summary>
-		/// <param name="FirstLineIndex">First line index to search from</param>
-		/// <param name="Text">Text to search for</param>
-		/// <param name="Stats">Receives stats for the search</param>
+		/// <param name="firstLineIndex">First line index to search from</param>
+		/// <param name="text">Text to search for</param>
+		/// <param name="stats">Receives stats for the search</param>
 		/// <returns>List of line numbers for the text</returns>
-		public IEnumerable<int> Search(int FirstLineIndex, SearchText Text, LogSearchStats Stats)
+		public IEnumerable<int> Search(int firstLineIndex, SearchText text, LogSearchStats stats)
 		{
-			int LastBlockCount = 0;
-			foreach (int BlockIdx in EnumeratePossibleBlocks(Text.Bytes, FirstLineIndex))
+			int lastBlockCount = 0;
+			foreach (int blockIdx in EnumeratePossibleBlocks(text.Bytes, firstLineIndex))
 			{
-				LogIndexBlock Block = Blocks[BlockIdx];
+				LogIndexBlock block = _blocks[blockIdx];
 
-				Stats.NumScannedBlocks++;
-				Stats.NumDecompressedBlocks += (Block.CachedPlainText == null)? 1 : 0;
+				stats.NumScannedBlocks++;
+				stats.NumDecompressedBlocks += (block.CachedPlainText == null)? 1 : 0;
 
-				Stats.NumSkippedBlocks += BlockIdx - LastBlockCount;
-				LastBlockCount = BlockIdx + 1;
+				stats.NumSkippedBlocks += blockIdx - lastBlockCount;
+				lastBlockCount = blockIdx + 1;
 
 				// Decompress the text
-				ILogText BlockText = Block.InflatePlainText();
+				ILogText blockText = block.InflatePlainText();
 
 				// Find the initial offset within this block
-				int Offset = 0;
-				if(FirstLineIndex > Block.LineIndex)
+				int offset = 0;
+				if(firstLineIndex > block.LineIndex)
 				{
-					int LineIndexWithinBlock = FirstLineIndex - Block.LineIndex;
-					Offset = BlockText.LineOffsets[LineIndexWithinBlock];
+					int lineIndexWithinBlock = firstLineIndex - block.LineIndex;
+					offset = blockText.LineOffsets[lineIndexWithinBlock];
 				}
 
 				// Search within this block
 				for(; ;)
 				{
 					// Find the next offset
-					int NextOffset = BlockText.Data.Span.FindNextOcurrence(Offset, Text);
-					if(NextOffset == -1)
+					int nextOffset = blockText.Data.Span.FindNextOcurrence(offset, text);
+					if(nextOffset == -1)
 					{
-						Stats.NumScannedBytes += BlockText.Length - Offset;
+						stats.NumScannedBytes += blockText.Length - offset;
 						break;
 					}
 
 					// Update the stats
-					Stats.NumScannedBytes += NextOffset - Offset;
-					Offset = NextOffset;
+					stats.NumScannedBytes += nextOffset - offset;
+					offset = nextOffset;
 
 					// Check it's not another match within the same line
-					int LineIndexWithinBlock = BlockText.GetLineIndexForOffset(Offset);
-					yield return Block.LineIndex + LineIndexWithinBlock;
+					int lineIndexWithinBlock = blockText.GetLineIndexForOffset(offset);
+					yield return block.LineIndex + lineIndexWithinBlock;
 
 					// Move to the next line
-					Offset = BlockText.LineOffsets[LineIndexWithinBlock + 1];
+					offset = blockText.LineOffsets[lineIndexWithinBlock + 1];
 				}
 
 				// If the last scanned bytes is zero, we didn't have any matches from this chunk
-				if(Offset == 0 && CachedTrie != null)
+				if(offset == 0 && _cachedTrie != null)
 				{
-					Stats.NumFalsePositiveBlocks++;
+					stats.NumFalsePositiveBlocks++;
 				}
 			}
-			Stats.NumSkippedBlocks += Blocks.Length - LastBlockCount;
+			stats.NumSkippedBlocks += _blocks.Length - lastBlockCount;
 		}
 
 		/// <summary>
 		/// Search for the given text in the index
 		/// </summary>
-		/// <param name="Text">Text to search for</param>
-		/// <param name="LineIndex">The first </param>
+		/// <param name="text">Text to search for</param>
+		/// <param name="lineIndex">The first </param>
 		/// <returns>List of line numbers for the text</returns>
-		IEnumerable<int> EnumeratePossibleBlocks(ReadOnlyMemory<byte> Text, int LineIndex)
+		IEnumerable<int> EnumeratePossibleBlocks(ReadOnlyMemory<byte> text, int lineIndex)
 		{
 			// Find the starting block index
-			int BlockIdx = Blocks.BinarySearch(x => x.LineIndex, LineIndex);
-			if(BlockIdx < 0)
+			int blockIdx = _blocks.BinarySearch(x => x.LineIndex, lineIndex);
+			if(blockIdx < 0)
 			{
-				BlockIdx = Math.Max(~BlockIdx - 1, 0);
+				blockIdx = Math.Max(~blockIdx - 1, 0);
 			}
 
 			// Make sure we're not starting out of range
-			if (BlockIdx < Blocks.Length)
+			if (blockIdx < _blocks.Length)
 			{
 				// In order to be considered a possible block for a positive match, the block must contain an exact list of tokens parsed from the
 				// search text, along with a set of partial matches (for substrings that are not guaranteed to start on token boundaries). The former
 				// is relatively cheap to compute, and done by stepping through all enumerators for each token in a wave, only returning blocks
 				// which all contain the token. At the very least, this must contain blocks in the index from BlockIdx onwards.
-				List<IEnumerator<int>> Enumerators = new List<IEnumerator<int>>();
-				Enumerators.Add(Enumerable.Range(BlockIdx, Blocks.Length - BlockIdx).GetEnumerator());
+				List<IEnumerator<int>> enumerators = new List<IEnumerator<int>>();
+				enumerators.Add(Enumerable.Range(blockIdx, _blocks.Length - blockIdx).GetEnumerator());
 
 				// The second aspect of the search requires a more expensive search through the trie. This is done by running a set of arbitrary 
 				// delegates to filter the matches returned from the enumerators.
-				List<Predicate<int>> Predicates = new List<Predicate<int>>();
+				List<Predicate<int>> predicates = new List<Predicate<int>>();
 
 				// If the index has a trie, tokenize the input and generate a list of enumerators and predicates.
-				if (CachedTrie != null)
+				if (_cachedTrie != null)
 				{
 					// Find a list of filters for matching blocks
-					HashSet<ulong> Tokens = new HashSet<ulong>();
-					for (int TokenPos = 0; TokenPos < Text.Length;)
+					HashSet<ulong> tokens = new HashSet<ulong>();
+					for (int tokenPos = 0; tokenPos < text.Length;)
 					{
-						ReadOnlySpan<byte> Token = LogToken.GetTokenText(Text.Span, TokenPos);
-						if (TokenPos == 0)
+						ReadOnlySpan<byte> token = LogToken.GetTokenText(text.Span, tokenPos);
+						if (tokenPos == 0)
 						{
-							GetUnalignedTokenPredicate(Token, Token.Length == Text.Length, Predicates);
+							GetUnalignedTokenPredicate(token, token.Length == text.Length, predicates);
 						}
 						else
 						{
-							GetAlignedTokenPredicate(Token, TokenPos + Token.Length == Text.Length, Tokens, Predicates);
+							GetAlignedTokenPredicate(token, tokenPos + token.Length == text.Length, tokens, predicates);
 						}
-						TokenPos += Token.Length;
+						tokenPos += token.Length;
 					}
 
 					// Create an enumerator for each token
-					foreach (ulong Token in Tokens)
+					foreach (ulong token in tokens)
 					{
-						ulong MinValue = (Token << NumBlockBits) | (ulong)(long)BlockIdx;
-						ulong MaxValue = MinValue + (1UL << NumBlockBits) - 1;
+						ulong minValue = (token << _numBlockBits) | (ulong)(long)blockIdx;
+						ulong maxValue = minValue + (1UL << _numBlockBits) - 1;
 
-						IEnumerator<int> Enumerator = CachedTrie.EnumerateRange(MinValue, MaxValue).Select(x => (int)(x - MinValue)).GetEnumerator();
-						if (!Enumerator.MoveNext())
+						IEnumerator<int> enumerator = _cachedTrie.EnumerateRange(minValue, maxValue).Select(x => (int)(x - minValue)).GetEnumerator();
+						if (!enumerator.MoveNext())
 						{
 							yield break;
 						}
 
-						Enumerators.Add(Enumerator);
+						enumerators.Add(enumerator);
 					}
 				}
 
@@ -388,19 +385,19 @@ namespace Horde.Build.Logs
 				{
 					// Advance all the enumerators that are behind the current block index. If they are all equal, we have a match.
 					bool bMatch = true;
-					foreach (IEnumerator<int> Enumerator in Enumerators)
+					foreach (IEnumerator<int> enumerator in enumerators)
 					{
-						while (Enumerator.Current < BlockIdx)
+						while (enumerator.Current < blockIdx)
 						{
-							if (!Enumerator.MoveNext())
+							if (!enumerator.MoveNext())
 							{
 								yield break;
 							}
 						}
 
-						if (Enumerator.Current > BlockIdx)
+						if (enumerator.Current > blockIdx)
 						{
-							BlockIdx = Enumerator.Current;
+							blockIdx = enumerator.Current;
 							bMatch = false;
 						}
 					}
@@ -408,11 +405,11 @@ namespace Horde.Build.Logs
 					// Return the match and move to the next block
 					if (bMatch)
 					{
-						if (Predicates.All(Predicate => Predicate(BlockIdx)))
+						if (predicates.All(predicate => predicate(blockIdx)))
 						{
-							yield return BlockIdx;
+							yield return blockIdx;
 						}
-						BlockIdx++;
+						blockIdx++;
 					}
 				}
 			}
@@ -421,173 +418,173 @@ namespace Horde.Build.Logs
 		/// <summary>
 		/// Gets predicates for matching a token that starts 
 		/// </summary>
-		/// <param name="Text">The token text</param>
+		/// <param name="text">The token text</param>
 		/// <param name="bAllowPartialMatch">Whether to allow a partial match of the token</param>
-		/// <param name="Tokens">Set of aligned tokens that are required</param>
-		/// <param name="Predicates">List of predicates for the search</param>
-		void GetAlignedTokenPredicate(ReadOnlySpan<byte> Text, bool bAllowPartialMatch, HashSet<ulong> Tokens, List<Predicate<int>> Predicates)
+		/// <param name="tokens">Set of aligned tokens that are required</param>
+		/// <param name="predicates">List of predicates for the search</param>
+		void GetAlignedTokenPredicate(ReadOnlySpan<byte> text, bool bAllowPartialMatch, HashSet<ulong> tokens, List<Predicate<int>> predicates)
 		{
-			for (int Offset = 0; Offset < Text.Length; Offset += LogToken.MaxTokenBytes)
+			for (int offset = 0; offset < text.Length; offset += LogToken.MaxTokenBytes)
 			{
-				ulong Token = LogToken.GetWindowedTokenValue(Text, Offset);
-				if (Offset + LogToken.MaxTokenBytes > Text.Length && bAllowPartialMatch)
+				ulong token = LogToken.GetWindowedTokenValue(text, offset);
+				if (offset + LogToken.MaxTokenBytes > text.Length && bAllowPartialMatch)
 				{
-					ulong TokenMask = LogToken.GetWindowedTokenMask(Text, Offset, true);
-					Predicates.Add(BlockIdx => BlockContainsToken(BlockIdx, Token, TokenMask));
+					ulong tokenMask = LogToken.GetWindowedTokenMask(text, offset, true);
+					predicates.Add(blockIdx => BlockContainsToken(blockIdx, token, tokenMask));
 					break;
 				}
-				Tokens.Add(Token);
+				tokens.Add(token);
 			}
 		}
 
 		/// <summary>
 		/// Generates a predicate for matching a token which may or may not start on a regular token boundary
 		/// </summary>
-		/// <param name="Text">The token text</param>
+		/// <param name="text">The token text</param>
 		/// <param name="bAllowPartialMatch">Whether to allow a partial match of the token</param>
-		/// <param name="Predicates">List of predicates for the search</param>
-		void GetUnalignedTokenPredicate(ReadOnlySpan<byte> Text, bool bAllowPartialMatch, List<Predicate<int>> Predicates)
+		/// <param name="predicates">List of predicates for the search</param>
+		void GetUnalignedTokenPredicate(ReadOnlySpan<byte> text, bool bAllowPartialMatch, List<Predicate<int>> predicates)
 		{
-			byte[] TextCopy = Text.ToArray();
+			byte[] textCopy = text.ToArray();
 
-			Lazy<HashSet<int>> Blocks = new Lazy<HashSet<int>>(() =>
+			Lazy<HashSet<int>> blocks = new Lazy<HashSet<int>>(() =>
 			{
-				HashSet<int> Union = new HashSet<int>();
-				for (int Shift = 0; Shift < LogToken.MaxTokenBytes; Shift++)
+				HashSet<int> union = new HashSet<int>();
+				for (int shift = 0; shift < LogToken.MaxTokenBytes; shift++)
 				{
-					HashSet<int> Blocks = new HashSet<int>(BlocksContainingToken(TextCopy.AsSpan(), -Shift, bAllowPartialMatch));
-					for (int Offset = -Shift + LogToken.MaxTokenBytes; Offset < TextCopy.Length && Blocks.Count > 0; Offset += LogToken.MaxTokenBytes)
+					HashSet<int> blocks = new HashSet<int>(BlocksContainingToken(textCopy.AsSpan(), -shift, bAllowPartialMatch));
+					for (int offset = -shift + LogToken.MaxTokenBytes; offset < textCopy.Length && blocks.Count > 0; offset += LogToken.MaxTokenBytes)
 					{
-						Blocks.IntersectWith(BlocksContainingToken(TextCopy.AsSpan(), Offset, bAllowPartialMatch));
+						blocks.IntersectWith(BlocksContainingToken(textCopy.AsSpan(), offset, bAllowPartialMatch));
 					}
-					if (Blocks.Count > 0)
+					if (blocks.Count > 0)
 					{
-						Union.UnionWith(Blocks);
+						union.UnionWith(blocks);
 					}
 				}
-				return Union;
+				return union;
 			});
 
-			Predicates.Add(BlockIdx => Blocks.Value.Contains(BlockIdx));
+			predicates.Add(blockIdx => blocks.Value.Contains(blockIdx));
 		}
 
 		/// <summary>
 		/// Tests whether a block contains a particular token
 		/// </summary>
-		/// <param name="BlockIdx">Index of the block to search</param>
-		/// <param name="Token">The token to test</param>
-		/// <param name="TokenMask">Mask of which bits in the token are valid</param>
+		/// <param name="blockIdx">Index of the block to search</param>
+		/// <param name="token">The token to test</param>
+		/// <param name="tokenMask">Mask of which bits in the token are valid</param>
 		/// <returns>True if the given block contains a token</returns>
-		bool BlockContainsToken(int BlockIdx, ulong Token, ulong TokenMask)
+		bool BlockContainsToken(int blockIdx, ulong token, ulong tokenMask)
 		{
-			Token = (Token << NumBlockBits) | (uint)BlockIdx;
-			TokenMask = (TokenMask << NumBlockBits) | ((1UL << NumBlockBits) - 1);
-			return CachedTrie!.EnumerateValues((Value, ValueMask) => (Value & TokenMask) == (Token & ValueMask)).Any();
+			token = (token << _numBlockBits) | (uint)blockIdx;
+			tokenMask = (tokenMask << _numBlockBits) | ((1UL << _numBlockBits) - 1);
+			return _cachedTrie!.EnumerateValues((value, valueMask) => (value & tokenMask) == (token & valueMask)).Any();
 		}
 
 		/// <summary>
 		/// Tests whether a block contains a particular token
 		/// </summary>
-		/// <param name="Text">The token to test</param>
-		/// <param name="Offset">Offset of the window into the token to test</param>
+		/// <param name="text">The token to test</param>
+		/// <param name="offset">Offset of the window into the token to test</param>
 		/// <param name="bAllowPartialMatch">Whether to allow a partial match of the token</param>
 		/// <returns>True if the given block contains a token</returns>
-		IEnumerable<int> BlocksContainingToken(ReadOnlySpan<byte> Text, int Offset, bool bAllowPartialMatch)
+		IEnumerable<int> BlocksContainingToken(ReadOnlySpan<byte> text, int offset, bool bAllowPartialMatch)
 		{
-			ulong Token = LogToken.GetWindowedTokenValue(Text, Offset) << NumBlockBits;
-			ulong TokenMask = LogToken.GetWindowedTokenMask(Text, Offset, bAllowPartialMatch) << NumBlockBits;
-			ulong BlockMask = (1UL << NumBlockBits) - 1;
-			return CachedTrie!.EnumerateValues((Value, ValueMask) => (Value & TokenMask) == (Token & ValueMask)).Select(x => (int)(x & BlockMask)).Distinct();
+			ulong token = LogToken.GetWindowedTokenValue(text, offset) << _numBlockBits;
+			ulong tokenMask = LogToken.GetWindowedTokenMask(text, offset, bAllowPartialMatch) << _numBlockBits;
+			ulong blockMask = (1UL << _numBlockBits) - 1;
+			return _cachedTrie!.EnumerateValues((value, valueMask) => (value & tokenMask) == (token & valueMask)).Select(x => (int)(x & blockMask)).Distinct();
 		}
 
 		/// <summary>
 		/// Deserialize the index from memory
 		/// </summary>
-		/// <param name="Reader">Reader to deserialize from</param>
+		/// <param name="reader">Reader to deserialize from</param>
 		/// <returns>Index data</returns>
-		public static LogIndexData Read(MemoryReader Reader)
+		public static LogIndexData Read(MemoryReader reader)
 		{
-			int Version = Reader.ReadInt32();
-			if (Version == 0)
+			int version = reader.ReadInt32();
+			if (version == 0)
 			{
 				return LogIndexData.Empty;
 			}
-			else if (Version == 1)
+			else if (version == 1)
 			{
-				ReadOnlyTrie Index = Reader.ReadTrie();
-				int NumBlockBits = Reader.ReadInt32();
+				ReadOnlyTrie index = reader.ReadTrie();
+				int numBlockBits = reader.ReadInt32();
 
-				LogIndexBlock[] Blocks = new LogIndexBlock[Reader.ReadInt32()];
-				for (int Idx = 0; Idx < Blocks.Length; Idx++)
+				LogIndexBlock[] blocks = new LogIndexBlock[reader.ReadInt32()];
+				for (int idx = 0; idx < blocks.Length; idx++)
 				{
-					int LineIndex = (Idx > 0) ? (Blocks[Idx - 1].LineIndex + Blocks[Idx - 1].LineCount) : 0;
-					int LineCount = Reader.ReadInt32();
-					ReadOnlyMemory<byte> CompressedPlainText = Reader.ReadVariableLengthBytes();
-					Reader.ReadTrie();
-					Blocks[Idx] = new LogIndexBlock(LineIndex, LineCount, null, CompressedPlainText);
+					int lineIndex = (idx > 0) ? (blocks[idx - 1].LineIndex + blocks[idx - 1].LineCount) : 0;
+					int lineCount = reader.ReadInt32();
+					ReadOnlyMemory<byte> compressedPlainText = reader.ReadVariableLengthBytes();
+					reader.ReadTrie();
+					blocks[idx] = new LogIndexBlock(lineIndex, lineCount, null, compressedPlainText);
 				}
 
-				return new LogIndexData(Index, NumBlockBits, Blocks);
+				return new LogIndexData(index, numBlockBits, blocks);
 			}
-			else if (Version == 2)
+			else if (version == 2)
 			{
-				ReadOnlyTrie Index = Reader.ReadTrie();
-				int NumBlockBits = Reader.ReadInt32();
+				ReadOnlyTrie index = reader.ReadTrie();
+				int numBlockBits = reader.ReadInt32();
 
-				LogIndexBlock[] Blocks = new LogIndexBlock[Reader.ReadInt32()];
-				for (int Idx = 0; Idx < Blocks.Length; Idx++)
+				LogIndexBlock[] blocks = new LogIndexBlock[reader.ReadInt32()];
+				for (int idx = 0; idx < blocks.Length; idx++)
 				{
-					int LineIndex = Reader.ReadInt32();
-					int LineCount = Reader.ReadInt32();
-					ReadOnlyMemory<byte> CompressedPlainText = Reader.ReadVariableLengthBytes();
-					Reader.ReadTrie();
-					Blocks[Idx] = new LogIndexBlock(LineIndex, LineCount, null, CompressedPlainText);
+					int lineIndex = reader.ReadInt32();
+					int lineCount = reader.ReadInt32();
+					ReadOnlyMemory<byte> compressedPlainText = reader.ReadVariableLengthBytes();
+					reader.ReadTrie();
+					blocks[idx] = new LogIndexBlock(lineIndex, lineCount, null, compressedPlainText);
 				}
 
-				return new LogIndexData(Index, NumBlockBits, Blocks);
+				return new LogIndexData(index, numBlockBits, blocks);
 			}
-			else if (Version == 3)
+			else if (version == 3)
 			{
-				ReadOnlyTrie Index = Reader.ReadTrie();
-				int NumBlockBits = Reader.ReadInt32();
-				LogIndexBlock[] Blocks = Reader.ReadVariableLengthArray(() => Reader.ReadLogIndexBlock());
-				for(int Idx = 1; Idx < Blocks.Length; Idx++)
+				ReadOnlyTrie index = reader.ReadTrie();
+				int numBlockBits = reader.ReadInt32();
+				LogIndexBlock[] blocks = reader.ReadVariableLengthArray(() => reader.ReadLogIndexBlock());
+				for(int idx = 1; idx < blocks.Length; idx++)
 				{
-					if(Blocks[Idx].LineIndex == 0)
+					if(blocks[idx].LineIndex == 0)
 					{
-						int LineIndex = Blocks[Idx - 1].LineIndex + Blocks[Idx - 1].LineCount;
-						Blocks[Idx] = new LogIndexBlock(LineIndex, Blocks[Idx].LineCount, Blocks[Idx].CachedPlainText, Blocks[Idx].CompressedPlainText);
+						int lineIndex = blocks[idx - 1].LineIndex + blocks[idx - 1].LineCount;
+						blocks[idx] = new LogIndexBlock(lineIndex, blocks[idx].LineCount, blocks[idx].CachedPlainText, blocks[idx].CompressedPlainText);
 					}
 				}
-				return new LogIndexData(Index, NumBlockBits, Blocks);
+				return new LogIndexData(index, numBlockBits, blocks);
 			}
 			else
 			{
-				throw new InvalidDataException($"Invalid index version number {Version}");
+				throw new InvalidDataException($"Invalid index version number {version}");
 			}
 		}
 
 		/// <summary>
 		/// Serialize an index into memory
 		/// </summary>
-		/// <param name="Writer">Writer to serialize to</param>
-		public void Write(MemoryWriter Writer)
+		/// <param name="writer">Writer to serialize to</param>
+		public void Write(MemoryWriter writer)
 		{
-			Writer.WriteInt32(3);
-			Writer.WriteTrie(BuildTrie());
-			Writer.WriteInt32(NumBlockBits);
-			Writer.WriteVariableLengthArray(Blocks, x => Writer.WriteLogIndexBlock(x));
+			writer.WriteInt32(3);
+			writer.WriteTrie(BuildTrie());
+			writer.WriteInt32(_numBlockBits);
+			writer.WriteVariableLengthArray(_blocks, x => writer.WriteLogIndexBlock(x));
 		}
 
 		/// <summary>
 		/// Deserialize the index from memory
 		/// </summary>
-		/// <param name="Memory">Memory to deserialize from</param>
+		/// <param name="memory">Memory to deserialize from</param>
 		/// <returns>Index data</returns>
-		public static LogIndexData FromMemory(ReadOnlyMemory<byte> Memory)
+		public static LogIndexData FromMemory(ReadOnlyMemory<byte> memory)
 		{
-			MemoryReader Reader = new MemoryReader(Memory);
-			return Read(Reader);
+			MemoryReader reader = new MemoryReader(memory);
+			return Read(reader);
 		}
 
 		/// <summary>
@@ -596,13 +593,13 @@ namespace Horde.Build.Logs
 		/// <returns>Index data</returns>
 		public byte[] ToByteArray()
 		{
-			byte[] Buffer = new byte[GetSerializedSize()];
+			byte[] buffer = new byte[GetSerializedSize()];
 
-			MemoryWriter Writer = new MemoryWriter(Buffer);
-			Write(Writer);
-			Writer.CheckOffset(Buffer.Length);
+			MemoryWriter writer = new MemoryWriter(buffer);
+			Write(writer);
+			writer.CheckOffset(buffer.Length);
 
-			return Buffer;
+			return buffer;
 		}
 
 		/// <summary>
@@ -611,7 +608,7 @@ namespace Horde.Build.Logs
 		/// <returns>The serialized size</returns>
 		public int GetSerializedSize()
 		{
-			return sizeof(int) + BuildTrie().GetSerializedSize() + sizeof(int) + (sizeof(int) + Blocks.Sum(x => x.GetSerializedSize()));
+			return sizeof(int) + BuildTrie().GetSerializedSize() + sizeof(int) + (sizeof(int) + _blocks.Sum(x => x.GetSerializedSize()));
 		}
 	}
 
@@ -620,21 +617,21 @@ namespace Horde.Build.Logs
 		/// <summary>
 		/// Deserialize the index from memory
 		/// </summary>
-		/// <param name="Reader">Reader to deserialize from</param>
+		/// <param name="reader">Reader to deserialize from</param>
 		/// <returns>Index data</returns>
-		public static LogIndexData ReadLogIndexData(this MemoryReader Reader)
+		public static LogIndexData ReadLogIndexData(this MemoryReader reader)
 		{
-			return LogIndexData.Read(Reader);
+			return LogIndexData.Read(reader);
 		}
 
 		/// <summary>
 		/// Serialize an index into memory
 		/// </summary>
-		/// <param name="Writer">Writer to serialize to</param>
-		/// <param name="Index">The index to write</param>
-		public static void WriteLogIndexData(this MemoryWriter Writer, LogIndexData Index)
+		/// <param name="writer">Writer to serialize to</param>
+		/// <param name="index">The index to write</param>
+		public static void WriteLogIndexData(this MemoryWriter writer, LogIndexData index)
 		{
-			Index.Write(Writer);
+			index.Write(writer);
 		}
 	}
 }
