Horde: Fix coding conventions within Horde.Build.

diff --git a/Engine/Source/Programs/Horde/Horde.Build/Logs/Data/LogSubChunkData.cs b/Engine/Source/Programs/Horde/Horde.Build/Logs/Data/LogSubChunkData.cs
--- a/Engine/Source/Programs/Horde/Horde.Build/Logs/Data/LogSubChunkData.cs
+++ b/Engine/Source/Programs/Horde/Horde.Build/Logs/Data/LogSubChunkData.cs
@@ -1,17 +1,9 @@
 // Copyright Epic Games, Inc. All Rights Reserved.
 
+using System;
 using EpicGames.Core;
 using Horde.Build.Api;
 using Horde.Build.Utilities;
-using ICSharpCode.SharpZipLib.BZip2;
-using Microsoft.AspNetCore.Routing.Constraints;
-using System;
-using System.Collections.Generic;
-using System.IO;
-using System.Linq;
-using System.Runtime.CompilerServices;
-using System.Text;
-using System.Threading.Tasks;
 
 namespace Horde.Build.Logs
 {
@@ -48,28 +40,28 @@ namespace Horde.Build.Logs
 		/// <summary>
 		/// Text data
 		/// </summary>
-		ILogText? TextInternal;
+		ILogText? _textInternal;
 
 		/// <summary>
 		/// Compressed text data
 		/// </summary>
-		ReadOnlyMemory<byte> CompressedTextInternal;
+		ReadOnlyMemory<byte> _compressedTextInternal;
 
 		/// <summary>
 		/// The index for this sub-chunk
 		/// </summary>
-		LogIndexData? IndexInternal;
+		LogIndexData? _indexInternal;
 
 		/// <summary>
 		/// The log text
 		/// </summary>
 		public ILogText InflateText()
 		{
-			if (TextInternal == null)
+			if (_textInternal == null)
 			{
-				TextInternal = new ReadOnlyLogText(CompressedTextInternal.DecompressBzip2());
+				_textInternal = new ReadOnlyLogText(_compressedTextInternal.DecompressBzip2());
 			}
-			return TextInternal;
+			return _textInternal;
 		}
 
 		/// <summary>
@@ -77,11 +69,11 @@ namespace Horde.Build.Logs
 		/// </summary>
 		public ReadOnlyMemory<byte> DeflateText()
 		{
-			if(CompressedTextInternal.IsEmpty)
+			if(_compressedTextInternal.IsEmpty)
 			{
-				CompressedTextInternal = TextInternal!.Data.CompressBzip2();
+				_compressedTextInternal = _textInternal!.Data.CompressBzip2();
 			}
-			return CompressedTextInternal;
+			return _compressedTextInternal;
 		}
 
 		/// <summary>
@@ -89,116 +81,116 @@ namespace Horde.Build.Logs
 		/// </summary>
 		public LogIndexData BuildIndex()
 		{
-			if(IndexInternal == null)
+			if(_indexInternal == null)
 			{
-				ILogText PlainText = InflateText();
+				ILogText plainText = InflateText();
 				if (Type != LogType.Text)
 				{
-					PlainText = PlainText.ToPlainText();
+					plainText = plainText.ToPlainText();
 				}
 
-				LogIndexBlock[] Blocks = new LogIndexBlock[1];
-				Blocks[0] = new LogIndexBlock(LineIndex, LineCount, PlainText, PlainText.Data.CompressBzip2());
-				IndexInternal = new LogIndexData(null, 0, Blocks);
+				LogIndexBlock[] blocks = new LogIndexBlock[1];
+				blocks[0] = new LogIndexBlock(LineIndex, LineCount, plainText, plainText.Data.CompressBzip2());
+				_indexInternal = new LogIndexData(null, 0, blocks);
 			}
-			return IndexInternal;
+			return _indexInternal;
 		}
 
 		/// <summary>
 		/// Constructor
 		/// </summary>
-		/// <param name="Type">Type of data stored in this subchunk</param>
-		/// <param name="Offset">Offset within the file of this sub-chunk</param>
-		/// <param name="LineIndex">Index of the first line in this sub-chunk</param>
-		/// <param name="Text">Text to add</param>
-		public LogSubChunkData(LogType Type, long Offset, int LineIndex, ILogText Text)
+		/// <param name="type">Type of data stored in this subchunk</param>
+		/// <param name="offset">Offset within the file of this sub-chunk</param>
+		/// <param name="lineIndex">Index of the first line in this sub-chunk</param>
+		/// <param name="text">Text to add</param>
+		public LogSubChunkData(LogType type, long offset, int lineIndex, ILogText text)
 		{
-			this.Type = Type;
-			this.Offset = Offset;
-			this.Length = Text.Data.Length;
-			this.LineIndex = LineIndex;
-			this.LineCount = Text.LineCount;
-			this.TextInternal = Text;
+			Type = type;
+			Offset = offset;
+			Length = text.Data.Length;
+			LineIndex = lineIndex;
+			LineCount = text.LineCount;
+			_textInternal = text;
 		}
 
 		/// <summary>
 		/// Constructor for raw data
 		/// </summary>
-		/// <param name="Type">Type of data stored in this subchunk</param>
-		/// <param name="Offset">Offset within the file of this sub-chunk</param>
-		/// <param name="Length">Length of the uncompressed data</param>
-		/// <param name="LineIndex">Index of the first line</param>
-		/// <param name="LineCount">Number of lines in the uncompressed text</param>
-		/// <param name="CompressedText">Compressed text data</param>
-		/// <param name="Index">Index data</param>
-		public LogSubChunkData(LogType Type, long Offset, int Length, int LineIndex, int LineCount, ReadOnlyMemory<byte> CompressedText, LogIndexData? Index)
+		/// <param name="type">Type of data stored in this subchunk</param>
+		/// <param name="offset">Offset within the file of this sub-chunk</param>
+		/// <param name="length">Length of the uncompressed data</param>
+		/// <param name="lineIndex">Index of the first line</param>
+		/// <param name="lineCount">Number of lines in the uncompressed text</param>
+		/// <param name="compressedText">Compressed text data</param>
+		/// <param name="index">Index data</param>
+		public LogSubChunkData(LogType type, long offset, int length, int lineIndex, int lineCount, ReadOnlyMemory<byte> compressedText, LogIndexData? index)
 		{
-			this.Type = Type;
-			this.Offset = Offset;
-			this.Length = Length;
-			this.LineIndex = LineIndex;
-			this.LineCount = LineCount;
-			this.CompressedTextInternal = CompressedText;
-			this.IndexInternal = Index;
+			Type = type;
+			Offset = offset;
+			Length = length;
+			LineIndex = lineIndex;
+			LineCount = lineCount;
+			_compressedTextInternal = compressedText;
+			_indexInternal = index;
 		}
 
 		/// <summary>
 		/// Constructs a sub-chunk from a block of memory. Uses slices of the given memory buffer rather than copying the data.
 		/// </summary>
-		/// <param name="Reader">The reader to read from</param>
-		/// <param name="Offset">Offset of the sub-chunk within this file</param>
-		/// <param name="LineIndex">Index of the first line in this subchunk</param>
+		/// <param name="reader">The reader to read from</param>
+		/// <param name="offset">Offset of the sub-chunk within this file</param>
+		/// <param name="lineIndex">Index of the first line in this subchunk</param>
 		/// <returns>New subchunk data</returns>
-		public static LogSubChunkData Read(MemoryReader Reader, long Offset, int LineIndex)
+		public static LogSubChunkData Read(MemoryReader reader, long offset, int lineIndex)
 		{
-			int Version = Reader.ReadInt32(); // Version placeholder
-			if (Version == 0)
+			int version = reader.ReadInt32(); // Version placeholder
+			if (version == 0)
 			{
-				LogType Type = (LogType)Reader.ReadInt32();
-				int Length = Reader.ReadInt32();
-				int LineCount = Reader.ReadInt32();
-				ReadOnlyMemory<byte> CompressedText = Reader.ReadVariableLengthBytes();
-				Reader.ReadVariableLengthBytes();
-				return new LogSubChunkData(Type, Offset, Length, LineIndex, LineCount, CompressedText, null);
+				LogType type = (LogType)reader.ReadInt32();
+				int length = reader.ReadInt32();
+				int lineCount = reader.ReadInt32();
+				ReadOnlyMemory<byte> compressedText = reader.ReadVariableLengthBytes();
+				reader.ReadVariableLengthBytes();
+				return new LogSubChunkData(type, offset, length, lineIndex, lineCount, compressedText, null);
 			}
-			else if (Version == 1)
+			else if (version == 1)
 			{
-				LogType Type = (LogType)Reader.ReadInt32();
-				int Length = Reader.ReadInt32();
-				int LineCount = Reader.ReadInt32();
-				ReadOnlyMemory<byte> CompressedText = Reader.ReadVariableLengthBytes();
-				ReadOnlyMemory<byte> CompressedPlainText = Reader.ReadVariableLengthBytes();
-				ReadOnlyTrie Trie = Reader.ReadTrie();
-				LogIndexBlock IndexBlock = new LogIndexBlock(LineIndex, LineCount, null, CompressedPlainText);
-				LogIndexData Index = new LogIndexData(Trie, 0, new[] { IndexBlock });
-				return new LogSubChunkData(Type, Offset, Length, LineIndex, LineCount, CompressedText, Index);
+				LogType type = (LogType)reader.ReadInt32();
+				int length = reader.ReadInt32();
+				int lineCount = reader.ReadInt32();
+				ReadOnlyMemory<byte> compressedText = reader.ReadVariableLengthBytes();
+				ReadOnlyMemory<byte> compressedPlainText = reader.ReadVariableLengthBytes();
+				ReadOnlyTrie trie = reader.ReadTrie();
+				LogIndexBlock indexBlock = new LogIndexBlock(lineIndex, lineCount, null, compressedPlainText);
+				LogIndexData index = new LogIndexData(trie, 0, new[] { indexBlock });
+				return new LogSubChunkData(type, offset, length, lineIndex, lineCount, compressedText, index);
 			}
 			else
 			{
-				LogType Type = (LogType)Reader.ReadInt32();
-				int Length = Reader.ReadInt32();
-				int LineCount = Reader.ReadInt32();
-				ReadOnlyMemory<byte> CompressedText = Reader.ReadVariableLengthBytes();
-				LogIndexData Index = Reader.ReadLogIndexData();
-				Index.SetBaseLineIndex(LineIndex); // Fix for incorrectly saved data
-				return new LogSubChunkData(Type, Offset, Length, LineIndex, LineCount, CompressedText, Index);
+				LogType type = (LogType)reader.ReadInt32();
+				int length = reader.ReadInt32();
+				int lineCount = reader.ReadInt32();
+				ReadOnlyMemory<byte> compressedText = reader.ReadVariableLengthBytes();
+				LogIndexData index = reader.ReadLogIndexData();
+				index.SetBaseLineIndex(lineIndex); // Fix for incorrectly saved data
+				return new LogSubChunkData(type, offset, length, lineIndex, lineCount, compressedText, index);
 			}
 		}
 
 		/// <summary>
 		/// Serializes the sub-chunk to a stream
 		/// </summary>
-		/// <param name="Writer">Writer to output to</param>
-		public void Write(MemoryWriter Writer)
+		/// <param name="writer">Writer to output to</param>
+		public void Write(MemoryWriter writer)
 		{
-			Writer.WriteInt32(2); // Version placeholder
+			writer.WriteInt32(2); // Version placeholder
 
-			Writer.WriteInt32((int)Type);
-			Writer.WriteInt32(Length);
-			Writer.WriteInt32(LineCount);
+			writer.WriteInt32((int)Type);
+			writer.WriteInt32(Length);
+			writer.WriteInt32(LineCount);
 
-			Writer.WriteVariableLengthBytes(DeflateText().Span);
-			Writer.WriteLogIndexData(BuildIndex());
+			writer.WriteVariableLengthBytes(DeflateText().Span);
+			writer.WriteLogIndexData(BuildIndex());
 		}
 
 		/// <summary>
@@ -207,13 +199,13 @@ namespace Horde.Build.Logs
 		/// <returns>Byte array</returns>
 		public byte[] ToByteArray()
 		{
-			byte[] Data = new byte[GetSerializedSize()];
+			byte[] data = new byte[GetSerializedSize()];
 
-			MemoryWriter Writer = new MemoryWriter(Data);
-			Write(Writer);
-			Writer.CheckOffset(Data.Length);
+			MemoryWriter writer = new MemoryWriter(data);
+			Write(writer);
+			writer.CheckOffset(data.Length);
 
-			return Data;
+			return data;
 		}
 
 		/// <summary>
@@ -233,23 +225,23 @@ namespace Horde.Build.Logs
 		/// <summary>
 		/// Constructs a sub-chunk from a block of memory. Uses slices of the given memory buffer rather than copying the data.
 		/// </summary>
-		/// <param name="Reader">The reader to read from</param>
-		/// <param name="Offset">Offset of this sub-chunk within the file</param>
-		/// <param name="LineIndex">Index of the first line within this chunk</param>
+		/// <param name="reader">The reader to read from</param>
+		/// <param name="offset">Offset of this sub-chunk within the file</param>
+		/// <param name="lineIndex">Index of the first line within this chunk</param>
 		/// <returns>New subchunk data</returns>
-		public static LogSubChunkData ReadLogSubChunkData(this MemoryReader Reader, long Offset, int LineIndex)
+		public static LogSubChunkData ReadLogSubChunkData(this MemoryReader reader, long offset, int lineIndex)
 		{
-			return LogSubChunkData.Read(Reader, Offset, LineIndex);
+			return LogSubChunkData.Read(reader, offset, lineIndex);
 		}
 
 		/// <summary>
 		/// Serializes the sub-chunk to a stream
 		/// </summary>
-		/// <param name="Writer">Writer to output to</param>
-		/// <param name="SubChunkData">The sub-chunk data to write</param>
-		public static void WriteLogSubChunkData(this MemoryWriter Writer, LogSubChunkData SubChunkData)
+		/// <param name="writer">Writer to output to</param>
+		/// <param name="subChunkData">The sub-chunk data to write</param>
+		public static void WriteLogSubChunkData(this MemoryWriter writer, LogSubChunkData subChunkData)
 		{
-			SubChunkData.Write(Writer);
+			subChunkData.Write(writer);
 		}
 	}
 }
diff --git a/Engine/Source/Programs/Horde/Horde.Build/Logs/Data/LogText.cs b/Engine/Source/Programs/Horde/Horde.Build/Logs/Data/LogText.cs
--- a/Engine/Source/Programs/Horde/Horde.Build/Logs/Data/LogText.cs
+++ b/Engine/Source/Programs/Horde/Horde.Build/Logs/Data/LogText.cs
@@ -1,16 +1,12 @@
 // Copyright Epic Games, Inc. All Rights Reserved.
 
-using EpicGames.Core;
-using Microsoft.VisualBasic.CompilerServices;
 using System;
-using System.Collections;
 using System.Collections.Generic;
 using System.Diagnostics;
 using System.IO;
-using System.Linq;
 using System.Text;
 using System.Text.Json;
-using System.Threading.Tasks;
+using EpicGames.Core;
 
 namespace Horde.Build.Logs
 {
@@ -53,51 +49,51 @@ namespace Horde.Build.Logs
 		/// <summary>
 		/// Finds the number of newline characters in the given span
 		/// </summary>
-		/// <param name="Span">Span of utf-8 characters to scan</param>
-		/// <param name="Offset">Starting offset within the span</param>
+		/// <param name="span">Span of utf-8 characters to scan</param>
+		/// <param name="offset">Starting offset within the span</param>
 		/// <returns>Number of lines in the given span</returns>
-		public static int GetLineCount(ReadOnlySpan<byte> Span, int Offset)
+		public static int GetLineCount(ReadOnlySpan<byte> span, int offset)
 		{
-			int LineCount = 0;
-			for (; Offset < Span.Length; Offset++)
+			int lineCount = 0;
+			for (; offset < span.Length; offset++)
 			{
-				if (Span[Offset] == (byte)'\n')
+				if (span[offset] == (byte)'\n')
 				{
-					LineCount++;
+					lineCount++;
 				}
 			}
-			return LineCount;
+			return lineCount;
 		}
 
 		/// <summary>
 		/// Find the offsets of the start of each line within the given span
 		/// </summary>
-		/// <param name="Span">Span to search</param>
-		/// <param name="Offset">Offset within the span to start searching</param>
-		/// <param name="LineOffsets">Receives the list of line offsets</param>
-		/// <param name="LineCount">The current line count</param>
-		public static void FindLineOffsets(ReadOnlySpan<byte> Span, int Offset, Span<int> LineOffsets, int LineCount)
+		/// <param name="span">Span to search</param>
+		/// <param name="offset">Offset within the span to start searching</param>
+		/// <param name="lineOffsets">Receives the list of line offsets</param>
+		/// <param name="lineCount">The current line count</param>
+		public static void FindLineOffsets(ReadOnlySpan<byte> span, int offset, Span<int> lineOffsets, int lineCount)
 		{
-			Debug.Assert(Span.Length == 0 || Span[Span.Length - 1] == (byte)'\n');
+			Debug.Assert(span.Length == 0 || span[^1] == (byte)'\n');
 			for(; ; )
 			{
 				// Store the start of this line
-				LineOffsets[LineCount++] = Offset;
+				lineOffsets[lineCount++] = offset;
 
 				// Check if we're at the end
-				if(Offset >= Span.Length)
+				if(offset >= span.Length)
 				{
 					break;
 				}
 
 				// Move to the end of this line
-				while (Span[Offset] != '\n')
+				while (span[offset] != '\n')
 				{
-					Offset++;
+					offset++;
 				}
 
 				// Move to the start of the next line
-				Offset++;
+				offset++;
 			}
 		}
 	}
@@ -111,42 +107,33 @@ namespace Horde.Build.Logs
 		public ReadOnlyMemory<byte> Data { get; private set; }
 
 		/// <inheritdoc/>
-		public int Length
-		{
-			get { return Data.Length; }
-		}
+		public int Length => Data.Length;
 
 		/// <inheritdoc/>
 		public IReadOnlyList<int> LineOffsets { get; private set; }
 
 		/// <inheritdoc/>
-		public int LineCount
-		{
-			get { return LineOffsets.Count - 1; }
-		}
+		public int LineCount => LineOffsets.Count - 1;
 
 		/// <inheritdoc/>
-		public int AllocatedSize
-		{
-			get { return Data.Length + (LineOffsets.Count * sizeof(int)); }
-		}
+		public int AllocatedSize => Data.Length + (LineOffsets.Count * sizeof(int));
 
 		/// <summary>
 		/// Constructor
 		/// </summary>
-		/// <param name="Data">Data to construct from</param>
-		public ReadOnlyLogText(ReadOnlyMemory<byte> Data)
+		/// <param name="data">Data to construct from</param>
+		public ReadOnlyLogText(ReadOnlyMemory<byte> data)
 		{
-			ReadOnlySpan<byte> Span = Data.Span;
+			ReadOnlySpan<byte> span = data.Span;
 
-			int LineCount = LogTextHelpers.GetLineCount(Span, 0);
+			int lineCount = LogTextHelpers.GetLineCount(span, 0);
 
-			int[] LineOffsets = new int[LineCount + 1];
-			LogTextHelpers.FindLineOffsets(Span, 0, LineOffsets, 0);
-			Debug.Assert(LineOffsets[LineCount] == Data.Length);
+			int[] lineOffsets = new int[lineCount + 1];
+			LogTextHelpers.FindLineOffsets(span, 0, lineOffsets, 0);
+			Debug.Assert(lineOffsets[lineCount] == data.Length);
 
-			this.Data = Data;
-			this.LineOffsets = LineOffsets;
+			Data = data;
+			LineOffsets = lineOffsets;
 		}
 	}
 
@@ -158,15 +145,15 @@ namespace Horde.Build.Logs
 		/// <summary>
 		/// Accessor for Data
 		/// </summary>
-		byte[] InternalData;
+		byte[] _internalData;
 
 		/// <inheritdoc/>
-		public ReadOnlyMemory<byte> Data => InternalData.AsMemory(0, Length);
+		public ReadOnlyMemory<byte> Data => _internalData.AsMemory(0, Length);
 
 		/// <summary>
 		/// Hacky accessor for InternalData, for serializing to Redis
 		/// </summary>
-		public byte[] InternalDataHACK => InternalData;
+		public byte[] InternalDataHack => _internalData;
 
 		/// <inheritdoc/>
 		public int Length
@@ -178,100 +165,94 @@ namespace Horde.Build.Logs
 		/// <summary>
 		/// Offsets of the start of each line within the data
 		/// </summary>
-		List<int> InternalLineOffsets = new List<int> { 0 };
+		readonly List<int> _internalLineOffsets = new List<int> { 0 };
 
 		/// <inheritdoc/>
-		public IReadOnlyList<int> LineOffsets => InternalLineOffsets;
+		public IReadOnlyList<int> LineOffsets => _internalLineOffsets;
 
 		/// <inheritdoc/>
-		public int LineCount => InternalLineOffsets.Count - 1;
+		public int LineCount => _internalLineOffsets.Count - 1;
 
 		/// <summary>
 		/// Chunk Size
 		/// </summary>
-		public int MaxLength
-		{
-			get { return InternalData.Length; }
-		}
+		public int MaxLength => _internalData.Length;
 
 		/// <summary>
 		/// Computes the size of data used by this object
 		/// </summary>
-		public int AllocatedSize
-		{
-			get { return InternalData.Length + (InternalLineOffsets.Capacity * sizeof(int)); }
-		}
+		public int AllocatedSize => _internalData.Length + (_internalLineOffsets.Capacity * sizeof(int));
 
 		/// <summary>
 		/// Constructor
 		/// </summary>
 		public LogText()
 		{
-			this.InternalData = Array.Empty<byte>();
+			_internalData = Array.Empty<byte>();
 		}
 
 		/// <summary>
 		/// Constructor
 		/// </summary>
-		/// <param name="Data">Data to initialize this chunk with. Ownership of this array is transferred to the chunk, and its length determines the chunk size.</param>
-		/// <param name="Length">Number of valid bytes within the initial data array</param>
-		public LogText(byte[] Data, int Length)
+		/// <param name="data">Data to initialize this chunk with. Ownership of this array is transferred to the chunk, and its length determines the chunk size.</param>
+		/// <param name="length">Number of valid bytes within the initial data array</param>
+		public LogText(byte[] data, int length)
 		{
-			this.InternalData = Data;
-			this.Length = 0; // Updated below
+			_internalData = data;
+			Length = 0; // Updated below
 
-			UpdateLength(Length);
+			UpdateLength(length);
 		}
 
 		/// <summary>
 		/// Create a new chunk data object with the given data appended. The internal buffers are reused, with the assumption that
 		/// there is no contention over writing to the same location in the chunk.
 		/// </summary>
-		/// <param name="TextData">The data to append</param>
+		/// <param name="textData">The data to append</param>
 		/// <returns>New chunk data object</returns>
-		public void Append(ReadOnlySpan<byte> TextData)
+		public void Append(ReadOnlySpan<byte> textData)
 		{
-			TextData.CopyTo(InternalData.AsSpan(Length, TextData.Length));
-			UpdateLength(Length + TextData.Length);
+			textData.CopyTo(_internalData.AsSpan(Length, textData.Length));
+			UpdateLength(Length + textData.Length);
 		}
 
 		/// <summary>
 		/// Updates the plain text representation of this chunk
 		/// </summary>
-		public void AppendPlainText(ILogText SrcText, int SrcLineIndex, int SrcLineCount)
+		public void AppendPlainText(ILogText srcText, int srcLineIndex, int srcLineCount)
 		{
 			// Convert all the lines to plain text
-			for (int Idx = 0; Idx < SrcLineCount; Idx++)
+			for (int idx = 0; idx < srcLineCount; idx++)
 			{
-				int LineOffset = SrcText.LineOffsets[SrcLineIndex + Idx];
-				int NextLineOffset = SrcText.LineOffsets[SrcLineIndex + Idx + 1];
-				ReadOnlySpan<byte> InputLine = SrcText.Data.Slice(LineOffset, NextLineOffset - LineOffset).Span;
+				int lineOffset = srcText.LineOffsets[srcLineIndex + idx];
+				int nextLineOffset = srcText.LineOffsets[srcLineIndex + idx + 1];
+				ReadOnlySpan<byte> inputLine = srcText.Data.Slice(lineOffset, nextLineOffset - lineOffset).Span;
 
 				// Make sure the output buffer is large enough
-				int RequiredSpace = NextLineOffset - LineOffset;
-				if (RequiredSpace > InternalData.Length - Length)
+				int requiredSpace = nextLineOffset - lineOffset;
+				if (requiredSpace > _internalData.Length - Length)
 				{
-					int MaxLineOffset = SrcText.LineOffsets[SrcLineIndex + SrcLineCount];
-					Array.Resize(ref InternalData, InternalData.Length + RequiredSpace + (MaxLineOffset - NextLineOffset) / 2);
+					int maxLineOffset = srcText.LineOffsets[srcLineIndex + srcLineCount];
+					Array.Resize(ref _internalData, _internalData.Length + requiredSpace + (maxLineOffset - nextLineOffset) / 2);
 				}
 
 				// Convert the line to plain text
-				Length = ConvertToPlainText(InputLine, InternalData, Length);
-				InternalLineOffsets.Add(Length);
+				Length = ConvertToPlainText(inputLine, _internalData, Length);
+				_internalLineOffsets.Add(Length);
 			}
 		}
 
 		/// <summary>
 		/// Determines if the given line is empty
 		/// </summary>
-		/// <param name="Input">The input data</param>
+		/// <param name="input">The input data</param>
 		/// <returns>True if the given text is empty</returns>
-		public static bool IsEmptyOrWhitespace(ReadOnlySpan<byte> Input)
+		public static bool IsEmptyOrWhitespace(ReadOnlySpan<byte> input)
 		{
-			for(int Idx = 0; Idx < Input.Length; Idx++)
+			for(int idx = 0; idx < input.Length; idx++)
 			{
-				byte Byte = Input[Idx];
-				if(Byte != (byte)'\n' && Byte != '\r' && Byte != ' ')
+				byte v = input[idx];
+				if(v != (byte)'\n' && v != '\r' && v != ' ')
 				{
 					return false;
 				}
@@ -282,85 +263,85 @@ namespace Horde.Build.Logs
 		/// <summary>
 		/// Converts a JSON log line to plain text
 		/// </summary>
-		/// <param name="Input">The JSON data</param>
-		/// <param name="Output">Output buffer for the converted line</param>
-		/// <param name="OutputOffset">Offset within the buffer to write the converted data</param>
+		/// <param name="input">The JSON data</param>
+		/// <param name="output">Output buffer for the converted line</param>
+		/// <param name="outputOffset">Offset within the buffer to write the converted data</param>
 		/// <returns></returns>
-		public static int ConvertToPlainText(ReadOnlySpan<byte> Input, byte[] Output, int OutputOffset)
+		public static int ConvertToPlainText(ReadOnlySpan<byte> input, byte[] output, int outputOffset)
 		{
-			if(IsEmptyOrWhitespace(Input))
+			if(IsEmptyOrWhitespace(input))
 			{
-				Output[OutputOffset] = (byte)'\n';
-				return OutputOffset + 1;
+				output[outputOffset] = (byte)'\n';
+				return outputOffset + 1;
 			}
 
-			Utf8JsonReader Reader = new Utf8JsonReader(Input);
-			if (Reader.Read() && Reader.TokenType == JsonTokenType.StartObject)
+			Utf8JsonReader reader = new Utf8JsonReader(input);
+			if (reader.Read() && reader.TokenType == JsonTokenType.StartObject)
 			{
-				while (Reader.Read() && Reader.TokenType == JsonTokenType.PropertyName)
+				while (reader.Read() && reader.TokenType == JsonTokenType.PropertyName)
 				{
-					if (!Reader.ValueTextEquals("message"))
+					if (!reader.ValueTextEquals("message"))
 					{
-						Reader.Skip();
+						reader.Skip();
 						continue;
 					}
-					if (!Reader.Read() || Reader.TokenType != JsonTokenType.String)
+					if (!reader.Read() || reader.TokenType != JsonTokenType.String)
 					{
-						Reader.Skip();
+						reader.Skip();
 						continue;
 					}
 
-					int UnescapedLength = UnescapeUtf8(Reader.ValueSpan, Output.AsSpan(OutputOffset));
-					OutputOffset += UnescapedLength;
+					int unescapedLength = UnescapeUtf8(reader.ValueSpan, output.AsSpan(outputOffset));
+					outputOffset += unescapedLength;
 
-					Output[OutputOffset] = (byte)'\n';
-					OutputOffset++;
+					output[outputOffset] = (byte)'\n';
+					outputOffset++;
 
 					break;
 				}
 			}
-			return OutputOffset;
+			return outputOffset;
 		}
 
 		/// <summary>
 		/// Unescape a json utf8 string
 		/// </summary>
-		/// <param name="Source">Source span of bytes</param>
-		/// <param name="Target">Target span of bytes</param>
+		/// <param name="source">Source span of bytes</param>
+		/// <param name="target">Target span of bytes</param>
 		/// <returns>Length of the converted data</returns>
-		static int UnescapeUtf8(ReadOnlySpan<byte> Source, Span<byte> Target)
+		static int UnescapeUtf8(ReadOnlySpan<byte> source, Span<byte> target)
 		{
-			int Length = 0;
+			int length = 0;
 			for (; ; )
 			{
 				// Copy up to the next backslash
-				int Backslash = Source.IndexOf((byte)'\\');
-				if (Backslash == -1)
+				int backslash = source.IndexOf((byte)'\\');
+				if (backslash == -1)
 				{
-					Source.CopyTo(Target);
-					Length += Source.Length;
+					source.CopyTo(target);
+					length += source.Length;
 					break;
 				}
-				else if (Backslash > 0)
+				else if (backslash > 0)
 				{
-					Source.Slice(0, Backslash).CopyTo(Target);
-					Source = Source.Slice(Backslash);
-					Target = Target.Slice(Backslash);
-					Length += Backslash;
+					source.Slice(0, backslash).CopyTo(target);
+					source = source.Slice(backslash);
+					target = target.Slice(backslash);
+					length += backslash;
 				}
 
 				// Check what the escape code is
-				if (Source[1] == 'u')
+				if (source[1] == 'u')
 				{
-					char[] Chars = { (char)((StringUtils.ParseHexByte(Source, 2) << 8) | StringUtils.ParseHexByte(Source, 4)) };
-					int EncodedLength = Encoding.UTF8.GetBytes(Chars.AsSpan(), Target);
-					Source = Source.Slice(6);
-					Target = Target.Slice(EncodedLength);
-					Length += EncodedLength;
+					char[] chars = { (char)((StringUtils.ParseHexByte(source, 2) << 8) | StringUtils.ParseHexByte(source, 4)) };
+					int encodedLength = Encoding.UTF8.GetBytes(chars.AsSpan(), target);
+					source = source.Slice(6);
+					target = target.Slice(encodedLength);
+					length += encodedLength;
 				}
 				else
 				{
-					Target[0] = Source[1] switch
+					target[0] = source[1] switch
 					{
 						(byte)'\"' => (byte)'\"',
 						(byte)'\\' => (byte)'\\',
@@ -369,55 +350,55 @@ namespace Horde.Build.Logs
 						(byte)'n' => (byte)'\n',
 						(byte)'r' => (byte)'\r',
 						(byte)'t' => (byte)'\t',
-						_ => Source[1]
+						_ => source[1]
 					};
-					Source = Source.Slice(2);
-					Target = Target.Slice(1);
-					Length++;
+					source = source.Slice(2);
+					target = target.Slice(1);
+					length++;
 				}
 			}
-			return Length;
+			return length;
 		}
 
 		/// <summary>
 		/// Generates placeholder data for a missing span
 		/// </summary>
-		/// <param name="ChunkIdx">Index of the chunk</param>
-		/// <param name="HostName">Host name of the machine that was holding the data</param>
-		/// <param name="TargetLength">Desired length of the buffer</param>
-		/// <param name="TargetLineCount">Desired line count for the buffer</param>
-		public void AppendMissingDataInfo(int ChunkIdx, string? HostName, int TargetLength, int TargetLineCount)
+		/// <param name="chunkIdx">Index of the chunk</param>
+		/// <param name="hostName">Host name of the machine that was holding the data</param>
+		/// <param name="targetLength">Desired length of the buffer</param>
+		/// <param name="targetLineCount">Desired line count for the buffer</param>
+		public void AppendMissingDataInfo(int chunkIdx, string? hostName, int targetLength, int targetLineCount)
 		{
-			int NewLength = Length;
+			int newLength = Length;
 
-			if(InternalData.Length < TargetLength)
+			if(_internalData.Length < targetLength)
 			{
-				Array.Resize(ref InternalData, TargetLength);
+				Array.Resize(ref _internalData, targetLength);
 			}
 
-			string Suffix = $" (server: {HostName})";
-			for(int NewLineCount = LineCount; NewLineCount < TargetLineCount; NewLineCount++)
+			string suffix = $" (server: {hostName})";
+			for(int newLineCount = LineCount; newLineCount < targetLineCount; newLineCount++)
 			{
-				byte[] ErrorBytes = Encoding.ASCII.GetBytes($"{{ \"level\":\"Error\",\"message\":\"[Missing data at chunk {ChunkIdx}, line {NewLineCount}{Suffix}]\" }}\n");
-				if (NewLength + ErrorBytes.Length > TargetLength)
+				byte[] errorBytes = Encoding.ASCII.GetBytes($"{{ \"level\":\"Error\",\"message\":\"[Missing data at chunk {chunkIdx}, line {newLineCount}{suffix}]\" }}\n");
+				if (newLength + errorBytes.Length > targetLength)
 				{
 					break;
 				}
 
-				ErrorBytes.AsSpan().CopyTo(InternalData.AsSpan(NewLength));
-				NewLength += ErrorBytes.Length;
+				errorBytes.AsSpan().CopyTo(_internalData.AsSpan(newLength));
+				newLength += errorBytes.Length;
 
-				Suffix = String.Empty;
+				suffix = String.Empty;
 			}
 
-			if (NewLength < TargetLength)
+			if (newLength < targetLength)
 			{
-				NewLength = Math.Max(NewLength - 1, 0);
-				InternalData.AsSpan(NewLength, (int)((TargetLength - 1) - NewLength)).Fill((byte)' ');
-				InternalData[TargetLength - 1] = (byte)'\n';
+				newLength = Math.Max(newLength - 1, 0);
+				_internalData.AsSpan(newLength, (int)((targetLength - 1) - newLength)).Fill((byte)' ');
+				_internalData[targetLength - 1] = (byte)'\n';
 			}
 
-			UpdateLength((int)TargetLength);
+			UpdateLength((int)targetLength);
 		}
 
 		/// <summary>
@@ -425,38 +406,38 @@ namespace Horde.Build.Logs
 		/// </summary>
 		public void Shrink()
 		{
-			if (Length < InternalData.Length)
+			if (Length < _internalData.Length)
 			{
-				Array.Resize(ref InternalData, Length);
+				Array.Resize(ref _internalData, Length);
 			}
 		}
 
 		/// <summary>
 		/// Updates the length of this chunk, computing all the newline offsets
 		/// </summary>
-		/// <param name="NewLength">New length of the chunk</param>
-		private void UpdateLength(int NewLength)
+		/// <param name="newLength">New length of the chunk</param>
+		private void UpdateLength(int newLength)
 		{
-			if (NewLength > Length)
+			if (newLength > Length)
 			{
 				// Make sure the data ends with a newline
-				if (InternalData[NewLength - 1] != '\n')
+				if (_internalData[newLength - 1] != '\n')
 				{
 					throw new InvalidDataException("Chunk data must end with a newline");
 				}
 
 				// Calculate the new number of newlines
-				ReadOnlySpan<byte> NewData = InternalData.AsSpan(0, NewLength);
-				for(int Idx = Length; Idx < NewLength; Idx++)
+				ReadOnlySpan<byte> newData = _internalData.AsSpan(0, newLength);
+				for(int idx = Length; idx < newLength; idx++)
 				{
-					if(NewData[Idx] == '\n')
+					if(newData[idx] == '\n')
 					{
-						InternalLineOffsets.Add(Idx + 1);
+						_internalLineOffsets.Add(idx + 1);
 					}
 				}
 
 				// Update the size of the buffer
-				Length = NewLength;
+				Length = newLength;
 			}
 		}
 	}
@@ -469,29 +450,29 @@ namespace Horde.Build.Logs
 		/// <summary>
 		/// Find the line index for a particular offset
 		/// </summary>
-		/// <param name="LogText">The text to search</param>
-		/// <param name="Offset">Offset within the text</param>
+		/// <param name="logText">The text to search</param>
+		/// <param name="offset">Offset within the text</param>
 		/// <returns>The line index</returns>
-		public static int GetLineIndexForOffset(this ILogText LogText, int Offset)
+		public static int GetLineIndexForOffset(this ILogText logText, int offset)
 		{
-			int LineIdx = LogText.LineOffsets.BinarySearch(Offset);
-			if(LineIdx < 0)
+			int lineIdx = logText.LineOffsets.BinarySearch(offset);
+			if(lineIdx < 0)
 			{
-				LineIdx = ~LineIdx - 1;
+				lineIdx = ~lineIdx - 1;
 			}
-			return LineIdx;
+			return lineIdx;
 		}
 
 		/// <summary>
 		/// Converts a log text instance to plain text
 		/// </summary>
-		/// <param name="LogText">The text to convert</param>
+		/// <param name="logText">The text to convert</param>
 		/// <returns>The plain text instance</returns>
-		public static ILogText ToPlainText(this ILogText LogText)
+		public static ILogText ToPlainText(this ILogText logText)
 		{
-			LogText Other = new LogText();
-			Other.AppendPlainText(LogText, 0, LogText.LineCount);
-			return Other;
+			LogText other = new LogText();
+			other.AppendPlainText(logText, 0, logText.LineCount);
+			return other;
 		}
 	}
 }
diff --git a/Engine/Source/Programs/Horde/Horde.Build/Logs/Data/LogToken.cs b/Engine/Source/Programs/Horde/Horde.Build/Logs/Data/LogToken.cs
--- a/Engine/Source/Programs/Horde/Horde.Build/Logs/Data/LogToken.cs
+++ b/Engine/Source/Programs/Horde/Horde.Build/Logs/Data/LogToken.cs
@@ -2,8 +2,6 @@
 
 using System;
 using System.Collections.Generic;
-using System.Linq;
-using System.Threading.Tasks;
 
 namespace Horde.Build.Logs
 {
@@ -25,95 +23,95 @@ namespace Horde.Build.Logs
 		/// <summary>
 		/// Lookup from input byte to token type
 		/// </summary>
-		static readonly byte[] TokenTypes = GetTokenTypes();
+		static readonly byte[] s_tokenTypes = GetTokenTypes();
 
 		/// <summary>
 		/// Lookup from input byte to token char
 		/// </summary>
-		static readonly byte[] TokenChars = GetTokenChars();
+		static readonly byte[] s_tokenChars = GetTokenChars();
 
 		/// <summary>
 		/// Get a set of tokens in a trie
 		/// </summary>
-		/// <param name="Text">Text to scan</param>
+		/// <param name="text">Text to scan</param>
 		/// <returns></returns>
-		public static ReadOnlyTrie GetTrie(ReadOnlySpan<byte> Text)
+		public static ReadOnlyTrie GetTrie(ReadOnlySpan<byte> text)
 		{
-			ReadOnlyTrieBuilder Builder = new ReadOnlyTrieBuilder();
-			GetTokens(Text, Builder.Add);
-			return Builder.Build();
+			ReadOnlyTrieBuilder builder = new ReadOnlyTrieBuilder();
+			GetTokens(text, builder.Add);
+			return builder.Build();
 		}
 
 		/// <summary>
 		/// Gets a single token
 		/// </summary>
-		/// <param name="Text">The text to parse</param>
+		/// <param name="text">The text to parse</param>
 		/// <returns>The token value</returns>
-		public static ulong GetToken(ReadOnlySpan<byte> Text)
+		public static ulong GetToken(ReadOnlySpan<byte> text)
 		{
-			ulong Token = 0;
-			GetTokens(Text, x => Token = x);
-			return Token;
+			ulong token = 0;
+			GetTokens(text, x => token = x);
+			return token;
 		}
 
 		/// <summary>
 		/// Get a set of tokens in a trie
 		/// </summary>
-		/// <param name="Text">Text to scan</param>
+		/// <param name="text">Text to scan</param>
 		/// <returns></returns>
-		public static HashSet<ulong> GetTokenSet(ReadOnlySpan<byte> Text)
+		public static HashSet<ulong> GetTokenSet(ReadOnlySpan<byte> text)
 		{
-			HashSet<ulong> Tokens = new HashSet<ulong>();
-			GetTokens(Text, x => Tokens.Add(x));
-			return Tokens;
+			HashSet<ulong> tokens = new HashSet<ulong>();
+			GetTokens(text, x => tokens.Add(x));
+			return tokens;
 		}
 
 		/// <summary>
 		/// Decompose a span of text into tokens
 		/// </summary>
-		/// <param name="Text">Text to scan</param>
-		/// <param name="AddToken">Receives a set of tokens</param>
-		public static void GetTokens(ReadOnlySpan<byte> Text, Action<ulong> AddToken)
+		/// <param name="text">Text to scan</param>
+		/// <param name="addToken">Receives a set of tokens</param>
+		public static void GetTokens(ReadOnlySpan<byte> text, Action<ulong> addToken)
 		{
-			if (Text.Length > 0)
+			if (text.Length > 0)
 			{
-				int Type = TokenTypes[Text[0]];
-				int NumTokenBits = 8;
-				ulong Token = TokenChars[Text[0]];
+				int type = s_tokenTypes[text[0]];
+				int numTokenBits = 8;
+				ulong token = s_tokenChars[text[0]];
 
-				for (int TextIdx = 1; TextIdx < Text.Length; TextIdx++)
+				for (int textIdx = 1; textIdx < text.Length; textIdx++)
 				{
-					byte NextChar = TokenChars[Text[TextIdx]];
-					int NextType = TokenTypes[NextChar];
-					if (Type != NextType || NumTokenBits + 8 > MaxTokenBits)
+					byte nextChar = s_tokenChars[text[textIdx]];
+					int nextType = s_tokenTypes[nextChar];
+					if (type != nextType || numTokenBits + 8 > MaxTokenBits)
 					{
-						AddToken(Token << (MaxTokenBits - NumTokenBits));
-						Token = 0;
-						NumTokenBits = 0;
-						Type = NextType;
+						addToken(token << (MaxTokenBits - numTokenBits));
+						token = 0;
+						numTokenBits = 0;
+						type = nextType;
 					}
-					Token = (Token << 8) | NextChar;
-					NumTokenBits += 8;
+					token = (token << 8) | nextChar;
+					numTokenBits += 8;
 				}
 
-				AddToken(Token << (MaxTokenBits - NumTokenBits));
+				addToken(token << (MaxTokenBits - numTokenBits));
 			}
 		}
 
 		/// <summary>
 		/// Gets the length of the first token in the given span
 		/// </summary>
-		/// <param name="Text">The text to search</param>
-		/// <param name="Pos">Start position for the search</param>
+		/// <param name="text">The text to search</param>
+		/// <param name="pos">Start position for the search</param>
 		/// <returns>Length of the first token</returns>
-		public static ReadOnlySpan<byte> GetTokenText(ReadOnlySpan<byte> Text, int Pos)
+		public static ReadOnlySpan<byte> GetTokenText(ReadOnlySpan<byte> text, int pos)
 		{
-			int Type = TokenTypes[Text[Pos]];
-			for (int End = Pos + 1; ; End++)
+			int type = s_tokenTypes[text[pos]];
+			for (int end = pos + 1; ; end++)
 			{
-				if (End == Text.Length || TokenTypes[Text[End]] != Type)
+				if (end == text.Length || s_tokenTypes[text[end]] != type)
 				{
-					return Text.Slice(Pos, End - Pos);
+					return text.Slice(pos, end - pos);
 				}
 			}
 		}
@@ -121,44 +119,44 @@ namespace Horde.Build.Logs
 		/// <summary>
 		/// Gets the length of the first token in the given span
 		/// </summary>
-		/// <param name="Text">The text to search</param>
-		/// <param name="Offset">Offset of the window to read from the token</param>
+		/// <param name="text">The text to search</param>
+		/// <param name="offset">Offset of the window to read from the token</param>
 		/// <returns>Length of the first token</returns>
-		public static ulong GetWindowedTokenValue(ReadOnlySpan<byte> Text, int Offset)
+		public static ulong GetWindowedTokenValue(ReadOnlySpan<byte> text, int offset)
 		{
-			ulong Token = 0;
-			for (int Idx = 0; Idx < MaxTokenBytes; Idx++)
+			ulong token = 0;
+			for (int idx = 0; idx < MaxTokenBytes; idx++)
 			{
-				Token <<= 8;
-				if (Offset >= 0 && Offset < Text.Length)
+				token <<= 8;
+				if (offset >= 0 && offset < text.Length)
 				{
-					Token |= TokenChars[Text[Offset]];
+					token |= s_tokenChars[text[offset]];
 				}
-				Offset++;
+				offset++;
 			}
-			return Token;
+			return token;
 		}
 
 		/// <summary>
 		/// Gets the length of the first token in the given span
 		/// </summary>
-		/// <param name="Text">The text to search</param>
-		/// <param name="Offset">Offset of the window to read from the token</param>
+		/// <param name="text">The text to search</param>
+		/// <param name="offset">Offset of the window to read from the token</param>
 		/// <param name="bAllowPartialMatch">Whether to allow only matching the start of the string</param>
 		/// <returns>Length of the first token</returns>
-		public static ulong GetWindowedTokenMask(ReadOnlySpan<byte> Text, int Offset, bool bAllowPartialMatch)
+		public static ulong GetWindowedTokenMask(ReadOnlySpan<byte> text, int offset, bool bAllowPartialMatch)
 		{
-			ulong Token = 0;
-			for (int Idx = 0; Idx < MaxTokenBytes; Idx++)
+			ulong token = 0;
+			for (int idx = 0; idx < MaxTokenBytes; idx++)
 			{
-				Token <<= 8;
-				if (Offset >= 0 && (Offset < Text.Length || !bAllowPartialMatch))
+				token <<= 8;
+				if (offset >= 0 && (offset < text.Length || !bAllowPartialMatch))
 				{
-					Token |= 0xff;
+					token |= 0xff;
 				}
-				Offset++;
+				offset++;
 			}
-			return Token;
+			return token;
 		}
 
 		/// <summary>
@@ -167,23 +165,23 @@ namespace Horde.Build.Logs
 		/// <returns>Array whose elements map from an input byte to token type</returns>
 		static byte[] GetTokenTypes()
 		{
-			byte[] CharTypes = new byte[256];
-			for (int Idx = 'a'; Idx <= 'z'; Idx++)
+			byte[] charTypes = new byte[256];
+			for (int idx = 'a'; idx <= 'z'; idx++)
 			{
-				CharTypes[Idx] = 1;
+				charTypes[idx] = 1;
 			}
-			for (int Idx = 'A'; Idx <= 'Z'; Idx++)
+			for (int idx = 'A'; idx <= 'Z'; idx++)
 			{
-				CharTypes[Idx] = 1;
+				charTypes[idx] = 1;
 			}
-			for (int Idx = '0'; Idx <= '9'; Idx++)
+			for (int idx = '0'; idx <= '9'; idx++)
 			{
-				CharTypes[Idx] = 2;
+				charTypes[idx] = 2;
 			}
-			CharTypes[' '] = 3;
-			CharTypes['\t'] = 3;
-			CharTypes['\n'] = 4;
-			return CharTypes;
+			charTypes[' '] = 3;
+			charTypes['\t'] = 3;
+			charTypes['\n'] = 4;
+			return charTypes;
 		}
 
 		/// <summary>
@@ -192,16 +190,16 @@ namespace Horde.Build.Logs
 		/// <returns>Array whose elements map from an input byte to token type</returns>
 		static byte[] GetTokenChars()
 		{
-			byte[] Chars = new byte[256];
-			for(int Idx = 0; Idx < 256; Idx++)
+			byte[] chars = new byte[256];
+			for(int idx = 0; idx < 256; idx++)
 			{
-				Chars[Idx] = (byte)Idx;
+				chars[idx] = (byte)idx;
 			}
-			for (int Idx = 'A'; Idx <= 'Z'; Idx++)
+			for (int idx = 'A'; idx <= 'Z'; idx++)
 			{
-				Chars[Idx] = (byte)('a' + Idx - 'A');
+				chars[idx] = (byte)('a' + idx - 'A');
 			}
-			return Chars;
+			return chars;
 		}
 	}
 }
diff --git a/Engine/Source/Programs/Horde/Horde.Build/Logs/ILogEvent.cs b/Engine/Source/Programs/Horde/Horde.Build/Logs/ILogEvent.cs
--- a/Engine/Source/Programs/Horde/Horde.Build/Logs/ILogEvent.cs
+++ b/Engine/Source/Programs/Horde/Horde.Build/Logs/ILogEvent.cs
@@ -1,16 +1,11 @@
 // Copyright Epic Games, Inc. All Rights Reserved.
 
-using HordeCommon;
-using HordeCommon.Rpc;
+using System.Collections.Generic;
+using System.Text.Json;
 using Horde.Build.Utilities;
+using HordeCommon;
 using Microsoft.Extensions.Logging;
 using MongoDB.Bson;
-using System;
-using System.Collections.Generic;
-using System.Diagnostics.CodeAnalysis;
-using System.Linq;
-using System.Text.Json;
-using System.Threading.Tasks;
 
 namespace Horde.Build.Models
 {
diff --git a/Engine/Source/Programs/Horde/Horde.Build/Logs/ILogEventCollection.cs b/Engine/Source/Programs/Horde/Horde.Build/Logs/ILogEventCollection.cs
--- a/Engine/Source/Programs/Horde/Horde.Build/Logs/ILogEventCollection.cs
+++ b/Engine/Source/Programs/Horde/Horde.Build/Logs/ILogEventCollection.cs
@@ -1,19 +1,10 @@
 // Copyright Epic Games, Inc. All Rights Reserved.
 
-using HordeCommon;
+using System.Collections.Generic;
+using System.Threading.Tasks;
 using Horde.Build.Models;
-using Horde.Build.Services;
 using Horde.Build.Utilities;
-using Microsoft.Extensions.Logging;
 using MongoDB.Bson;
-using MongoDB.Bson.Serialization;
-using MongoDB.Bson.Serialization.Attributes;
-using MongoDB.Driver;
-using Serilog.Events;
-using System;
-using System.Collections.Generic;
-using System.Linq;
-using System.Threading.Tasks;
 
 namespace Horde.Build.Collections
 {
@@ -27,47 +18,47 @@ namespace Horde.Build.Collections
 		/// <summary>
 		/// Creates a new event
 		/// </summary>
-		/// <param name="NewEvent">The new event to vreate</param>
-		Task AddAsync(NewLogEventData NewEvent);
+		/// <param name="newEvent">The new event to vreate</param>
+		Task AddAsync(NewLogEventData newEvent);
 
 		/// <summary>
 		/// Creates a new event
 		/// </summary>
-		/// <param name="NewEvents">List of events to create</param>
-		Task AddManyAsync(List<NewLogEventData> NewEvents);
+		/// <param name="newEvents">List of events to create</param>
+		Task AddManyAsync(List<NewLogEventData> newEvents);
 
 		/// <summary>
 		/// Finds events within a log file
 		/// </summary>
-		/// <param name="LogId">Unique id of the log containing this event</param>
-		/// <param name="Index">Start index within the matching results</param>
-		/// <param name="Count">Maximum number of results to return</param>
+		/// <param name="logId">Unique id of the log containing this event</param>
+		/// <param name="index">Start index within the matching results</param>
+		/// <param name="count">Maximum number of results to return</param>
 		/// <returns>List of events matching the query</returns>
-		Task<List<ILogEvent>> FindAsync(LogId LogId, int? Index = null, int? Count = null);
+		Task<List<ILogEvent>> FindAsync(LogId logId, int? index = null, int? count = null);
 
 		/// <summary>
 		/// Finds a list of events for a set of spans
 		/// </summary>
-		/// <param name="SpanIds">The span ids</param>
-		/// <param name="LogIds">List of log ids to query</param>
-		/// <param name="Index">Index of the first result to return</param>
-		/// <param name="Count">Number of results to return</param>
+		/// <param name="spanIds">The span ids</param>
+		/// <param name="logIds">List of log ids to query</param>
+		/// <param name="index">Index of the first result to return</param>
+		/// <param name="count">Number of results to return</param>
 		/// <returns>List of events for this issue</returns>
-		Task<List<ILogEvent>> FindEventsForSpansAsync(IEnumerable<ObjectId> SpanIds, LogId[]? LogIds = null, int Index = 0, int Count = 10);
+		Task<List<ILogEvent>> FindEventsForSpansAsync(IEnumerable<ObjectId> spanIds, LogId[]? logIds = null, int index = 0, int count = 10);
 
 		/// <summary>
 		/// Delete all the events for a log file
 		/// </summary>
-		/// <param name="LogId">Unique id of the log</param>
+		/// <param name="logId">Unique id of the log</param>
 		/// <returns>Async task</returns>
-		Task DeleteLogAsync(LogId LogId);
+		Task DeleteLogAsync(LogId logId);
 
 		/// <summary>
 		/// Update the span for an event
 		/// </summary>
-		/// <param name="Events">The events to update</param>
-		/// <param name="SpanId">New span id</param>
+		/// <param name="events">The events to update</param>
+		/// <param name="spanId">New span id</param>
 		/// <returns>Async task</returns>
-		Task AddSpanToEventsAsync(IEnumerable<ILogEvent> Events, ObjectId SpanId);
+		Task AddSpanToEventsAsync(IEnumerable<ILogEvent> events, ObjectId spanId);
 	}
 }
diff --git a/Engine/Source/Programs/Horde/Horde.Build/Logs/ILogFile.cs b/Engine/Source/Programs/Horde/Horde.Build/Logs/ILogFile.cs
--- a/Engine/Source/Programs/Horde/Horde.Build/Logs/ILogFile.cs
+++ b/Engine/Source/Programs/Horde/Horde.Build/Logs/ILogFile.cs
@@ -1,16 +1,10 @@
 // Copyright Epic Games, Inc. All Rights Reserved.
 
+using System;
+using System.Collections.Generic;
 using EpicGames.Core;
 using Horde.Build.Api;
-using Horde.Build.Services;
 using Horde.Build.Utilities;
-using MongoDB.Bson;
-using MongoDB.Bson.Serialization.Attributes;
-using System;
-using System.Collections.Generic;
-using System.Diagnostics.CodeAnalysis;
-using System.Linq;
-using System.Threading.Tasks;
 
 namespace Horde.Build.Models
 {
@@ -93,33 +87,33 @@ namespace Horde.Build.Models
 		/// <summary>
 		/// Gets the chunk index containing the given offset.
 		/// </summary>
-		/// <param name="Chunks">The chunks to search</param>
-		/// <param name="Offset">The offset to search for</param>
+		/// <param name="chunks">The chunks to search</param>
+		/// <param name="offset">The offset to search for</param>
 		/// <returns>The chunk index containing the given offset</returns>
-		public static int GetChunkForOffset(this IReadOnlyList<ILogChunk> Chunks, long Offset)
+		public static int GetChunkForOffset(this IReadOnlyList<ILogChunk> chunks, long offset)
 		{
-			int ChunkIndex = Chunks.BinarySearch(x => x.Offset, Offset);
-			if (ChunkIndex < 0)
+			int chunkIndex = chunks.BinarySearch(x => x.Offset, offset);
+			if (chunkIndex < 0)
 			{
-				ChunkIndex = ~ChunkIndex - 1;
+				chunkIndex = ~chunkIndex - 1;
 			}
-			return ChunkIndex;
+			return chunkIndex;
 		}
 
 		/// <summary>
 		/// Gets the starting chunk index for the given line
 		/// </summary>
-		/// <param name="Chunks">The chunks to search</param>
-		/// <param name="LineIndex">Index of the line to query</param>
+		/// <param name="chunks">The chunks to search</param>
+		/// <param name="lineIndex">Index of the line to query</param>
 		/// <returns>Index of the chunk to fetch</returns>
-		public static int GetChunkForLine(this IReadOnlyList<ILogChunk> Chunks, int LineIndex)
+		public static int GetChunkForLine(this IReadOnlyList<ILogChunk> chunks, int lineIndex)
 		{
-			int ChunkIndex = Chunks.BinarySearch(x => x.LineIndex, LineIndex);
-			if(ChunkIndex < 0)
+			int chunkIndex = chunks.BinarySearch(x => x.LineIndex, lineIndex);
+			if(chunkIndex < 0)
 			{
-				ChunkIndex = ~ChunkIndex - 1;
+				chunkIndex = ~chunkIndex - 1;
 			}
-			return ChunkIndex;
+			return chunkIndex;
 		}
 	}
 }
diff --git a/Engine/Source/Programs/Horde/Horde.Build/Logs/ILogFileCollection.cs b/Engine/Source/Programs/Horde/Horde.Build/Logs/ILogFileCollection.cs
--- a/Engine/Source/Programs/Horde/Horde.Build/Logs/ILogFileCollection.cs
+++ b/Engine/Source/Programs/Horde/Horde.Build/Logs/ILogFileCollection.cs
@@ -1,17 +1,10 @@
 // Copyright Epic Games, Inc. All Rights Reserved.
 
+using System.Collections.Generic;
+using System.Threading.Tasks;
 using Horde.Build.Api;
 using Horde.Build.Models;
-using Horde.Build.Services;
 using Horde.Build.Utilities;
-using MongoDB.Bson;
-using MongoDB.Bson.Serialization.Attributes;
-using MongoDB.Driver;
-using System;
-using System.Collections.Generic;
-using System.Linq;
-using System.Net;
-using System.Threading.Tasks;
 
 namespace Horde.Build.Collections
 {
@@ -42,14 +35,14 @@ namespace Horde.Build.Collections
 		/// <summary>
 		/// Constructor
 		/// </summary>
-		/// <param name="Index">Index of the chunk</param>
-		/// <param name="Length">New length for the chunk</param>
-		/// <param name="LineCount">Number of lines in the chunk</param>
-		public CompleteLogChunkUpdate(int Index, int Length, int LineCount)
+		/// <param name="index">Index of the chunk</param>
+		/// <param name="length">New length for the chunk</param>
+		/// <param name="lineCount">Number of lines in the chunk</param>
+		public CompleteLogChunkUpdate(int index, int length, int lineCount)
 		{
-			this.Index = Index;
-			this.Length = Length;
-			this.LineCount = LineCount;
+			Index = index;
+			Length = length;
+			LineCount = lineCount;
 		}
 	}
 
@@ -61,48 +54,48 @@ namespace Horde.Build.Collections
 		/// <summary>
 		/// Creates a new log
 		/// </summary>
-		/// <param name="JobId">Unique id of the job that owns this log file</param>
-		/// <param name="SessionId">Agent session allowed to update the log</param>
-		/// <param name="Type">Type of events to be stored in the log</param>
+		/// <param name="jobId">Unique id of the job that owns this log file</param>
+		/// <param name="sessionId">Agent session allowed to update the log</param>
+		/// <param name="type">Type of events to be stored in the log</param>
 		/// <returns>The new log file document</returns>
-		Task<ILogFile> CreateLogFileAsync(JobId JobId, SessionId? SessionId, LogType Type);
+		Task<ILogFile> CreateLogFileAsync(JobId jobId, SessionId? sessionId, LogType type);
 
 		/// <summary>
 		/// Adds a new chunk
 		/// </summary>
-		/// <param name="LogFileInterface">The current log file</param>
-		/// <param name="Offset">Offset of the new chunk</param>
-		/// <param name="LineIndex">Line index for the start of the chunk</param>
+		/// <param name="logFileInterface">The current log file</param>
+		/// <param name="offset">Offset of the new chunk</param>
+		/// <param name="lineIndex">Line index for the start of the chunk</param>
 		/// <returns>The updated log file document</returns>
-		Task<ILogFile?> TryAddChunkAsync(ILogFile LogFileInterface, long Offset, int LineIndex);
+		Task<ILogFile?> TryAddChunkAsync(ILogFile logFileInterface, long offset, int lineIndex);
 
 		/// <summary>
 		/// Update the log file with final information about certain chunks
 		/// </summary>
-		/// <param name="LogFileInterface">The current log file</param>
-		/// <param name="Chunks">Chunks to update. New chunks will be inserted</param>
+		/// <param name="logFileInterface">The current log file</param>
+		/// <param name="chunks">Chunks to update. New chunks will be inserted</param>
 		/// <returns>The updated log file document</returns>
-		Task<ILogFile?> TryCompleteChunksAsync(ILogFile LogFileInterface, IEnumerable<CompleteLogChunkUpdate> Chunks);
+		Task<ILogFile?> TryCompleteChunksAsync(ILogFile logFileInterface, IEnumerable<CompleteLogChunkUpdate> chunks);
 
 		/// <summary>
 		/// Update the log file with final information about the index
 		/// </summary>
-		/// <param name="LogFileInterface">The current log file</param>
-		/// <param name="NewIndexLength">New length of the index</param>
+		/// <param name="logFileInterface">The current log file</param>
+		/// <param name="newIndexLength">New length of the index</param>
 		/// <returns>The updated log file document</returns>
-		Task<ILogFile?> TryUpdateIndexAsync(ILogFile LogFileInterface, long NewIndexLength);
+		Task<ILogFile?> TryUpdateIndexAsync(ILogFile logFileInterface, long newIndexLength);
 
 		/// <summary>
 		/// Gets a logfile by ID
 		/// </summary>
-		/// <param name="LogFileId">Unique id of the log file</param>
+		/// <param name="logFileId">Unique id of the log file</param>
 		/// <returns>The logfile document</returns>
-		Task<ILogFile?> GetLogFileAsync(LogId LogFileId);
+		Task<ILogFile?> GetLogFileAsync(LogId logFileId);
 
 		/// <summary>
 		/// Gets all the log files
 		/// </summary>
 		/// <returns>List of log files</returns>
-		Task<List<ILogFile>> GetLogFilesAsync(int? Index = null, int? Count = null);
+		Task<List<ILogFile>> GetLogFilesAsync(int? index = null, int? count = null);
 	}
 }
diff --git a/Engine/Source/Programs/Horde/Horde.Build/Logs/Impl/LogEventCollection.cs b/Engine/Source/Programs/Horde/Horde.Build/Logs/Impl/LogEventCollection.cs
--- a/Engine/Source/Programs/Horde/Horde.Build/Logs/Impl/LogEventCollection.cs
+++ b/Engine/Source/Programs/Horde/Horde.Build/Logs/Impl/LogEventCollection.cs
@@ -1,19 +1,17 @@
 // Copyright Epic Games, Inc. All Rights Reserved.
 
-using HordeCommon;
+using System;
+using System.Collections.Generic;
+using System.Linq;
+using System.Threading.Tasks;
 using Horde.Build.Models;
 using Horde.Build.Services;
 using Horde.Build.Utilities;
+using HordeCommon;
 using Microsoft.Extensions.Logging;
 using MongoDB.Bson;
-using MongoDB.Bson.Serialization;
 using MongoDB.Bson.Serialization.Attributes;
 using MongoDB.Driver;
-using Serilog.Events;
-using System;
-using System.Collections.Generic;
-using System.Linq;
-using System.Threading.Tasks;
 
 namespace Horde.Build.Collections.Impl
 {
@@ -54,19 +52,19 @@ namespace Horde.Build.Collections.Impl
 
 			public LogEventDocument()
 			{
-				this.Id = new LogEventId();
+				Id = new LogEventId();
 			}
 
-			public LogEventDocument(LogId LogId, EventSeverity Severity, int LineIndex, int LineCount, ObjectId? SpanId)
+			public LogEventDocument(LogId logId, EventSeverity severity, int lineIndex, int lineCount, ObjectId? spanId)
 			{
-				this.Id = new LogEventId { LogId = LogId, LineIndex = LineIndex };
-				this.IsWarning = Severity == EventSeverity.Warning;
-				this.LineCount = (LineCount > 1) ? (int?)LineCount : null;
-				this.SpanId = SpanId;
+				Id = new LogEventId { LogId = logId, LineIndex = lineIndex };
+				IsWarning = severity == EventSeverity.Warning;
+				LineCount = (lineCount > 1) ? (int?)lineCount : null;
+				SpanId = spanId;
 			}
 
-			public LogEventDocument(NewLogEventData Data)
-				: this(Data.LogId, Data.Severity, Data.LineIndex, Data.LineCount, Data.SpanId)
+			public LogEventDocument(NewLogEventData data)
+				: this(data.LogId, data.Severity, data.LineIndex, data.LineCount, data.SpanId)
 			{
 			}
 		}
@@ -93,112 +91,112 @@ namespace Horde.Build.Collections.Impl
 		/// <summary>
 		/// Collection of event documents
 		/// </summary>
-		IMongoCollection<LogEventDocument> LogEvents;
+		readonly IMongoCollection<LogEventDocument> _logEvents;
 
 		/// <summary>
 		/// Collection of legacy event documents
 		/// </summary>
-		IMongoCollection<LegacyEventDocument> LegacyEvents;
+		readonly IMongoCollection<LegacyEventDocument> _legacyEvents;
 
 		/// <summary>
 		/// Logger for upgrade messages
 		/// </summary>
-		ILogger<LogEventCollection> Logger;
+		readonly ILogger<LogEventCollection> _logger;
 
 		/// <summary>
 		/// Constructor
 		/// </summary>
-		/// <param name="DatabaseService">The database service instance</param>
-		/// <param name="Logger">The logger instance</param>
-		public LogEventCollection(DatabaseService DatabaseService, ILogger<LogEventCollection> Logger)
+		/// <param name="databaseService">The database service instance</param>
+		/// <param name="logger">The logger instance</param>
+		public LogEventCollection(DatabaseService databaseService, ILogger<LogEventCollection> logger)
 		{
-			this.Logger = Logger;
+			_logger = logger;
 
-			LogEvents = DatabaseService.GetCollection<LogEventDocument>("LogEvents");
-			LegacyEvents = DatabaseService.GetCollection<LegacyEventDocument>("Events");
+			_logEvents = databaseService.GetCollection<LogEventDocument>("LogEvents");
+			_legacyEvents = databaseService.GetCollection<LegacyEventDocument>("Events");
 
-			if (!DatabaseService.ReadOnlyMode)
+			if (!databaseService.ReadOnlyMode)
 			{
-				LogEvents.Indexes.CreateOne(new CreateIndexModel<LogEventDocument>(Builders<LogEventDocument>.IndexKeys.Ascending(x => x.Id.LogId)));
-				LogEvents.Indexes.CreateOne(new CreateIndexModel<LogEventDocument>(Builders<LogEventDocument>.IndexKeys.Ascending(x => x.SpanId).Ascending(x => x.Id)));
-				LegacyEvents.Indexes.CreateOne(new CreateIndexModel<LegacyEventDocument>(Builders<LegacyEventDocument>.IndexKeys.Ascending(x => x.LogId)));
+				_logEvents.Indexes.CreateOne(new CreateIndexModel<LogEventDocument>(Builders<LogEventDocument>.IndexKeys.Ascending(x => x.Id.LogId)));
+				_logEvents.Indexes.CreateOne(new CreateIndexModel<LogEventDocument>(Builders<LogEventDocument>.IndexKeys.Ascending(x => x.SpanId).Ascending(x => x.Id)));
+				_legacyEvents.Indexes.CreateOne(new CreateIndexModel<LegacyEventDocument>(Builders<LegacyEventDocument>.IndexKeys.Ascending(x => x.LogId)));
 			}
 		}
 
 		/// <inheritdoc/>
-		public Task AddAsync(NewLogEventData NewEvent)
+		public Task AddAsync(NewLogEventData newEvent)
 		{
-			return LogEvents.InsertOneAsync(new LogEventDocument(NewEvent));
+			return _logEvents.InsertOneAsync(new LogEventDocument(newEvent));
 		}
 
 		/// <inheritdoc/>
-		public Task AddManyAsync(List<NewLogEventData> NewEvents)
+		public Task AddManyAsync(List<NewLogEventData> newEvents)
 		{
-			return LogEvents.InsertManyAsync(NewEvents.ConvertAll(x => new LogEventDocument(x)));
+			return _logEvents.InsertManyAsync(newEvents.ConvertAll(x => new LogEventDocument(x)));
 		}
 
 		/// <inheritdoc/>
-		public async Task<List<ILogEvent>> FindAsync(LogId LogId, int? Index = null, int? Count = null)
+		public async Task<List<ILogEvent>> FindAsync(LogId logId, int? index = null, int? count = null)
 		{
-			Logger.LogInformation("Querying for log events for log {LogId} creation time {CreateTime}", LogId, LogId.Value.CreationTime);
+			_logger.LogInformation("Querying for log events for log {LogId} creation time {CreateTime}", logId, logId.Value.CreationTime);
 
-			FilterDefinitionBuilder<LogEventDocument> Builder = Builders<LogEventDocument>.Filter;
+			FilterDefinitionBuilder<LogEventDocument> builder = Builders<LogEventDocument>.Filter;
 
-			FilterDefinition<LogEventDocument> Filter = Builder.Eq(x => x.Id.LogId, LogId);
+			FilterDefinition<LogEventDocument> filter = builder.Eq(x => x.Id.LogId, logId);
 
-			IFindFluent<LogEventDocument, LogEventDocument> Results = LogEvents.Find(Filter).SortBy(x => x.Id);
-			if (Index != null)
+			IFindFluent<LogEventDocument, LogEventDocument> results = _logEvents.Find(filter).SortBy(x => x.Id);
+			if (index != null)
 			{
-				Results = Results.Skip(Index.Value);
+				results = results.Skip(index.Value);
 			}
-			if (Count != null)
+			if (count != null)
 			{
-				Results = Results.Limit(Count.Value);
+				results = results.Limit(count.Value);
 			}
 
-			List<LogEventDocument> EventDocuments = await Results.ToListAsync();
-			return EventDocuments.ConvertAll<ILogEvent>(x => x);
+			List<LogEventDocument> eventDocuments = await results.ToListAsync();
+			return eventDocuments.ConvertAll<ILogEvent>(x => x);
 		}
 
 		/// <inheritdoc/>
-		public async Task<List<ILogEvent>> FindEventsForSpanAsync(ObjectId SpanId, LogId[]? LogIds, int Index, int Count)
+		public async Task<List<ILogEvent>> FindEventsForSpanAsync(ObjectId spanId, LogId[]? logIds, int index, int count)
 		{
-			FilterDefinition<LogEventDocument> Filter = Builders<LogEventDocument>.Filter.Eq(x => x.SpanId, SpanId);
-			if (LogIds != null && LogIds.Length > 0)
+			FilterDefinition<LogEventDocument> filter = Builders<LogEventDocument>.Filter.Eq(x => x.SpanId, spanId);
+			if (logIds != null && logIds.Length > 0)
 			{
-				Filter &= Builders<LogEventDocument>.Filter.In(x => x.Id.LogId, LogIds);
+				filter &= Builders<LogEventDocument>.Filter.In(x => x.Id.LogId, logIds);
 			}
 
-			List<LogEventDocument> EventDocuments = await LogEvents.Find(Filter).Skip(Index).Limit(Count).ToListAsync();
-			return EventDocuments.ConvertAll<ILogEvent>(x => x);
+			List<LogEventDocument> eventDocuments = await _logEvents.Find(filter).Skip(index).Limit(count).ToListAsync();
+			return eventDocuments.ConvertAll<ILogEvent>(x => x);
 		}
 
 		/// <inheritdoc/>
-		public async Task<List<ILogEvent>> FindEventsForSpansAsync(IEnumerable<ObjectId> SpanIds, LogId[]? LogIds, int Index, int Count)
+		public async Task<List<ILogEvent>> FindEventsForSpansAsync(IEnumerable<ObjectId> spanIds, LogId[]? logIds, int index, int count)
 		{
-			FilterDefinition<LogEventDocument> Filter = Builders<LogEventDocument>.Filter.In(x => x.SpanId, SpanIds.Select<ObjectId, ObjectId?>(x => x));
-			if (LogIds != null && LogIds.Length > 0)
+			FilterDefinition<LogEventDocument> filter = Builders<LogEventDocument>.Filter.In(x => x.SpanId, spanIds.Select<ObjectId, ObjectId?>(x => x));
+			if (logIds != null && logIds.Length > 0)
 			{
-				Filter &= Builders<LogEventDocument>.Filter.In(x => x.Id.LogId, LogIds);
+				filter &= Builders<LogEventDocument>.Filter.In(x => x.Id.LogId, logIds);
 			}
 
-			List<LogEventDocument> EventDocuments = await LogEvents.Find(Filter).Skip(Index).Limit(Count).ToListAsync();
-			return EventDocuments.ConvertAll<ILogEvent>(x => x);
+			List<LogEventDocument> eventDocuments = await _logEvents.Find(filter).Skip(index).Limit(count).ToListAsync();
+			return eventDocuments.ConvertAll<ILogEvent>(x => x);
 		}
 
 		/// <inheritdoc/>
-		public async Task DeleteLogAsync(LogId LogId)
+		public async Task DeleteLogAsync(LogId logId)
 		{
-			await LogEvents.DeleteManyAsync(x => x.Id.LogId == LogId);
-			await LegacyEvents.DeleteManyAsync(x => x.LogId == LogId);
+			await _logEvents.DeleteManyAsync(x => x.Id.LogId == logId);
+			await _legacyEvents.DeleteManyAsync(x => x.LogId == logId);
 		}
 
 		/// <inheritdoc/>
-		public async Task AddSpanToEventsAsync(IEnumerable<ILogEvent> Events, ObjectId SpanId)
+		public async Task AddSpanToEventsAsync(IEnumerable<ILogEvent> events, ObjectId spanId)
 		{
-			FilterDefinition<LogEventDocument> EventFilter = Builders<LogEventDocument>.Filter.In(x => x.Id, Events.OfType<LogEventDocument>().Select(x => x.Id));
-			UpdateDefinition<LogEventDocument> EventUpdate = Builders<LogEventDocument>.Update.Set(x => x.SpanId, SpanId);
-			await LogEvents.UpdateManyAsync(EventFilter, EventUpdate);
+			FilterDefinition<LogEventDocument> eventFilter = Builders<LogEventDocument>.Filter.In(x => x.Id, events.OfType<LogEventDocument>().Select(x => x.Id));
+			UpdateDefinition<LogEventDocument> eventUpdate = Builders<LogEventDocument>.Update.Set(x => x.SpanId, spanId);
+			await _logEvents.UpdateManyAsync(eventFilter, eventUpdate);
 		}
 	}
 }
diff --git a/Engine/Source/Programs/Horde/Horde.Build/Logs/Impl/LogFileCollection.cs b/Engine/Source/Programs/Horde/Horde.Build/Logs/Impl/LogFileCollection.cs
--- a/Engine/Source/Programs/Horde/Horde.Build/Logs/Impl/LogFileCollection.cs
+++ b/Engine/Source/Programs/Horde/Horde.Build/Logs/Impl/LogFileCollection.cs
@@ -1,17 +1,15 @@
 // Copyright Epic Games, Inc. All Rights Reserved.
 
+using System.Collections.Generic;
+using System.Linq;
+using System.Net;
+using System.Threading.Tasks;
 using Horde.Build.Api;
 using Horde.Build.Models;
 using Horde.Build.Services;
 using Horde.Build.Utilities;
-using MongoDB.Bson;
 using MongoDB.Bson.Serialization.Attributes;
 using MongoDB.Driver;
-using System;
-using System.Collections.Generic;
-using System.Linq;
-using System.Net;
-using System.Threading.Tasks;
 
 namespace Horde.Build.Collections.Impl
 {
@@ -38,12 +36,12 @@ namespace Horde.Build.Collections.Impl
 			{
 			}
 
-			public LogChunkDocument(LogChunkDocument Other)
+			public LogChunkDocument(LogChunkDocument other)
 			{
-				Offset = Other.Offset;
-				Length = Other.Length;
-				LineIndex = Other.LineIndex;
-				Server = Other.Server;
+				Offset = other.Offset;
+				Length = other.Length;
+				LineIndex = other.LineIndex;
+				Server = other.Server;
 			}
 
 			public LogChunkDocument Clone()
@@ -81,155 +79,155 @@ namespace Horde.Build.Collections.Impl
 			{
 			}
 
-			public LogFileDocument(JobId JobId, SessionId? SessionId, LogType Type)
+			public LogFileDocument(JobId jobId, SessionId? sessionId, LogType type)
 			{
-				this.Id = LogId.GenerateNewId();
-				this.JobId = JobId;
-				this.SessionId = SessionId;
-				this.Type = Type;
-				this.MaxLineIndex = 0;
+				Id = LogId.GenerateNewId();
+				JobId = jobId;
+				SessionId = sessionId;
+				Type = type;
+				MaxLineIndex = 0;
 			}
 
 			public LogFileDocument Clone()
 			{
-				LogFileDocument Document = (LogFileDocument)MemberwiseClone();
-				Document.Chunks = Document.Chunks.ConvertAll(x => x.Clone());
-				return Document;
+				LogFileDocument document = (LogFileDocument)MemberwiseClone();
+				document.Chunks = document.Chunks.ConvertAll(x => x.Clone());
+				return document;
 			}
 		}
 
 		/// <summary>
 		/// The jobs collection
 		/// </summary>
-		IMongoCollection<LogFileDocument> LogFiles;
+		readonly IMongoCollection<LogFileDocument> _logFiles;
 
 		/// <summary>
 		/// Hostname for the current server
 		/// </summary>
-		string HostName;
+		readonly string _hostName;
 
 		/// <summary>
 		/// Constructor
 		/// </summary>
-		/// <param name="DatabaseService">The database service singleton</param>
-		public LogFileCollection(DatabaseService DatabaseService)
+		/// <param name="databaseService">The database service singleton</param>
+		public LogFileCollection(DatabaseService databaseService)
 		{
-			LogFiles = DatabaseService.GetCollection<LogFileDocument>("LogFiles");
-			this.HostName = Dns.GetHostName();
+			_logFiles = databaseService.GetCollection<LogFileDocument>("LogFiles");
+			_hostName = Dns.GetHostName();
 		}
 
 		/// <inheritdoc/>
-		public async Task<ILogFile> CreateLogFileAsync(JobId JobId, SessionId? SessionId, LogType Type)
+		public async Task<ILogFile> CreateLogFileAsync(JobId jobId, SessionId? sessionId, LogType type)
 		{
-			LogFileDocument NewLogFile = new LogFileDocument(JobId, SessionId, Type);
-			await LogFiles.InsertOneAsync(NewLogFile);
-			return NewLogFile;
+			LogFileDocument newLogFile = new LogFileDocument(jobId, sessionId, type);
+			await _logFiles.InsertOneAsync(newLogFile);
+			return newLogFile;
 		}
 
 		/// <inheritdoc/>
-		public async Task<ILogFile?> TryAddChunkAsync(ILogFile LogFileInterface, long Offset, int LineIndex)
+		public async Task<ILogFile?> TryAddChunkAsync(ILogFile logFileInterface, long offset, int lineIndex)
 		{
-			LogFileDocument LogFile = ((LogFileDocument)LogFileInterface).Clone();
+			LogFileDocument logFile = ((LogFileDocument)logFileInterface).Clone();
 
-			int ChunkIdx = LogFile.Chunks.GetChunkForOffset(Offset) + 1;
+			int chunkIdx = logFile.Chunks.GetChunkForOffset(offset) + 1;
 
-			LogChunkDocument Chunk = new LogChunkDocument();
-			Chunk.Offset = Offset;
-			Chunk.LineIndex = LineIndex;
-			Chunk.Server = HostName;
-			LogFile.Chunks.Insert(ChunkIdx, Chunk);
+			LogChunkDocument chunk = new LogChunkDocument();
+			chunk.Offset = offset;
+			chunk.LineIndex = lineIndex;
+			chunk.Server = _hostName;
+			logFile.Chunks.Insert(chunkIdx, chunk);
 
-			UpdateDefinition<LogFileDocument> Update = Builders<LogFileDocument>.Update.Set(x => x.Chunks, LogFile.Chunks);
-			if (ChunkIdx == LogFile.Chunks.Count - 1)
+			UpdateDefinition<LogFileDocument> update = Builders<LogFileDocument>.Update.Set(x => x.Chunks, logFile.Chunks);
+			if (chunkIdx == logFile.Chunks.Count - 1)
 			{
-				LogFile.MaxLineIndex = null;
-				Update = Update.Unset(x => x.MaxLineIndex);
+				logFile.MaxLineIndex = null;
+				update = update.Unset(x => x.MaxLineIndex);
 			}
 
-			if (!await TryUpdateLogFileAsync(LogFile, Update))
+			if (!await TryUpdateLogFileAsync(logFile, update))
 			{
 				return null;
 			}
 
-			return LogFile;
+			return logFile;
 		}
 
 		/// <inheritdoc/>
-		public async Task<ILogFile?> TryCompleteChunksAsync(ILogFile LogFileInterface, IEnumerable<CompleteLogChunkUpdate> ChunkUpdates)
+		public async Task<ILogFile?> TryCompleteChunksAsync(ILogFile logFileInterface, IEnumerable<CompleteLogChunkUpdate> chunkUpdates)
 		{
-			LogFileDocument LogFile = ((LogFileDocument)LogFileInterface).Clone();
+			LogFileDocument logFile = ((LogFileDocument)logFileInterface).Clone();
 
 			// Update the length of any complete chunks
-			UpdateDefinitionBuilder<LogFileDocument> UpdateBuilder = Builders<LogFileDocument>.Update;
-			List<UpdateDefinition<LogFileDocument>> Updates = new List<UpdateDefinition<LogFileDocument>>();
-			foreach (CompleteLogChunkUpdate ChunkUpdate in ChunkUpdates)
+			UpdateDefinitionBuilder<LogFileDocument> updateBuilder = Builders<LogFileDocument>.Update;
+			List<UpdateDefinition<LogFileDocument>> updates = new List<UpdateDefinition<LogFileDocument>>();
+			foreach (CompleteLogChunkUpdate chunkUpdate in chunkUpdates)
 			{
-				LogChunkDocument Chunk = LogFile.Chunks[ChunkUpdate.Index];
-				Chunk.Length = ChunkUpdate.Length;
-				Updates.Add(UpdateBuilder.Set(x => x.Chunks[ChunkUpdate.Index].Length, ChunkUpdate.Length));
+				LogChunkDocument chunk = logFile.Chunks[chunkUpdate.Index];
+				chunk.Length = chunkUpdate.Length;
+				updates.Add(updateBuilder.Set(x => x.Chunks[chunkUpdate.Index].Length, chunkUpdate.Length));
 
-				if (ChunkUpdate.Index == LogFile.Chunks.Count - 1)
+				if (chunkUpdate.Index == logFile.Chunks.Count - 1)
 				{
-					LogFile.MaxLineIndex = Chunk.LineIndex + ChunkUpdate.LineCount;
-					Updates.Add(UpdateBuilder.Set(x => x.MaxLineIndex, LogFile.MaxLineIndex));
+					logFile.MaxLineIndex = chunk.LineIndex + chunkUpdate.LineCount;
+					updates.Add(updateBuilder.Set(x => x.MaxLineIndex, logFile.MaxLineIndex));
 				}
 			}
 
 			// Try to apply the updates
-			if (Updates.Count > 0 && !await TryUpdateLogFileAsync(LogFile, UpdateBuilder.Combine(Updates)))
+			if (updates.Count > 0 && !await TryUpdateLogFileAsync(logFile, updateBuilder.Combine(updates)))
 			{
 				return null;
 			}
 
-			return LogFile;
+			return logFile;
 		}
 
 		/// <inheritdoc/>
-		public async Task<ILogFile?> TryUpdateIndexAsync(ILogFile LogFileInterface, long NewIndexLength)
+		public async Task<ILogFile?> TryUpdateIndexAsync(ILogFile logFileInterface, long newIndexLength)
 		{
-			LogFileDocument LogFile = ((LogFileDocument)LogFileInterface).Clone();
+			LogFileDocument logFile = ((LogFileDocument)logFileInterface).Clone();
 
-			UpdateDefinition<LogFileDocument> Update = Builders<LogFileDocument>.Update.Set(x => x.IndexLength, NewIndexLength);
-			if (!await TryUpdateLogFileAsync(LogFile, Update))
+			UpdateDefinition<LogFileDocument> update = Builders<LogFileDocument>.Update.Set(x => x.IndexLength, newIndexLength);
+			if (!await TryUpdateLogFileAsync(logFile, update))
 			{
 				return null;
 			}
 
-			LogFile.IndexLength = NewIndexLength;
-			return LogFile;
+			logFile.IndexLength = newIndexLength;
+			return logFile;
 		}
 
 		/// <inheritdoc/>
-		private async Task<bool> TryUpdateLogFileAsync(LogFileDocument Current, UpdateDefinition<LogFileDocument> Update)
+		private async Task<bool> TryUpdateLogFileAsync(LogFileDocument current, UpdateDefinition<LogFileDocument> update)
 		{
-			int PrevUpdateIndex = Current.UpdateIndex;
-			Current.UpdateIndex++;
-			UpdateResult Result = await LogFiles.UpdateOneAsync<LogFileDocument>(x => x.Id == Current.Id && x.UpdateIndex == PrevUpdateIndex, Update.Set(x => x.UpdateIndex, Current.UpdateIndex));
-			return Result.ModifiedCount == 1;
+			int prevUpdateIndex = current.UpdateIndex;
+			current.UpdateIndex++;
+			UpdateResult result = await _logFiles.UpdateOneAsync<LogFileDocument>(x => x.Id == current.Id && x.UpdateIndex == prevUpdateIndex, update.Set(x => x.UpdateIndex, current.UpdateIndex));
+			return result.ModifiedCount == 1;
 		}
 
 		/// <inheritdoc/>
-		public async Task<ILogFile?> GetLogFileAsync(LogId LogFileId)
+		public async Task<ILogFile?> GetLogFileAsync(LogId logFileId)
 		{
-			LogFileDocument LogFile = await LogFiles.Find<LogFileDocument>(x => x.Id == LogFileId).FirstOrDefaultAsync();
-			return LogFile;
+			LogFileDocument logFile = await _logFiles.Find<LogFileDocument>(x => x.Id == logFileId).FirstOrDefaultAsync();
+			return logFile;
 		}
 
 		/// <inheritdoc/>
-		public async Task<List<ILogFile>> GetLogFilesAsync(int? Index = null, int? Count = null)
+		public async Task<List<ILogFile>> GetLogFilesAsync(int? index = null, int? count = null)
 		{
-			IFindFluent<LogFileDocument, LogFileDocument> Query = LogFiles.Find(FilterDefinition<LogFileDocument>.Empty);
-			if(Index != null)
+			IFindFluent<LogFileDocument, LogFileDocument> query = _logFiles.Find(FilterDefinition<LogFileDocument>.Empty);
+			if(index != null)
 			{
-				Query = Query.Skip(Index.Value);
+				query = query.Skip(index.Value);
 			}
-			if(Count != null)
+			if(count != null)
 			{
-				Query = Query.Limit(Count.Value);
+				query = query.Limit(count.Value);
 			}
 
-			List<LogFileDocument> Results = await Query.ToListAsync();
-			return Results.ConvertAll<ILogFile>(x => x);
+			List<LogFileDocument> results = await query.ToListAsync();
+			return results.ConvertAll<ILogFile>(x => x);
 		}
 	}
 }
