Horde: Fix coding conventions within Horde.Build.

diff --git a/Engine/Source/Programs/Horde/Horde.Build/Logs/Impl/LogFileService.cs b/Engine/Source/Programs/Horde/Horde.Build/Logs/Impl/LogFileService.cs
--- a/Engine/Source/Programs/Horde/Horde.Build/Logs/Impl/LogFileService.cs
+++ b/Engine/Source/Programs/Horde/Horde.Build/Logs/Impl/LogFileService.cs
@@ -1,39 +1,32 @@
 // Copyright Epic Games, Inc. All Rights Reserved.
 
-using EpicGames.Core;
-using HordeCommon;
-using Horde.Build.Api;
-using Horde.Build.Collections;
-using Horde.Build.Logs;
-using Horde.Build.Models;
-using Horde.Build.Storage;
-using Horde.Build.Utilities;
-using Microsoft.Extensions.Caching.Memory;
-using Microsoft.Extensions.Logging;
-using MongoDB.Bson;
-using MongoDB.Driver;
 using System;
-using System.Collections.Concurrent;
 using System.Collections.Generic;
 using System.Collections.Immutable;
 using System.Diagnostics;
-using System.Diagnostics.CodeAnalysis;
 using System.Globalization;
 using System.IO;
 using System.Linq;
-using System.Runtime.CompilerServices;
 using System.Security.Claims;
-using System.Text.Json;
 using System.Text;
-using System.Threading.Tasks;
+using System.Text.Json;
 using System.Threading;
-
-using ILogger = Microsoft.Extensions.Logging.ILogger;
-using Stream = System.IO.Stream;
-using OpenTracing.Util;
-using OpenTracing;
+using System.Threading.Tasks;
+using EpicGames.Core;
+using Horde.Build.Api;
+using Horde.Build.Collections;
+using Horde.Build.Logs;
+using Horde.Build.Models;
+using Horde.Build.Utilities;
+using HordeCommon;
+using Microsoft.Extensions.Caching.Memory;
 using Microsoft.Extensions.Hosting;
-using Microsoft.Extensions.Logging.Abstractions;
+using Microsoft.Extensions.Logging;
+using MongoDB.Bson;
+using MongoDB.Driver;
+using OpenTracing;
+using OpenTracing.Util;
+using Stream = System.IO.Stream;
 
 namespace Horde.Build.Services
 {
@@ -65,124 +58,124 @@ namespace Horde.Build.Services
 		/// <summary>
 		/// Creates a new log
 		/// </summary>
-		/// <param name="JobId">Unique id of the job that owns this log file</param>
-		/// <param name="SessionId">Agent session allowed to update the log</param>
-		/// <param name="Type">Type of events to be stored in the log</param>
+		/// <param name="jobId">Unique id of the job that owns this log file</param>
+		/// <param name="sessionId">Agent session allowed to update the log</param>
+		/// <param name="type">Type of events to be stored in the log</param>
 		/// <returns>The new log file document</returns>
-		Task<ILogFile> CreateLogFileAsync(JobId JobId, SessionId? SessionId, LogType Type);
+		Task<ILogFile> CreateLogFileAsync(JobId jobId, SessionId? sessionId, LogType type);
 
 		/// <summary>
 		/// Gets a logfile by ID
 		/// </summary>
-		/// <param name="LogFileId">Unique id of the log file</param>
+		/// <param name="logFileId">Unique id of the log file</param>
 		/// <returns>The logfile document</returns>
-		Task<ILogFile?> GetLogFileAsync(LogId LogFileId);
+		Task<ILogFile?> GetLogFileAsync(LogId logFileId);
 
 		/// <summary>
 		/// Gets a logfile by ID, returning a cached copy if available. This should only be used to retrieve constant properties set at creation, such as the session or job it's associated with.
 		/// </summary>
-		/// <param name="LogFileId">Unique id of the log file</param>
+		/// <param name="logFileId">Unique id of the log file</param>
 		/// <returns>The logfile document</returns>
-		Task<ILogFile?> GetCachedLogFileAsync(LogId LogFileId);
+		Task<ILogFile?> GetCachedLogFileAsync(LogId logFileId);
 
 		/// <summary>
 		/// Returns a list of log files
 		/// </summary>
-		/// <param name="Index">Index of the first result to return</param>
-		/// <param name="Count">Number of results to return</param>
+		/// <param name="index">Index of the first result to return</param>
+		/// <param name="count">Number of results to return</param>
 		/// <returns>List of logfile documents</returns>
-		Task<List<ILogFile>> GetLogFilesAsync(int? Index = null, int? Count = null);
+		Task<List<ILogFile>> GetLogFilesAsync(int? index = null, int? count = null);
 
 		/// <summary>
 		/// Writes out chunk data and assigns to a file
 		/// </summary>
-		/// <param name="LogFile">The log file</param>
-		/// <param name="Offset">Offset within the file of data</param>
-		/// <param name="LineIndex">Current line index of the data (need not be the starting of the line)</param>
-		/// <param name="Data">the data to add</param>
-		/// <param name="Flush">Whether the current chunk is complete and should be flushed</param>
-		/// <param name="MaxChunkLength">The maximum chunk length. Defaults to 128kb.</param>
-		/// <param name="MaxSubChunkLineCount">Maximum number of lines in each sub-chunk.</param>
+		/// <param name="logFile">The log file</param>
+		/// <param name="offset">Offset within the file of data</param>
+		/// <param name="lineIndex">Current line index of the data (need not be the starting of the line)</param>
+		/// <param name="data">the data to add</param>
+		/// <param name="flush">Whether the current chunk is complete and should be flushed</param>
+		/// <param name="maxChunkLength">The maximum chunk length. Defaults to 128kb.</param>
+		/// <param name="maxSubChunkLineCount">Maximum number of lines in each sub-chunk.</param>
 		/// <returns></returns>
-		Task<ILogFile?> WriteLogDataAsync(ILogFile LogFile, long Offset, int LineIndex, ReadOnlyMemory<byte> Data, bool Flush, int MaxChunkLength = 256 * 1024, int MaxSubChunkLineCount = 128);
+		Task<ILogFile?> WriteLogDataAsync(ILogFile logFile, long offset, int lineIndex, ReadOnlyMemory<byte> data, bool flush, int maxChunkLength = 256 * 1024, int maxSubChunkLineCount = 128);
 
 		/// <summary>
 		/// Gets metadata about the log file
 		/// </summary>
-		/// <param name="LogFile">The log file to query</param>
+		/// <param name="logFile">The log file to query</param>
 		/// <returns>Metadata about the log file</returns>
-		Task<LogMetadata> GetMetadataAsync(ILogFile LogFile);
+		Task<LogMetadata> GetMetadataAsync(ILogFile logFile);
 
 		/// <summary>
 		/// Creates new log events
 		/// </summary>
-		/// <param name="NewEvents">List of events</param>
+		/// <param name="newEvents">List of events</param>
 		/// <returns>Async task</returns>
-		Task CreateEventsAsync(List<NewLogEventData> NewEvents);
+		Task CreateEventsAsync(List<NewLogEventData> newEvents);
 
 		/// <summary>
 		/// Find events for a particular log file
 		/// </summary>
-		/// <param name="LogFile">The log file instance</param>
-		/// <param name="Index">Index of the first event to retrieve</param>
-		/// <param name="Count">Number of events to retrieve</param>
+		/// <param name="logFile">The log file instance</param>
+		/// <param name="index">Index of the first event to retrieve</param>
+		/// <param name="count">Number of events to retrieve</param>
 		/// <returns>List of log events</returns>
-		Task<List<ILogEvent>> FindLogEventsAsync(ILogFile LogFile, int? Index = null, int? Count = null);
+		Task<List<ILogEvent>> FindLogEventsAsync(ILogFile logFile, int? index = null, int? count = null);
 
 		/// <summary>
 		/// Adds events to a log span
 		/// </summary>
-		/// <param name="Events">The events to add</param>
-		/// <param name="SpanId">The span id</param>
+		/// <param name="events">The events to add</param>
+		/// <param name="spanId">The span id</param>
 		/// <returns>Async task</returns>
-		Task AddSpanToEventsAsync(IEnumerable<ILogEvent> Events, ObjectId SpanId);
+		Task AddSpanToEventsAsync(IEnumerable<ILogEvent> events, ObjectId spanId);
 
 		/// <summary>
 		/// Find events for an issue
 		/// </summary>
-		/// <param name="SpanIds">The span ids</param>
-		/// <param name="LogIds">Log ids to include</param>
-		/// <param name="Index">Index within the events for results to return</param>
-		/// <param name="Count">Number of results to return</param>
+		/// <param name="spanIds">The span ids</param>
+		/// <param name="logIds">Log ids to include</param>
+		/// <param name="index">Index within the events for results to return</param>
+		/// <param name="count">Number of results to return</param>
 		/// <returns>Async task</returns>
-		Task<List<ILogEvent>> FindEventsForSpansAsync(IEnumerable<ObjectId> SpanIds, LogId[]? LogIds, int Index, int Count);
+		Task<List<ILogEvent>> FindEventsForSpansAsync(IEnumerable<ObjectId> spanIds, LogId[]? logIds, int index, int count);
 
 		/// <summary>
 		/// Gets the data for an event
 		/// </summary>
-		/// <param name="LogFile">The log file instance</param>
-		/// <param name="LineIndex">Index of the line in the file</param>
-		/// <param name="LineCount">Number of lines in the event</param>
+		/// <param name="logFile">The log file instance</param>
+		/// <param name="lineIndex">Index of the line in the file</param>
+		/// <param name="lineCount">Number of lines in the event</param>
 		/// <returns>New event data instance</returns>
-		Task<ILogEventData> GetEventDataAsync(ILogFile LogFile, int LineIndex, int LineCount);
+		Task<ILogEventData> GetEventDataAsync(ILogFile logFile, int lineIndex, int lineCount);
 
 		/// <summary>
 		/// Gets lines from the given log 
 		/// </summary>
-		/// <param name="LogFile">The log file</param>
-		/// <param name="Offset">Offset of the data to return</param>
-		/// <param name="Length">Length of the data to return</param>
+		/// <param name="logFile">The log file</param>
+		/// <param name="offset">Offset of the data to return</param>
+		/// <param name="length">Length of the data to return</param>
 		/// <returns>Data for the requested range</returns>
-		Task<Stream> OpenRawStreamAsync(ILogFile LogFile, long Offset, long Length);
+		Task<Stream> OpenRawStreamAsync(ILogFile logFile, long offset, long length);
 
 		/// <summary>
 		/// Gets the offset of the given line number
 		/// </summary>
-		/// <param name="LogFile">The log file to search</param>
-		/// <param name="LineIdx">The line index to retrieve the offset for</param>
+		/// <param name="logFile">The log file to search</param>
+		/// <param name="lineIdx">The line index to retrieve the offset for</param>
 		/// <returns>The actual clamped line number and offset</returns>
-		Task<(int, long)> GetLineOffsetAsync(ILogFile LogFile, int LineIdx);
+		Task<(int, long)> GetLineOffsetAsync(ILogFile logFile, int lineIdx);
 
 		/// <summary>
 		/// Search for the specified text in a log file
 		/// </summary>
-		/// <param name="LogFile">The log file to search</param>
-		/// <param name="Text">Text to search for</param>
-		/// <param name="FirstLine">Line to start search from</param>
-		/// <param name="Count">Number of results to return</param>
-		/// <param name="Stats">Receives stats for the search</param>
+		/// <param name="logFile">The log file to search</param>
+		/// <param name="text">Text to search for</param>
+		/// <param name="firstLine">Line to start search from</param>
+		/// <param name="count">Number of results to return</param>
+		/// <param name="stats">Receives stats for the search</param>
 		/// <returns>List of line numbers containing the given term</returns>
-		Task<List<int>> SearchLogDataAsync(ILogFile LogFile, string Text, int FirstLine, int Count, LogSearchStats Stats);
+		Task<List<int>> SearchLogDataAsync(ILogFile logFile, string text, int firstLine, int count, LogSearchStats stats);
 	}
 
 	/// <summary>
@@ -193,83 +186,83 @@ namespace Horde.Build.Services
 		/// <summary>
 		/// Parses a stream of json text and outputs plain text
 		/// </summary>
-		/// <param name="LogFileService">The log file service</param>
-		/// <param name="LogFile">The log file to query</param>
-		/// <param name="Offset">Offset within the log file to copy</param>
-		/// <param name="Length">Length of the data to copy</param>
-		/// <param name="OutputStream">Output stream to receive the text data</param>
+		/// <param name="logFileService">The log file service</param>
+		/// <param name="logFile">The log file to query</param>
+		/// <param name="offset">Offset within the log file to copy</param>
+		/// <param name="length">Length of the data to copy</param>
+		/// <param name="outputStream">Output stream to receive the text data</param>
 		/// <returns>Async text</returns>
-		public static async Task CopyRawStreamAsync(this ILogFileService LogFileService, ILogFile LogFile, long Offset, long Length, Stream OutputStream)
+		public static async Task CopyRawStreamAsync(this ILogFileService logFileService, ILogFile logFile, long offset, long length, Stream outputStream)
 		{
-			using (Stream Stream = await LogFileService.OpenRawStreamAsync(LogFile, Offset, Length))
+			using (Stream stream = await logFileService.OpenRawStreamAsync(logFile, offset, length))
 			{
-				await Stream.CopyToAsync(OutputStream);
+				await stream.CopyToAsync(outputStream);
 			}
 		}
 
 		/// <summary>
 		/// Parses a stream of json text and outputs plain text
 		/// </summary>
-		/// <param name="LogFileService">The log file service</param>
-		/// <param name="LogFile">The log file to query</param>
-		/// <param name="Offset">Offset within the data to copy from</param>
-		/// <param name="Length">Length of the data to copy</param>
-		/// <param name="OutputStream">Output stream to receive the text data</param>
+		/// <param name="logFileService">The log file service</param>
+		/// <param name="logFile">The log file to query</param>
+		/// <param name="offset">Offset within the data to copy from</param>
+		/// <param name="length">Length of the data to copy</param>
+		/// <param name="outputStream">Output stream to receive the text data</param>
 		/// <returns>Async text</returns>
-		public static async Task CopyPlainTextStreamAsync(this ILogFileService LogFileService, ILogFile LogFile, long Offset, long Length, Stream OutputStream)
+		public static async Task CopyPlainTextStreamAsync(this ILogFileService logFileService, ILogFile logFile, long offset, long length, Stream outputStream)
 		{
-			using (Stream Stream = await LogFileService.OpenRawStreamAsync(LogFile, 0, long.MaxValue))
+			using (Stream stream = await logFileService.OpenRawStreamAsync(logFile, 0, Int64.MaxValue))
 			{
-				byte[] ReadBuffer = new byte[4096];
-				int ReadBufferLength = 0;
+				byte[] readBuffer = new byte[4096];
+				int readBufferLength = 0;
 
-				byte[] WriteBuffer = new byte[4096];
-				int WriteBufferLength = 0;
+				byte[] writeBuffer = new byte[4096];
+				int writeBufferLength = 0;
 
-				while (Length > 0)
+				while (length > 0)
 				{
 					// Add more data to the buffer
-					int ReadBytes = await Stream.ReadAsync(ReadBuffer.AsMemory(ReadBufferLength, ReadBuffer.Length - ReadBufferLength));
-					ReadBufferLength += ReadBytes;
+					int readBytes = await stream.ReadAsync(readBuffer.AsMemory(readBufferLength, readBuffer.Length - readBufferLength));
+					readBufferLength += readBytes;
 
 					// Copy as many lines as possible to the output
-					int ConvertedBytes = 0;
-					for (int EndIdx = 1; EndIdx < ReadBufferLength; EndIdx++)
+					int convertedBytes = 0;
+					for (int endIdx = 1; endIdx < readBufferLength; endIdx++)
 					{
-						if (ReadBuffer[EndIdx] == '\n')
+						if (readBuffer[endIdx] == '\n')
 						{
-							WriteBufferLength = LogText.ConvertToPlainText(ReadBuffer.AsSpan(ConvertedBytes, EndIdx - ConvertedBytes), WriteBuffer, WriteBufferLength);
-							ConvertedBytes = EndIdx + 1;
+							writeBufferLength = LogText.ConvertToPlainText(readBuffer.AsSpan(convertedBytes, endIdx - convertedBytes), writeBuffer, writeBufferLength);
+							convertedBytes = endIdx + 1;
 						}
 					}
 
 					// If there's anything in the write buffer, write it out
-					if (WriteBufferLength > 0)
+					if (writeBufferLength > 0)
 					{
-						if (Offset < WriteBufferLength)
+						if (offset < writeBufferLength)
 						{
-							int WriteLength = (int)Math.Min((long)WriteBufferLength - Offset, Length);
-							await OutputStream.WriteAsync(WriteBuffer.AsMemory((int)Offset, WriteLength));
-							Length -= WriteLength;
+							int writeLength = (int)Math.Min((long)writeBufferLength - offset, length);
+							await outputStream.WriteAsync(writeBuffer.AsMemory((int)offset, writeLength));
+							length -= writeLength;
 						}
-						Offset = Math.Max(Offset - WriteBufferLength, 0);
-						WriteBufferLength = 0;
+						offset = Math.Max(offset - writeBufferLength, 0);
+						writeBufferLength = 0;
 					}
 
 					// If we were able to read something, shuffle down the rest of the buffer. Otherwise expand the read buffer.
-					if (ConvertedBytes > 0)
+					if (convertedBytes > 0)
 					{
-						Buffer.BlockCopy(ReadBuffer, ConvertedBytes, ReadBuffer, 0, ReadBufferLength - ConvertedBytes);
-						ReadBufferLength -= ConvertedBytes;
+						Buffer.BlockCopy(readBuffer, convertedBytes, readBuffer, 0, readBufferLength - convertedBytes);
+						readBufferLength -= convertedBytes;
 					}
-					else if (ReadBufferLength > 0)
+					else if (readBufferLength > 0)
 					{
-						Array.Resize(ref ReadBuffer, ReadBuffer.Length + 128);
-						WriteBuffer = new byte[ReadBuffer.Length];
+						Array.Resize(ref readBuffer, readBuffer.Length + 128);
+						writeBuffer = new byte[readBuffer.Length];
 					}
 
 					// Exit if we didn't read anything in this iteration
-					if (ReadBytes == 0)
+					if (readBytes == 0)
 					{
 						break;
 					}
@@ -283,50 +276,17 @@ namespace Horde.Build.Services
 	/// </summary>
 	public sealed class LogFileService : IHostedService, ILogFileService, IDisposable
 	{
-		/// <summary>
-		/// Information Logger
-		/// </summary>
-		private readonly ILogger<LogFileService> Logger;
-
-		/// <summary>
-		/// Collection of log documents
-		/// </summary>
-		private readonly ILogFileCollection LogFiles;
-
-		/// <summary>
-		/// Collection of log events
-		/// </summary>
-		private readonly ILogEventCollection LogEvents;
-
-		/// <summary>
-		/// Interface for the log reader
-		/// </summary>
-		private readonly ILogStorage Storage;
-
-		/// <summary>
-		/// Interface for the log builder
-		/// </summary>
-		private readonly ILogBuilder Builder;
-
-		/// <summary>
-		/// Lock object for the <see cref="WriteTasks"/> and <see cref="WriteChunks"/> members
-		/// </summary>
-		private object WriteLock = new object();
-
-		/// <summary>
-		/// List of active write tasks
-		/// </summary>
-		private List<Task> WriteTasks = new List<Task>();
-
-		/// <summary>
-		/// Set of chunks which is currently being written to
-		/// </summary>
-		private HashSet<(LogId, long)> WriteChunks = new HashSet<(LogId, long)>();
-
-		/// <summary>
-		/// Cache of session ids for each log
-		/// </summary>
-		private IMemoryCache LogFileCache;
+		private readonly ILogger<LogFileService> _logger;
+		private readonly ILogFileCollection _logFiles;
+		private readonly ILogEventCollection _logEvents;
+		private readonly ILogStorage _storage;
+		private readonly ILogBuilder _builder;
+
+		// Lock object for the <see cref="_writeTasks"/> and <see cref="_writeChunks"/> members
+		private readonly object _writeLock = new object();
+		private readonly List<Task> _writeTasks = new List<Task>();
+		private readonly HashSet<(LogId, long)> _writeChunks = new HashSet<(LogId, long)>();
+		private readonly IMemoryCache _logFileCache;
 
 		/// <summary>
 		/// Streams log data to a caller
@@ -336,67 +296,67 @@ namespace Horde.Build.Services
 			/// <summary>
 			/// The log file service that created this stream
 			/// </summary>
-			LogFileService LogFileService;
+			readonly LogFileService _logFileService;
 
 			/// <summary>
 			/// The log file being read
 			/// </summary>
-			ILogFile LogFile;
+			readonly ILogFile _logFile;
 
 			/// <summary>
 			/// Starting offset within the file of the data to return 
 			/// </summary>
-			long ResponseOffset;
+			readonly long _responseOffset;
 
 			/// <summary>
 			/// Length of data to return
 			/// </summary>
-			long ResponseLength;
+			readonly long _responseLength;
 
 			/// <summary>
 			/// Current offset within the stream
 			/// </summary>
-			long CurrentOffset;
+			long _currentOffset;
 
 			/// <summary>
 			/// The current chunk index
 			/// </summary>
-			int ChunkIdx;
+			int _chunkIdx;
 
 			/// <summary>
 			/// Buffer containing a message for missing data
 			/// </summary>
-			ReadOnlyMemory<byte> SourceBuffer;
+			ReadOnlyMemory<byte> _sourceBuffer;
 
 			/// <summary>
 			/// Offset within the source buffer
 			/// </summary>
-			int SourcePos;
+			int _sourcePos;
 
 			/// <summary>
 			/// Length of the source buffer being copied from
 			/// </summary>
-			int SourceEnd;
+			int _sourceEnd;
 
 			/// <summary>
 			/// Constructor
 			/// </summary>
-			/// <param name="LogFileService">The log file service, for q</param>
-			/// <param name="LogFile"></param>
-			/// <param name="Offset"></param>
-			/// <param name="Length"></param>
-			public ResponseStream(LogFileService LogFileService, ILogFile LogFile, long Offset, long Length)
+			/// <param name="logFileService">The log file service, for q</param>
+			/// <param name="logFile"></param>
+			/// <param name="offset"></param>
+			/// <param name="length"></param>
+			public ResponseStream(LogFileService logFileService, ILogFile logFile, long offset, long length)
 			{
-				this.LogFileService = LogFileService;
-				this.LogFile = LogFile;
+				_logFileService = logFileService;
+				_logFile = logFile;
 
-				this.ResponseOffset = Offset;
-				this.ResponseLength = Length;
+				_responseOffset = offset;
+				_responseLength = length;
 
-				this.CurrentOffset = Offset;
+				_currentOffset = offset;
 
-				this.ChunkIdx = LogFile.Chunks.GetChunkForOffset(Offset);
-				this.SourceBuffer = null!;
+				_chunkIdx = logFile.Chunks.GetChunkForOffset(offset);
+				_sourceBuffer = null!;
 			}
 
 			/// <inheritdoc/>
@@ -409,15 +369,12 @@ namespace Horde.Build.Services
 			public override bool CanWrite => false;
 
 			/// <inheritdoc/>
-			public override long Length
-			{
-				get { return ResponseLength; }
-			}
+			public override long Length => _responseLength;
 
 			/// <inheritdoc/>
 			public override long Position
 			{
-				get { return CurrentOffset - ResponseOffset; }
+				get => _currentOffset - _responseOffset;
 				set => throw new NotImplementedException();
 			}
 
@@ -427,53 +384,53 @@ namespace Horde.Build.Services
 			}
 
 			/// <inheritdoc/>
-			public override int Read(byte[] Buffer, int Offset, int Count)
+			public override int Read(byte[] buffer, int offset, int count)
 			{
-				return ReadAsync(Buffer, Offset, Count, CancellationToken.None).Result;
+				return ReadAsync(buffer, offset, count, CancellationToken.None).Result;
 			}
 
 			/// <inheritdoc/>
-			public override async Task<int> ReadAsync(byte[] Buffer, int Offset, int Length, CancellationToken CancellationToken)
+			public override async Task<int> ReadAsync(byte[] buffer, int offset, int length, CancellationToken cancellationToken)
 			{
-				return await ReadAsync(Buffer.AsMemory(Offset, Length), CancellationToken);
+				return await ReadAsync(buffer.AsMemory(offset, length), cancellationToken);
 			}
 
 			/// <inheritdoc/>
-			public override async ValueTask<int> ReadAsync(Memory<byte> Buffer, CancellationToken CancellationToken)
+			public override async ValueTask<int> ReadAsync(Memory<byte> buffer, CancellationToken cancellationToken)
 			{
-				int ReadBytes = 0;
-				while (ReadBytes < Buffer.Length)
+				int readBytes = 0;
+				while (readBytes < buffer.Length)
 				{
-					if (SourcePos < SourceEnd)
+					if (_sourcePos < _sourceEnd)
 					{
 						// Try to copy from the current buffer
-						int BlockSize = Math.Min(SourceEnd - SourcePos, Buffer.Length - ReadBytes);
-						SourceBuffer.Slice(SourcePos, BlockSize).Span.CopyTo(Buffer.Slice(ReadBytes).Span);
-						CurrentOffset += BlockSize;
-						ReadBytes += BlockSize;
-						SourcePos += BlockSize;
+						int blockSize = Math.Min(_sourceEnd - _sourcePos, buffer.Length - readBytes);
+						_sourceBuffer.Slice(_sourcePos, blockSize).Span.CopyTo(buffer.Slice(readBytes).Span);
+						_currentOffset += blockSize;
+						readBytes += blockSize;
+						_sourcePos += blockSize;
 					}
-					else if (CurrentOffset < ResponseOffset + ResponseLength)
+					else if (_currentOffset < _responseOffset + _responseLength)
 					{
 						// Move to the right chunk
-						while (ChunkIdx + 1 < LogFile.Chunks.Count && CurrentOffset >= LogFile.Chunks[ChunkIdx + 1].Offset)
+						while (_chunkIdx + 1 < _logFile.Chunks.Count && _currentOffset >= _logFile.Chunks[_chunkIdx + 1].Offset)
 						{
-							ChunkIdx++;
+							_chunkIdx++;
 						}
 
 						// Get the chunk data
-						ILogChunk Chunk = LogFile.Chunks[ChunkIdx];
-						LogChunkData ChunkData = await LogFileService.ReadChunkAsync(LogFile, ChunkIdx);
+						ILogChunk chunk = _logFile.Chunks[_chunkIdx];
+						LogChunkData chunkData = await _logFileService.ReadChunkAsync(_logFile, _chunkIdx);
 
 						// Figure out which sub-chunk to use
-						int SubChunkIdx = ChunkData.GetSubChunkForOffsetWithinChunk((int)(CurrentOffset - Chunk.Offset));
-						LogSubChunkData SubChunkData = ChunkData.SubChunks[SubChunkIdx];
+						int subChunkIdx = chunkData.GetSubChunkForOffsetWithinChunk((int)(_currentOffset - chunk.Offset));
+						LogSubChunkData subChunkData = chunkData.SubChunks[subChunkIdx];
 
 						// Get the source data
-						long SubChunkOffset = Chunk.Offset + ChunkData.SubChunkOffset[SubChunkIdx];
-						SourceBuffer = SubChunkData.InflateText().Data;
-						SourcePos = (int)(CurrentOffset - SubChunkOffset);
-						SourceEnd = (int)Math.Min(SourceBuffer.Length, (ResponseOffset + ResponseLength) - SubChunkOffset);
+						long subChunkOffset = chunk.Offset + chunkData.SubChunkOffset[subChunkIdx];
+						_sourceBuffer = subChunkData.InflateText().Data;
+						_sourcePos = (int)(_currentOffset - subChunkOffset);
+						_sourceEnd = (int)Math.Min(_sourceBuffer.Length, (_responseOffset + _responseLength) - subChunkOffset);
 					}
 					else
 					{
@@ -481,7 +438,7 @@ namespace Horde.Build.Services
 						break;
 					}
 				}
-				return ReadBytes;
+				return readBytes;
 			}
 
 			/// <inheritdoc/>
@@ -494,215 +451,215 @@ namespace Horde.Build.Services
 			public override void Write(byte[] buffer, int offset, int count) => throw new NotImplementedException();
 		}
 
-		ITicker Ticker;
+		readonly ITicker _ticker;
 
 		/// <summary>
 		/// Constructor
 		/// </summary>
-		public LogFileService(ILogFileCollection LogFiles, ILogEventCollection LogEvents, ILogBuilder Builder, ILogStorage Storage, IClock Clock, ILogger<LogFileService> Logger)
+		public LogFileService(ILogFileCollection logFiles, ILogEventCollection logEvents, ILogBuilder builder, ILogStorage storage, IClock clock, ILogger<LogFileService> logger)
 		{
-			this.LogFiles = LogFiles;
-			this.LogEvents = LogEvents;
-			this.LogFileCache = new MemoryCache(new MemoryCacheOptions());
-			this.Builder = Builder;
-			this.Storage = Storage;
-			this.Ticker = Clock.AddTicker(TimeSpan.FromSeconds(30.0), TickAsync, Logger);
-			this.Logger = Logger;
+			_logFiles = logFiles;
+			_logEvents = logEvents;
+			_logFileCache = new MemoryCache(new MemoryCacheOptions());
+			_builder = builder;
+			_storage = storage;
+			_ticker = clock.AddTicker<LogFileService>(TimeSpan.FromSeconds(30.0), TickAsync, logger);
+			_logger = logger;
 		}
 
 		/// <inheritdoc/>
-		public Task StartAsync(CancellationToken CancellationToken) => Ticker.StartAsync();
+		public Task StartAsync(CancellationToken cancellationToken) => _ticker.StartAsync();
 
 		/// <inheritdoc/>
-		public async Task StopAsync(CancellationToken CancellationToken)
+		public async Task StopAsync(CancellationToken cancellationToken)
 		{
-			Logger.LogInformation("Stopping log file service");
-			if (Builder.FlushOnShutdown)
+			_logger.LogInformation("Stopping log file service");
+			if (_builder.FlushOnShutdown)
 			{
 				await FlushAsync();
 			}
-			await Ticker.StopAsync();
-			Logger.LogInformation("Log service stopped");
+			await _ticker.StopAsync();
+			_logger.LogInformation("Log service stopped");
 		}
 
 		/// <inheritdoc/>
 		public void Dispose()
 		{
-			LogFileCache.Dispose();
-			Storage.Dispose();
-			Ticker.Dispose();
+			_logFileCache.Dispose();
+			_storage.Dispose();
+			_ticker.Dispose();
 		}
 
 		/// <inheritdoc/>
-		public Task<ILogFile> CreateLogFileAsync(JobId JobId, SessionId? SessionId, LogType Type)
+		public Task<ILogFile> CreateLogFileAsync(JobId jobId, SessionId? sessionId, LogType type)
 		{
-			return LogFiles.CreateLogFileAsync(JobId, SessionId, Type);
+			return _logFiles.CreateLogFileAsync(jobId, sessionId, type);
 		}
 
 		/// <inheritdoc/>
-		public async Task<ILogFile?> GetLogFileAsync(LogId LogFileId)
+		public async Task<ILogFile?> GetLogFileAsync(LogId logFileId)
 		{
-			ILogFile? LogFile = await LogFiles.GetLogFileAsync(LogFileId);
-			if(LogFile != null)
+			ILogFile? logFile = await _logFiles.GetLogFileAsync(logFileId);
+			if(logFile != null)
 			{
-				AddCachedLogFile(LogFile);
+				AddCachedLogFile(logFile);
 			}
-			return LogFile;
+			return logFile;
 		}
 
 		/// <summary>
 		/// Adds a log file to the cache
 		/// </summary>
-		/// <param name="LogFile">The log file to cache</param>
-		void AddCachedLogFile(ILogFile LogFile)
+		/// <param name="logFile">The log file to cache</param>
+		void AddCachedLogFile(ILogFile logFile)
 		{
-			MemoryCacheEntryOptions Options = new MemoryCacheEntryOptions().SetSlidingExpiration(TimeSpan.FromSeconds(30));
-			LogFileCache.Set(LogFile.Id, LogFile, Options);
+			MemoryCacheEntryOptions options = new MemoryCacheEntryOptions().SetSlidingExpiration(TimeSpan.FromSeconds(30));
+			_logFileCache.Set(logFile.Id, logFile, options);
 		}
 
 		/// <summary>
 		/// Gets a cached log file by id
 		/// </summary>
-		/// <param name="LogFileId">The log file id</param>
+		/// <param name="logFileId">The log file id</param>
 		/// <returns>New log file, or null if not found</returns>
-		public async Task<ILogFile?> GetCachedLogFileAsync(LogId LogFileId)
+		public async Task<ILogFile?> GetCachedLogFileAsync(LogId logFileId)
 		{
-			object? LogFile;
-			if (!LogFileCache.TryGetValue(LogFileId, out LogFile))
+			object? logFile;
+			if (!_logFileCache.TryGetValue(logFileId, out logFile))
 			{
-				LogFile = await GetLogFileAsync(LogFileId);
+				logFile = await GetLogFileAsync(logFileId);
 			}
-			return (ILogFile?)LogFile;
+			return (ILogFile?)logFile;
 		}
 
 		/// <inheritdoc/>
-		public Task<List<ILogFile>> GetLogFilesAsync(int? Index = null, int? Count = null)
+		public Task<List<ILogFile>> GetLogFilesAsync(int? index = null, int? count = null)
 		{
-			return LogFiles.GetLogFilesAsync(Index, Count);
+			return _logFiles.GetLogFilesAsync(index, count);
 		}
 
 		class WriteState
 		{
-			public long Offset;
-			public int LineIndex;
-			public ReadOnlyMemory<byte> Memory;
+			public long _offset;
+			public int _lineIndex;
+			public ReadOnlyMemory<byte> _memory;
 
-			public WriteState(long Offset, int LineIndex, ReadOnlyMemory<byte> Memory)
+			public WriteState(long offset, int lineIndex, ReadOnlyMemory<byte> memory)
 			{
-				this.Offset = Offset;
-				this.LineIndex = LineIndex;
-				this.Memory = Memory;
+				_offset = offset;
+				_lineIndex = lineIndex;
+				_memory = memory;
 			}
 		}
 
 		/// <inheritdoc/>
-		public async Task<ILogFile?> WriteLogDataAsync(ILogFile LogFile, long Offset, int LineIndex, ReadOnlyMemory<byte> Data, bool Flush, int MaxChunkLength, int MaxSubChunkLineCount)
+		public async Task<ILogFile?> WriteLogDataAsync(ILogFile logFile, long offset, int lineIndex, ReadOnlyMemory<byte> data, bool flush, int maxChunkLength, int maxSubChunkLineCount)
 		{
-			using IScope Scope = GlobalTracer.Instance.BuildSpan("WriteLogDataAsync").StartActive();
-			Scope.Span.SetTag("LogId", LogFile.Id.ToString());
-			Scope.Span.SetTag("Offset", Offset.ToString(CultureInfo.InvariantCulture));
-			Scope.Span.SetTag("Length", Data.Length.ToString(CultureInfo.InvariantCulture));
-			Scope.Span.SetTag("LineIndex", LineIndex.ToString(CultureInfo.InvariantCulture));
+			using IScope scope = GlobalTracer.Instance.BuildSpan("WriteLogDataAsync").StartActive();
+			scope.Span.SetTag("LogId", logFile.Id.ToString());
+			scope.Span.SetTag("Offset", offset.ToString(CultureInfo.InvariantCulture));
+			scope.Span.SetTag("Length", data.Length.ToString(CultureInfo.InvariantCulture));
+			scope.Span.SetTag("LineIndex", lineIndex.ToString(CultureInfo.InvariantCulture));
 
 			// Make sure the data ends in a newline
-			if (Data.Length > 0 && Data.Span[Data.Length - 1] != '\n')
+			if (data.Length > 0 && data.Span[data.Length - 1] != '\n')
 			{
-				throw new ArgumentException("Log data must consist of a whole number of lines", nameof(Data));
+				throw new ArgumentException("Log data must consist of a whole number of lines", nameof(data));
 			}
 
 			// Make sure the line count is a power of two
-			if ((MaxSubChunkLineCount & (MaxSubChunkLineCount - 1)) != 0)
+			if ((maxSubChunkLineCount & (maxSubChunkLineCount - 1)) != 0)
 			{
-				throw new ArgumentException("Maximum line count per sub-chunk must be a power of two", nameof(MaxSubChunkLineCount));
+				throw new ArgumentException("Maximum line count per sub-chunk must be a power of two", nameof(maxSubChunkLineCount));
 			}
 
 			// List of the flushed chunks
-			List<long> CompleteOffsets = new List<long>();
+			List<long> completeOffsets = new List<long>();
 
 			// Add the data to new chunks
-			WriteState State = new WriteState(Offset, LineIndex, Data);
-			while (State.Memory.Length > 0)
+			WriteState state = new WriteState(offset, lineIndex, data);
+			while (state._memory.Length > 0)
 			{
 				// Find an existing chunk to append to
-				int ChunkIdx = LogFile.Chunks.GetChunkForOffset(State.Offset);
-				if (ChunkIdx >= 0)
+				int chunkIdx = logFile.Chunks.GetChunkForOffset(state._offset);
+				if (chunkIdx >= 0)
 				{
-					ILogChunk Chunk = LogFile.Chunks[ChunkIdx];
-					if (await WriteLogChunkDataAsync(LogFile, Chunk, State, CompleteOffsets, MaxChunkLength, MaxSubChunkLineCount))
+					ILogChunk chunk = logFile.Chunks[chunkIdx];
+					if (await WriteLogChunkDataAsync(logFile, chunk, state, completeOffsets, maxChunkLength, maxSubChunkLineCount))
 					{
 						continue;
 					}
 				}
 
 				// Create a new chunk. Ensure that there's a chunk at the start of the file, even if the current write is beyond it.
-				ILogFile? NewLogFile;
-				if (LogFile.Chunks.Count == 0)
+				ILogFile? newLogFile;
+				if (logFile.Chunks.Count == 0)
 				{
-					NewLogFile = await LogFiles.TryAddChunkAsync(LogFile, 0, 0);
+					newLogFile = await _logFiles.TryAddChunkAsync(logFile, 0, 0);
 				}
 				else
 				{
-					NewLogFile = await LogFiles.TryAddChunkAsync(LogFile, State.Offset, State.LineIndex);
+					newLogFile = await _logFiles.TryAddChunkAsync(logFile, state._offset, state._lineIndex);
 				}
 
 				// Try to add a new chunk at the new location
-				if (NewLogFile == null)
+				if (newLogFile == null)
 				{
-					NewLogFile = await LogFiles.GetLogFileAsync(LogFile.Id);
-					if (NewLogFile == null)
+					newLogFile = await _logFiles.GetLogFileAsync(logFile.Id);
+					if (newLogFile == null)
 					{
-						Logger.LogError("Unable to update log file {LogId}", LogFile.Id);
+						_logger.LogError("Unable to update log file {LogId}", logFile.Id);
 						return null;
 					}
-					LogFile = NewLogFile;
+					logFile = newLogFile;
 				}
 				else
 				{
 					// Logger.LogDebug("Added new chunk at offset {Offset} to log {LogId}", State.Offset, LogFile.Id);
-					LogFile = NewLogFile;
+					logFile = newLogFile;
 				}
 			}
 
 			// Flush any pending chunks on this log file
-			if (Flush)
+			if (flush)
 			{
-				foreach(ILogChunk Chunk in LogFile.Chunks)
+				foreach(ILogChunk chunk in logFile.Chunks)
 				{
-					if (Chunk.Length == 0 && !CompleteOffsets.Contains(Chunk.Offset))
+					if (chunk.Length == 0 && !completeOffsets.Contains(chunk.Offset))
 					{
-						await Builder.CompleteChunkAsync(LogFile.Id, Chunk.Offset);
-						CompleteOffsets.Add(Chunk.Offset);
+						await _builder.CompleteChunkAsync(logFile.Id, chunk.Offset);
+						completeOffsets.Add(chunk.Offset);
 					}
 				}
 			}
 
 			// Write all the chunks
-			if (CompleteOffsets.Count > 0 || Flush)
+			if (completeOffsets.Count > 0 || flush)
 			{
-				ILogFile? NewLogFile = await WriteCompleteChunksForLogAsync(LogFile, CompleteOffsets, Flush);
-				if (NewLogFile == null)
+				ILogFile? newLogFile = await WriteCompleteChunksForLogAsync(logFile, completeOffsets, flush);
+				if (newLogFile == null)
 				{
 					return null;
 				}
-				LogFile = NewLogFile;
+				logFile = newLogFile;
 			}
-			return LogFile;
+			return logFile;
 		}
 
 		/// <summary>
 		/// Append data to an existing chunk.
 		/// </summary>
-		/// <param name="LogFile">The log file to append to</param>
-		/// <param name="Chunk">Chunk within the log file to update</param>
-		/// <param name="State">Data remaining to be written</param>
-		/// <param name="CompleteOffsets">List of complete chunks</param>
-		/// <param name="MaxChunkLength">Maximum length of each chunk</param>
-		/// <param name="MaxSubChunkLineCount">Maximum number of lines in each subchunk</param>
+		/// <param name="logFile">The log file to append to</param>
+		/// <param name="chunk">Chunk within the log file to update</param>
+		/// <param name="state">Data remaining to be written</param>
+		/// <param name="completeOffsets">List of complete chunks</param>
+		/// <param name="maxChunkLength">Maximum length of each chunk</param>
+		/// <param name="maxSubChunkLineCount">Maximum number of lines in each subchunk</param>
 		/// <returns>True if data was appended to </returns>
-		private async Task<bool> WriteLogChunkDataAsync(ILogFile LogFile, ILogChunk Chunk, WriteState State, List<long> CompleteOffsets, int MaxChunkLength, int MaxSubChunkLineCount)
+		private async Task<bool> WriteLogChunkDataAsync(ILogFile logFile, ILogChunk chunk, WriteState state, List<long> completeOffsets, int maxChunkLength, int maxSubChunkLineCount)
 		{
 			// Don't allow data to be appended if the chunk is complete
-			if(Chunk.Length > 0)
+			if(chunk.Length > 0)
 			{
 				return false;
 			}
@@ -712,50 +669,50 @@ namespace Horde.Build.Services
 			for (; ; )
 			{
 				// Flush the current sub-chunk if we're on a boundary
-				if (State.LineIndex > 0 && (State.LineIndex & (MaxSubChunkLineCount - 1)) == 0)
+				if (state._lineIndex > 0 && (state._lineIndex & (maxSubChunkLineCount - 1)) == 0)
 				{
-					Logger.LogDebug("Completing log {LogId} chunk offset {Offset} sub-chunk at line {LineIndex}", LogFile.Id, Chunk.Offset, State.LineIndex);
-					await Builder.CompleteSubChunkAsync(LogFile.Id, Chunk.Offset);
+					_logger.LogDebug("Completing log {LogId} chunk offset {Offset} sub-chunk at line {LineIndex}", logFile.Id, chunk.Offset, state._lineIndex);
+					await _builder.CompleteSubChunkAsync(logFile.Id, chunk.Offset);
 				}
 
 				// Figure out the max length to write to the current chunk
-				int MaxLength = Math.Min((int)((Chunk.Offset + MaxChunkLength) - State.Offset), State.Memory.Length);
+				int maxLength = Math.Min((int)((chunk.Offset + maxChunkLength) - state._offset), state._memory.Length);
 
 				// Figure out the maximum line index for the current sub chunk
-				int MinLineIndex = State.LineIndex;
-				int MaxLineIndex = (MinLineIndex & ~(MaxSubChunkLineCount - 1)) + MaxSubChunkLineCount;
+				int minLineIndex = state._lineIndex;
+				int maxLineIndex = (minLineIndex & ~(maxSubChunkLineCount - 1)) + maxSubChunkLineCount;
 
 				// Append this data
-				(int Length, int LineCount) = GetWriteLength(State.Memory.Span, MaxLength, MaxLineIndex - MinLineIndex, State.Offset == Chunk.Offset);
-				if (Length > 0)
+				(int length, int lineCount) = GetWriteLength(state._memory.Span, maxLength, maxLineIndex - minLineIndex, state._offset == chunk.Offset);
+				if (length > 0)
 				{
 					// Append this data
-					ReadOnlyMemory<byte> AppendData = State.Memory.Slice(0, Length);
-					if (!await Builder.AppendAsync(LogFile.Id, Chunk.Offset, State.Offset, State.LineIndex, LineCount, AppendData, LogFile.Type))
+					ReadOnlyMemory<byte> appendData = state._memory.Slice(0, length);
+					if (!await _builder.AppendAsync(logFile.Id, chunk.Offset, state._offset, state._lineIndex, lineCount, appendData, logFile.Type))
 					{
 						break;
 					}
 
 					// Update the state
 					//Logger.LogDebug("Append to log {LogId} chunk offset {Offset} (LineIndex={LineIndex}, LineCount={LineCount}, Offset={WriteOffset}, Length={WriteLength})", LogFile.Id, Chunk.Offset, State.LineIndex, LineCount, State.Offset, Length);
-					State.Offset += Length;
-					State.LineIndex += LineCount;
-					State.Memory = State.Memory.Slice(Length);
+					state._offset += length;
+					state._lineIndex += lineCount;
+					state._memory = state._memory.Slice(length);
 					bResult = true;
 
 					// If this is the end of the data, bail out
-					if(State.Memory.Length == 0)
+					if(state._memory.Length == 0)
 					{
 						break;
 					}
 				}
 
 				// Flush the sub-chunk if it's full
-				if (State.LineIndex < MaxLineIndex)
+				if (state._lineIndex < maxLineIndex)
 				{
-					Logger.LogDebug("Completing chunk for log {LogId} at offset {Offset}", LogFile.Id, Chunk.Offset);
-					await Builder.CompleteChunkAsync(LogFile.Id, Chunk.Offset);
-					CompleteOffsets.Add(Chunk.Offset);
+					_logger.LogDebug("Completing chunk for log {LogId} at offset {Offset}", logFile.Id, chunk.Offset);
+					await _builder.CompleteChunkAsync(logFile.Id, chunk.Offset);
+					completeOffsets.Add(chunk.Offset);
 					break;
 				}
 			}
@@ -765,105 +722,104 @@ namespace Horde.Build.Services
 		/// <summary>
 		/// Get the amount of data to write from the given span
 		/// </summary>
-		/// <param name="Span">Data to write</param>
-		/// <param name="MaxLength">Maximum length of the data to write</param>
-		/// <param name="MaxLineCount">Maximum number of lines to write</param>
+		/// <param name="span">Data to write</param>
+		/// <param name="maxLength">Maximum length of the data to write</param>
+		/// <param name="maxLineCount">Maximum number of lines to write</param>
 		/// <param name="bIsEmptyChunk">Whether the current chunk is empty</param>
 		/// <returns>A tuple consisting of the amount of data to write and number of lines in it</returns>
-		private static (int, int) GetWriteLength(ReadOnlySpan<byte> Span, int MaxLength, int MaxLineCount, bool bIsEmptyChunk)
+		private static (int, int) GetWriteLength(ReadOnlySpan<byte> span, int maxLength, int maxLineCount, bool bIsEmptyChunk)
 		{
-			int Length = 0;
-			int LineCount = 0;
-			for (int Idx = 0; Idx < MaxLength || bIsEmptyChunk; Idx++)
+			int length = 0;
+			int lineCount = 0;
+			for (int idx = 0; idx < maxLength || bIsEmptyChunk; idx++)
 			{
-				if (Span[Idx] == '\n')
+				if (span[idx] == '\n')
 				{
-					Length = Idx + 1;
-					LineCount++;
+					length = idx + 1;
+					lineCount++;
 					bIsEmptyChunk = false;
 
-					if (LineCount >= MaxLineCount)
+					if (lineCount >= maxLineCount)
 					{
 						break;
 					}
 				}
 			}
-			return (Length, LineCount);
+			return (length, lineCount);
 		}
 
 		/// <inheritdoc/>
-		public async Task<LogMetadata> GetMetadataAsync(ILogFile LogFile)
+		public async Task<LogMetadata> GetMetadataAsync(ILogFile logFile)
 		{
-			LogMetadata Metadata = new LogMetadata();
-			if (LogFile.Chunks.Count > 0)
+			LogMetadata metadata = new LogMetadata();
+			if (logFile.Chunks.Count > 0)
 			{
-				ILogChunk Chunk = LogFile.Chunks[LogFile.Chunks.Count - 1];
-				if (LogFile.MaxLineIndex == null || Chunk.Length == 0)
+				ILogChunk chunk = logFile.Chunks[logFile.Chunks.Count - 1];
+				if (logFile.MaxLineIndex == null || chunk.Length == 0)
 				{
-					LogChunkData ChunkData = await ReadChunkAsync(LogFile, LogFile.Chunks.Count - 1);
-					Metadata.Length = Chunk.Offset + ChunkData.Length;
-					Metadata.MaxLineIndex = Chunk.LineIndex + ChunkData.LineCount;
+					LogChunkData chunkData = await ReadChunkAsync(logFile, logFile.Chunks.Count - 1);
+					metadata.Length = chunk.Offset + chunkData.Length;
+					metadata.MaxLineIndex = chunk.LineIndex + chunkData.LineCount;
 				}
 				else
 				{
-					Metadata.Length = Chunk.Offset + Chunk.Length;
-					Metadata.MaxLineIndex = LogFile.MaxLineIndex.Value;
+					metadata.Length = chunk.Offset + chunk.Length;
+					metadata.MaxLineIndex = logFile.MaxLineIndex.Value;
 				}
 			}
-			return Metadata;
+			return metadata;
 		}
 
 		/// <inheritdoc/>
-		public Task CreateEventsAsync(List<NewLogEventData> NewEvents)
+		public Task CreateEventsAsync(List<NewLogEventData> newEvents)
 		{
-			return LogEvents.AddManyAsync(NewEvents);
+			return _logEvents.AddManyAsync(newEvents);
 		}
 
 		/// <inheritdoc/>
-		public Task<List<ILogEvent>> FindLogEventsAsync(ILogFile LogFile, int? Index = null, int? Count = null)
+		public Task<List<ILogEvent>> FindLogEventsAsync(ILogFile logFile, int? index = null, int? count = null)
 		{
-			return LogEvents.FindAsync(LogFile.Id, Index, Count);
+			return _logEvents.FindAsync(logFile.Id, index, count);
 		}
 
 		class LogEventLine : ILogEventLine
 		{
-			LogLevel Level;
+			readonly LogLevel _level;
 			public EventId? EventId { get; }
 			public string Message { get; }
 			public JsonElement Data { get; }
 
-			LogLevel ILogEventLine.Level => Level;
+			LogLevel ILogEventLine.Level => _level;
 
-			public LogEventLine(ReadOnlySpan<byte> Data)
-				: this(JsonSerializer.Deserialize<JsonElement>(Data))
+			public LogEventLine(ReadOnlySpan<byte> data)
+				: this(JsonSerializer.Deserialize<JsonElement>(data))
 			{
 			}
 
-			[SuppressMessage("Design", "CA1031:Do not catch general exception types")]
-			public LogEventLine(JsonElement Data)
+			public LogEventLine(JsonElement data)
 			{
-				this.Data = Data;
+				Data = data;
 
-				JsonElement LevelElement;
-				if (!Data.TryGetProperty("level", out LevelElement) || !Enum.TryParse(LevelElement.GetString(), out Level))
+				JsonElement levelElement;
+				if (!data.TryGetProperty("level", out levelElement) || !Enum.TryParse(levelElement.GetString(), out _level))
 				{
-					Level = LogLevel.Information;
+					_level = LogLevel.Information;
 				}
 
-				JsonElement IdElement;
-				if (Data.TryGetProperty("id", out IdElement))
+				JsonElement idElement;
+				if (data.TryGetProperty("id", out idElement))
 				{
-					int IdValue;
-					if (IdElement.TryGetInt32(out IdValue))
+					int idValue;
+					if (idElement.TryGetInt32(out idValue))
 					{
-						EventId = IdValue;
+						EventId = idValue;
 					}
 				}
 
-				JsonElement MessageElement;
-				if (Data.TryGetProperty("renderedMessage", out MessageElement) || Data.TryGetProperty("message", out MessageElement))
+				JsonElement messageElement;
+				if (data.TryGetProperty("renderedMessage", out messageElement) || data.TryGetProperty("message", out messageElement))
 				{
-					Message = MessageElement.GetString() ?? "(Invalid)";
+					Message = messageElement.GetString() ?? "(Invalid)";
 				}
 				else
 				{
@@ -880,147 +836,147 @@ namespace Horde.Build.Services
 			EventSeverity ILogEventData.Severity => (Lines.Count == 0) ? EventSeverity.Information : (Lines[0].Level == LogLevel.Warning) ? EventSeverity.Warning : EventSeverity.Error;
 			string ILogEventData.Message => String.Join("\n", Lines.Select(x => x.Message));
 
-			public LogEventData(IReadOnlyList<ILogEventLine> Lines)
+			public LogEventData(IReadOnlyList<ILogEventLine> lines)
 			{
-				this.Lines = Lines;
+				Lines = lines;
 			}
 		}
 
 		/// <inheritdoc/>
-		public Task AddSpanToEventsAsync(IEnumerable<ILogEvent> Events, ObjectId SpanId)
+		public Task AddSpanToEventsAsync(IEnumerable<ILogEvent> events, ObjectId spanId)
 		{
-			return LogEvents.AddSpanToEventsAsync(Events, SpanId);
+			return _logEvents.AddSpanToEventsAsync(events, spanId);
 		}
 
 		/// <inheritdoc/>
-		public Task<List<ILogEvent>> FindEventsForSpansAsync(IEnumerable<ObjectId> SpanIds, LogId[]? LogIds, int Index, int Count)
+		public Task<List<ILogEvent>> FindEventsForSpansAsync(IEnumerable<ObjectId> spanIds, LogId[]? logIds, int index, int count)
 		{
-			return LogEvents.FindEventsForSpansAsync(SpanIds, LogIds, Index, Count);
+			return _logEvents.FindEventsForSpansAsync(spanIds, logIds, index, count);
 		}
 
 		/// <inheritdoc/>
-		public async Task<ILogEventData> GetEventDataAsync(ILogFile LogFile, int LineIndex, int LineCount)
+		public async Task<ILogEventData> GetEventDataAsync(ILogFile logFile, int lineIndex, int lineCount)
 		{
-			using IScope Scope = GlobalTracer.Instance.BuildSpan("GetEventDataAsync").StartActive();
-			Scope.Span.SetTag("LogId", LogFile.Id.ToString());
-			Scope.Span.SetTag("LineIndex", LineIndex.ToString(CultureInfo.InvariantCulture));
-			Scope.Span.SetTag("LineCount", LineCount.ToString(CultureInfo.InvariantCulture));
+			using IScope scope = GlobalTracer.Instance.BuildSpan("GetEventDataAsync").StartActive();
+			scope.Span.SetTag("LogId", logFile.Id.ToString());
+			scope.Span.SetTag("LineIndex", lineIndex.ToString(CultureInfo.InvariantCulture));
+			scope.Span.SetTag("LineCount", lineCount.ToString(CultureInfo.InvariantCulture));
 
-			(_, long MinOffset) = await GetLineOffsetAsync(LogFile, LineIndex);
-			(_, long MaxOffset) = await GetLineOffsetAsync(LogFile, LineIndex + LineCount);
+			(_, long minOffset) = await GetLineOffsetAsync(logFile, lineIndex);
+			(_, long maxOffset) = await GetLineOffsetAsync(logFile, lineIndex + lineCount);
 
-			byte[] Data = new byte[MaxOffset - MinOffset];
-			using (Stream Stream = await OpenRawStreamAsync(LogFile, MinOffset, MaxOffset - MinOffset))
+			byte[] data = new byte[maxOffset - minOffset];
+			using (Stream stream = await OpenRawStreamAsync(logFile, minOffset, maxOffset - minOffset))
 			{
-				int Length = await Stream.ReadAsync(Data.AsMemory());
-				if(Length != Data.Length)
+				int length = await stream.ReadAsync(data.AsMemory());
+				if(length != data.Length)
 				{
-					Logger.LogWarning("Read less than expected from log stream (Expected {Expected}, Got {Got})", Data.Length, Length);
+					_logger.LogWarning("Read less than expected from log stream (Expected {Expected}, Got {Got})", data.Length, length);
 				}
-				return ParseEventData(Data.AsSpan(0, Length));
+				return ParseEventData(data.AsSpan(0, length));
 			}
 		}
 
 		/// <summary>
 		/// Parses event data from a buffer
 		/// </summary>
-		/// <param name="Data">Data to parse</param>
+		/// <param name="data">Data to parse</param>
 		/// <returns>Parsed event data</returns>
-		private ILogEventData ParseEventData(ReadOnlySpan<byte> Data)
+		private ILogEventData ParseEventData(ReadOnlySpan<byte> data)
 		{
-			List<LogEventLine> Lines = new List<LogEventLine>();
+			List<LogEventLine> lines = new List<LogEventLine>();
 
-			ReadOnlySpan<byte> RemainingData = Data;
-			while(RemainingData.Length > 0)
+			ReadOnlySpan<byte> remainingData = data;
+			while(remainingData.Length > 0)
 			{
-				int EndOfLine = RemainingData.IndexOf((byte)'\n');
-				if(EndOfLine == -1)
+				int endOfLine = remainingData.IndexOf((byte)'\n');
+				if(endOfLine == -1)
 				{
 					break;
 				}
 
-				ReadOnlySpan<byte> LineData = RemainingData.Slice(0, EndOfLine);
+				ReadOnlySpan<byte> lineData = remainingData.Slice(0, endOfLine);
 				try
 				{
-					Lines.Add(new LogEventLine(LineData));
+					lines.Add(new LogEventLine(lineData));
 				}
-				catch(JsonException Ex)
+				catch(JsonException ex)
 				{
-					Logger.LogWarning(Ex, "Unable to parse line from log file: {Line}", Encoding.UTF8.GetString(LineData));
+					_logger.LogWarning(ex, "Unable to parse line from log file: {Line}", Encoding.UTF8.GetString(lineData));
 				}
-				RemainingData = RemainingData.Slice(EndOfLine + 1);
+				remainingData = remainingData.Slice(endOfLine + 1);
 			}
 
-			return new LogEventData(Lines);
+			return new LogEventData(lines);
 		}
 
 		/// <inheritdoc/>
-		public async Task<Stream> OpenRawStreamAsync(ILogFile LogFile, long Offset, long Length)
+		public async Task<Stream> OpenRawStreamAsync(ILogFile logFile, long offset, long length)
 		{
-			if (LogFile.Chunks.Count == 0)
+			if (logFile.Chunks.Count == 0)
 			{
 				return new MemoryStream(Array.Empty<byte>(), false);
 			}
 			else
 			{
-				int LastChunkIdx = LogFile.Chunks.Count - 1;
+				int lastChunkIdx = logFile.Chunks.Count - 1;
 
 				// Clamp the length of the request
-				ILogChunk LastChunk = LogFile.Chunks[LastChunkIdx];
-				if (Length > LastChunk.Offset)
+				ILogChunk lastChunk = logFile.Chunks[lastChunkIdx];
+				if (length > lastChunk.Offset)
 				{
-					long LastChunkLength = LastChunk.Length;
-					if (LastChunkLength <= 0)
+					long lastChunkLength = lastChunk.Length;
+					if (lastChunkLength <= 0)
 					{
-						LogChunkData LastChunkData = await ReadChunkAsync(LogFile, LastChunkIdx);
-						LastChunkLength = LastChunkData.Length;
+						LogChunkData lastChunkData = await ReadChunkAsync(logFile, lastChunkIdx);
+						lastChunkLength = lastChunkData.Length;
 					}
-					Length = Math.Min(Length, (LastChunk.Offset + LastChunkLength) - Offset);
+					length = Math.Min(length, (lastChunk.Offset + lastChunkLength) - offset);
 				}
 
 				// Create the new stream
-				return new ResponseStream(this, LogFile, Offset, Length);
+				return new ResponseStream(this, logFile, offset, length);
 			}
 		}
 
 		/// <inheritdoc/>
-		public async Task<(int, long)> GetLineOffsetAsync(ILogFile LogFile, int LineIdx)
+		public async Task<(int, long)> GetLineOffsetAsync(ILogFile logFile, int lineIdx)
 		{
-			int ChunkIdx = LogFile.Chunks.GetChunkForLine(LineIdx);
+			int chunkIdx = logFile.Chunks.GetChunkForLine(lineIdx);
 
-			ILogChunk Chunk = LogFile.Chunks[ChunkIdx];
-			LogChunkData ChunkData = await ReadChunkAsync(LogFile, ChunkIdx);
+			ILogChunk chunk = logFile.Chunks[chunkIdx];
+			LogChunkData chunkData = await ReadChunkAsync(logFile, chunkIdx);
 
-			if (LineIdx < Chunk.LineIndex)
+			if (lineIdx < chunk.LineIndex)
 			{
-				LineIdx = Chunk.LineIndex;
+				lineIdx = chunk.LineIndex;
 			}
 
-			int MaxLineIndex = Chunk.LineIndex + ChunkData.LineCount;
-			if (LineIdx >= MaxLineIndex)
+			int maxLineIndex = chunk.LineIndex + chunkData.LineCount;
+			if (lineIdx >= maxLineIndex)
 			{
-				LineIdx = MaxLineIndex;
+				lineIdx = maxLineIndex;
 			}
 
-			long Offset = Chunk.Offset + ChunkData.GetLineOffsetWithinChunk(LineIdx - Chunk.LineIndex);
-			return (LineIdx, Offset);
+			long offset = chunk.Offset + chunkData.GetLineOffsetWithinChunk(lineIdx - chunk.LineIndex);
+			return (lineIdx, offset);
 		}
 
 		/// <summary>
 		/// Executes a background task
 		/// </summary>
-		/// <param name="StoppingToken">Cancellation token</param>
-		async ValueTask TickAsync(CancellationToken StoppingToken)
+		/// <param name="stoppingToken">Cancellation token</param>
+		async ValueTask TickAsync(CancellationToken stoppingToken)
 		{
-			lock (WriteLock)
+			lock (_writeLock)
 			{
 				try
 				{
-					WriteTasks.RemoveCompleteTasks();
+					_writeTasks.RemoveCompleteTasks();
 				}
-				catch (Exception Ex)
+				catch (Exception ex)
 				{
-					Logger.LogError(Ex, "Exception while waiting for write tasks to complete");
+					_logger.LogError(ex, "Exception while waiting for write tasks to complete");
 				}
 			}
 			await IncrementalFlush();
@@ -1033,31 +989,30 @@ namespace Horde.Build.Services
 		private async Task IncrementalFlush()
 		{
 			// Get all the chunks older than 20 minutes
-			List<(LogId, long)> FlushChunks = await Builder.TouchChunksAsync(TimeSpan.FromMinutes(10.0));
-			Logger.LogDebug("Performing incremental flush of log builder ({NumChunks} chunks)", FlushChunks.Count);
+			List<(LogId, long)> flushChunks = await _builder.TouchChunksAsync(TimeSpan.FromMinutes(10.0));
+			_logger.LogDebug("Performing incremental flush of log builder ({NumChunks} chunks)", flushChunks.Count);
 
 			// Mark them all as complete
-			foreach ((LogId LogId, long Offset) in FlushChunks)
+			foreach ((LogId logId, long offset) in flushChunks)
 			{
-				await Builder.CompleteChunkAsync(LogId, Offset);
+				await _builder.CompleteChunkAsync(logId, offset);
 			}
 
 			// Add tasks for flushing all the chunks
-			WriteCompleteChunks(FlushChunks, true);
+			WriteCompleteChunks(flushChunks, true);
 		}
 
-
 		/// <summary>
 		/// Flushes the write cache
 		/// </summary>
 		/// <returns>Async task</returns>
 		public async Task FlushAsync()
 		{
-			Logger.LogInformation("Forcing flush of pending log chunks...");
+			_logger.LogInformation("Forcing flush of pending log chunks...");
 
 			// Mark everything in the cache as complete
-			List<(LogId, long)> WriteChunks = await Builder.TouchChunksAsync(TimeSpan.Zero);
-			WriteCompleteChunks(WriteChunks, true);
+			List<(LogId, long)> writeChunks = await _builder.TouchChunksAsync(TimeSpan.Zero);
+			WriteCompleteChunks(writeChunks, true);
 
 			// Wait for everything to flush
 			await FlushPendingWritesAsync();
@@ -1072,54 +1027,54 @@ namespace Horde.Build.Services
 			for(; ;)
 			{
 				// Capture the current contents of the WriteTasks list
-				List<Task> Tasks;
-				lock (WriteLock)
+				List<Task> tasks;
+				lock (_writeLock)
 				{
-					WriteTasks.RemoveCompleteTasks();
-					Tasks = new List<Task>(WriteTasks);
+					_writeTasks.RemoveCompleteTasks();
+					tasks = new List<Task>(_writeTasks);
 				}
-				if (Tasks.Count == 0)
+				if (tasks.Count == 0)
 				{
 					break;
 				}
 
 				// Also add a delay so we'll periodically refresh the list
-				Tasks.Add(Task.Delay(TimeSpan.FromSeconds(5.0)));
-				await Task.WhenAny(Tasks);
+				tasks.Add(Task.Delay(TimeSpan.FromSeconds(5.0)));
+				await Task.WhenAny(tasks);
 			}
 		}
 
 		/// <summary>
 		/// Adds tasks for writing a list of complete chunks
 		/// </summary>
-		/// <param name="ChunksToWrite">List of chunks to write</param>
+		/// <param name="chunksToWrite">List of chunks to write</param>
 		/// <param name="bCreateIndex">Create an index for the log</param>
-		private void WriteCompleteChunks(List<(LogId, long)> ChunksToWrite, bool bCreateIndex)
+		private void WriteCompleteChunks(List<(LogId, long)> chunksToWrite, bool bCreateIndex)
 		{
-			foreach (IGrouping<LogId, long> Group in ChunksToWrite.GroupBy(x => x.Item1, x => x.Item2))
+			foreach (IGrouping<LogId, long> group in chunksToWrite.GroupBy(x => x.Item1, x => x.Item2))
 			{
-				LogId LogId = Group.Key;
+				LogId logId = group.Key;
 
 				// Find offsets of new chunks to write
-				List<long> Offsets = new List<long>();
-				lock (WriteLock)
+				List<long> offsets = new List<long>();
+				lock (_writeLock)
 				{
-					foreach (long Offset in Group.OrderBy(x => x))
+					foreach (long offset in group.OrderBy(x => x))
 					{
-						if (WriteChunks.Add((LogId, Offset)))
+						if (_writeChunks.Add((logId, offset)))
 						{
-							Offsets.Add(Offset);
+							offsets.Add(offset);
 						}
 					}
 				}
 
 				// Create the write task
-				if (Offsets.Count > 0)
+				if (offsets.Count > 0)
 				{
-					Task Task = Task.Run(() => WriteCompleteChunksForLogAsync(LogId, Offsets, bCreateIndex));
-					lock (WriteLock)
+					Task task = Task.Run(() => WriteCompleteChunksForLogAsync(logId, offsets, bCreateIndex));
+					lock (_writeLock)
 					{
-						WriteTasks.Add(Task);
+						_writeTasks.Add(task);
 					}
 				}
 			}
@@ -1128,354 +1083,350 @@ namespace Horde.Build.Services
 		/// <summary>
 		/// Writes a set of chunks to the database
 		/// </summary>
-		/// <param name="LogId">Log file to update</param>
-		/// <param name="Offsets">Chunks to write</param>
+		/// <param name="logId">Log file to update</param>
+		/// <param name="offsets">Chunks to write</param>
 		/// <param name="bCreateIndex">Whether to create the index for this log</param>
 		/// <returns>Async task</returns>
-		private async Task<ILogFile?> WriteCompleteChunksForLogAsync(LogId LogId, List<long> Offsets, bool bCreateIndex)
+		private async Task<ILogFile?> WriteCompleteChunksForLogAsync(LogId logId, List<long> offsets, bool bCreateIndex)
 		{
-			ILogFile? LogFile = await LogFiles.GetLogFileAsync(LogId);
-			if(LogFile != null)
+			ILogFile? logFile = await _logFiles.GetLogFileAsync(logId);
+			if(logFile != null)
 			{
-				LogFile = await WriteCompleteChunksForLogAsync(LogFile, Offsets, bCreateIndex);
+				logFile = await WriteCompleteChunksForLogAsync(logFile, offsets, bCreateIndex);
 			}
-			return LogFile;
+			return logFile;
 		}
 
 		/// <summary>
 		/// Writes a set of chunks to the database
 		/// </summary>
-		/// <param name="LogFileInterface">Log file to update</param>
-		/// <param name="Offsets">Chunks to write</param>
+		/// <param name="logFileInterface">Log file to update</param>
+		/// <param name="offsets">Chunks to write</param>
 		/// <param name="bCreateIndex">Whether to create the index for this log</param>
 		/// <returns>Async task</returns>
-		[SuppressMessage("Design", "CA1031:Do not catch general exception types", Justification = "<Pending>")]
-		private async Task<ILogFile?> WriteCompleteChunksForLogAsync(ILogFile LogFileInterface, List<long> Offsets, bool bCreateIndex)
+		private async Task<ILogFile?> WriteCompleteChunksForLogAsync(ILogFile logFileInterface, List<long> offsets, bool bCreateIndex)
 		{
 			// Write the data to the storage provider
-			List<Task<LogChunkData?>> ChunkWriteTasks = new List<Task<LogChunkData?>>();
-			foreach (long Offset in Offsets)
+			List<Task<LogChunkData?>> chunkWriteTasks = new List<Task<LogChunkData?>>();
+			foreach (long offset in offsets)
 			{
-				int ChunkIdx = LogFileInterface.Chunks.BinarySearch(x => x.Offset, Offset);
-				if (ChunkIdx >= 0)
+				int chunkIdx = logFileInterface.Chunks.BinarySearch(x => x.Offset, offset);
+				if (chunkIdx >= 0)
 				{
-					Logger.LogDebug("Queuing write of log {LogId} chunk {ChunkIdx} offset {Offset}", LogFileInterface.Id, ChunkIdx, Offset);
-					int LineIndex = LogFileInterface.Chunks[ChunkIdx].LineIndex;
-					ChunkWriteTasks.Add(Task.Run(() => WriteChunkAsync(LogFileInterface.Id, Offset, LineIndex)));
+					_logger.LogDebug("Queuing write of log {LogId} chunk {ChunkIdx} offset {Offset}", logFileInterface.Id, chunkIdx, offset);
+					int lineIndex = logFileInterface.Chunks[chunkIdx].LineIndex;
+					chunkWriteTasks.Add(Task.Run(() => WriteChunkAsync(logFileInterface.Id, offset, lineIndex)));
 				}
 			}
 
 			// Wait for the tasks to complete, periodically updating the log file object
-			ILogFile? LogFile = LogFileInterface;
-			while (ChunkWriteTasks.Count > 0)
+			ILogFile? logFile = logFileInterface;
+			while (chunkWriteTasks.Count > 0)
 			{
 				// Wait for all tasks to be complete OR (any task has completed AND 30 seconds has elapsed)
-				Task AllCompleteTask = Task.WhenAll(ChunkWriteTasks);
-				Task AnyCompleteTask = Task.WhenAny(ChunkWriteTasks);
-				await Task.WhenAny(AllCompleteTask, Task.WhenAll(AnyCompleteTask, Task.Delay(TimeSpan.FromSeconds(30.0))));
+				Task allCompleteTask = Task.WhenAll(chunkWriteTasks);
+				Task anyCompleteTask = Task.WhenAny(chunkWriteTasks);
+				await Task.WhenAny(allCompleteTask, Task.WhenAll(anyCompleteTask, Task.Delay(TimeSpan.FromSeconds(30.0))));
 
 				// Update the log file with the written chunks
-				List<LogChunkData?> WrittenChunks = ChunkWriteTasks.RemoveCompleteTasks();
-				while (LogFile != null)
+				List<LogChunkData?> writtenChunks = chunkWriteTasks.RemoveCompleteTasks();
+				while (logFile != null)
 				{
 					// Update the length of any complete chunks
-					List<CompleteLogChunkUpdate> Updates = new List<CompleteLogChunkUpdate>();
-					foreach (LogChunkData? ChunkData in WrittenChunks)
+					List<CompleteLogChunkUpdate> updates = new List<CompleteLogChunkUpdate>();
+					foreach (LogChunkData? chunkData in writtenChunks)
 					{
-						if (ChunkData != null)
+						if (chunkData != null)
 						{
-							int ChunkIdx = LogFile.Chunks.GetChunkForOffset(ChunkData.Offset);
-							if (ChunkIdx >= 0)
+							int chunkIdx = logFile.Chunks.GetChunkForOffset(chunkData.Offset);
+							if (chunkIdx >= 0)
 							{
-								ILogChunk Chunk = LogFile.Chunks[ChunkIdx];
-								if (Chunk.Offset == ChunkData.Offset)
+								ILogChunk chunk = logFile.Chunks[chunkIdx];
+								if (chunk.Offset == chunkData.Offset)
 								{
-									CompleteLogChunkUpdate Update = new CompleteLogChunkUpdate(ChunkIdx, ChunkData.Length, ChunkData.LineCount);
-									Updates.Add(Update);
+									CompleteLogChunkUpdate update = new CompleteLogChunkUpdate(chunkIdx, chunkData.Length, chunkData.LineCount);
+									updates.Add(update);
 								}
 							}
 						}
 					}
 
 					// Try to apply the updates
-					ILogFile? NewLogFile = await LogFiles.TryCompleteChunksAsync(LogFile, Updates);
-					if (NewLogFile != null)
+					ILogFile? newLogFile = await _logFiles.TryCompleteChunksAsync(logFile, updates);
+					if (newLogFile != null)
 					{
-						LogFile = NewLogFile;
+						logFile = newLogFile;
 						break;
 					}
 
 					// Update the log file
-					LogFile = await GetLogFileAsync(LogFile.Id);
+					logFile = await GetLogFileAsync(logFile.Id);
 				}
 			}
 
 			// Create the index if necessary
-			if (bCreateIndex && LogFile != null)
+			if (bCreateIndex && logFile != null)
 			{
 				try
 				{
-					LogFile = await CreateIndexAsync(LogFile);
+					logFile = await CreateIndexAsync(logFile);
 				}
-				catch(Exception Ex)
+				catch(Exception ex)
 				{
-					Logger.LogError(Ex, "Failed to create index for log {LogId}", LogFileInterface.Id);
+					_logger.LogError(ex, "Failed to create index for log {LogId}", logFileInterface.Id);
 				}
 			}
 
-			return LogFile;
+			return logFile;
 		}
 
 		/// <summary>
 		/// Creates an index for the given log file
 		/// </summary>
-		/// <param name="LogFile">The log file object</param>
+		/// <param name="logFile">The log file object</param>
 		/// <returns>Updated log file</returns>
-		private async Task<ILogFile?> CreateIndexAsync(ILogFile LogFile)
+		private async Task<ILogFile?> CreateIndexAsync(ILogFile logFile)
 		{
-			if(LogFile.Chunks.Count == 0)
+			if(logFile.Chunks.Count == 0)
 			{
-				return LogFile;
+				return logFile;
 			}
 
 			// Get the new length of the log, and early out if it won't be any longer
-			ILogChunk LastChunk = LogFile.Chunks[LogFile.Chunks.Count - 1];
-			if(LastChunk.Offset + LastChunk.Length <= (LogFile.IndexLength ?? 0))
+			ILogChunk lastChunk = logFile.Chunks[logFile.Chunks.Count - 1];
+			if(lastChunk.Offset + lastChunk.Length <= (logFile.IndexLength ?? 0))
 			{
-				return LogFile;
+				return logFile;
 			}
 
 			// Save stats for the index creation
-			using IScope Scope = GlobalTracer.Instance.BuildSpan("CreateIndexAsync").StartActive();
-			Scope.Span.SetTag("LogId", LogFile.Id.ToString());
-			Scope.Span.SetTag("Length", (LastChunk.Offset + LastChunk.Length).ToString(CultureInfo.InvariantCulture));
+			using IScope scope = GlobalTracer.Instance.BuildSpan("CreateIndexAsync").StartActive();
+			scope.Span.SetTag("LogId", logFile.Id.ToString());
+			scope.Span.SetTag("Length", (lastChunk.Offset + lastChunk.Length).ToString(CultureInfo.InvariantCulture));
 
-			long NewLength = 0;
-			int NewLineCount = 0;
+			long newLength = 0;
+			int newLineCount = 0;
 
 			// Read the existing index if there is one
-			List<LogIndexData> Indexes = new List<LogIndexData>();
-			if (LogFile.IndexLength != null)
+			List<LogIndexData> indexes = new List<LogIndexData>();
+			if (logFile.IndexLength != null)
 			{
-				LogIndexData? ExistingIndex = await ReadIndexAsync(LogFile, LogFile.IndexLength.Value);
-				if(ExistingIndex != null)
+				LogIndexData? existingIndex = await ReadIndexAsync(logFile, logFile.IndexLength.Value);
+				if(existingIndex != null)
 				{
-					Indexes.Add(ExistingIndex);
-					NewLineCount = ExistingIndex.LineCount;
+					indexes.Add(existingIndex);
+					newLineCount = existingIndex.LineCount;
 				}
 			}
 
 			// Add all the new chunks
-			int ChunkIdx = LogFile.Chunks.GetChunkForLine(NewLineCount);
-			if (ChunkIdx < 0)
+			int chunkIdx = logFile.Chunks.GetChunkForLine(newLineCount);
+			if (chunkIdx < 0)
 			{
-				int FirstLine = (LogFile.Chunks.Count > 0) ? LogFile.Chunks[0].LineIndex : -1;
-				throw new Exception($"Invalid chunk index {ChunkIdx}. Index.LineCount={NewLineCount}, Chunks={LogFile.Chunks.Count}, First line={FirstLine}");
+				int firstLine = (logFile.Chunks.Count > 0) ? logFile.Chunks[0].LineIndex : -1;
+				throw new Exception($"Invalid chunk index {chunkIdx}. Index.LineCount={newLineCount}, Chunks={logFile.Chunks.Count}, First line={firstLine}");
 			}
 
-			for (; ChunkIdx < LogFile.Chunks.Count; ChunkIdx++)
+			for (; chunkIdx < logFile.Chunks.Count; chunkIdx++)
 			{
-				ILogChunk Chunk = LogFile.Chunks[ChunkIdx];
-				LogChunkData ChunkData = await ReadChunkAsync(LogFile, ChunkIdx);
+				ILogChunk chunk = logFile.Chunks[chunkIdx];
+				LogChunkData chunkData = await ReadChunkAsync(logFile, chunkIdx);
 
-				int SubChunkIdx = ChunkData.GetSubChunkForLine(Math.Max(NewLineCount - Chunk.LineIndex, 0));
-				if(SubChunkIdx < 0)
+				int subChunkIdx = chunkData.GetSubChunkForLine(Math.Max(newLineCount - chunk.LineIndex, 0));
+				if(subChunkIdx < 0)
 				{
-					throw new Exception($"Invalid subchunk index {SubChunkIdx}. Chunk {ChunkIdx}/{LogFile.Chunks.Count}. Index.LineCount={NewLineCount}, Chunk.LineIndex={Chunk.LineIndex}, First subchunk {ChunkData.SubChunkLineIndex[0]}");
+					throw new Exception($"Invalid subchunk index {subChunkIdx}. Chunk {chunkIdx}/{logFile.Chunks.Count}. Index.LineCount={newLineCount}, Chunk.LineIndex={chunk.LineIndex}, First subchunk {chunkData.SubChunkLineIndex[0]}");
 				}
 
-				for (; SubChunkIdx < ChunkData.SubChunks.Count; SubChunkIdx++)
+				for (; subChunkIdx < chunkData.SubChunks.Count; subChunkIdx++)
 				{
-					LogSubChunkData SubChunkData = ChunkData.SubChunks[SubChunkIdx];
-					if (SubChunkData.LineIndex >= NewLineCount)
+					LogSubChunkData subChunkData = chunkData.SubChunks[subChunkIdx];
+					if (subChunkData.LineIndex >= newLineCount)
 					{
 						try
 						{
-							Indexes.Add(SubChunkData.BuildIndex());
+							indexes.Add(subChunkData.BuildIndex());
 						}
-						catch (Exception Ex)
+						catch (Exception ex)
 						{
-							throw new Exception($"Failed to create index block - log {LogFile.Id}, chunk {ChunkIdx} ({LogFile.Chunks.Count}), subchunk {SubChunkIdx} ({ChunkData.SubChunks.Count}), index lines: {NewLineCount}, chunk index: {Chunk.LineIndex}, subchunk index: {Chunk.LineIndex + ChunkData.SubChunkLineIndex[SubChunkIdx]}, subchunk count: {SubChunkData.LineCount}", Ex);
+							throw new Exception($"Failed to create index block - log {logFile.Id}, chunk {chunkIdx} ({logFile.Chunks.Count}), subchunk {subChunkIdx} ({chunkData.SubChunks.Count}), index lines: {newLineCount}, chunk index: {chunk.LineIndex}, subchunk index: {chunk.LineIndex + chunkData.SubChunkLineIndex[subChunkIdx]}, subchunk count: {subChunkData.LineCount}", ex);
 						}
 
-						NewLength = SubChunkData.Offset + SubChunkData.Length;
-						NewLineCount = SubChunkData.LineIndex + SubChunkData.LineCount;
+						newLength = subChunkData.Offset + subChunkData.Length;
+						newLineCount = subChunkData.LineIndex + subChunkData.LineCount;
 					}
 				}
 			}
 
 			// Try to update the log file
-			ILogFile? NewLogFile = LogFile;
-			if (NewLength > (LogFile.IndexLength ?? 0))
+			ILogFile? newLogFile = logFile;
+			if (newLength > (logFile.IndexLength ?? 0))
 			{
-				LogIndexData Index = LogIndexData.Merge(Indexes);
-				Logger.LogDebug("Writing index for log {LogId} covering {Length} (index length {IndexLength})", LogFile.Id, NewLength, Index.GetSerializedSize());
+				LogIndexData index = LogIndexData.Merge(indexes);
+				_logger.LogDebug("Writing index for log {LogId} covering {Length} (index length {IndexLength})", logFile.Id, newLength, index.GetSerializedSize());
 
-				await WriteIndexAsync(LogFile.Id, NewLength, Index);
+				await WriteIndexAsync(logFile.Id, newLength, index);
 
-				while(NewLogFile != null && NewLength > (NewLogFile.IndexLength ?? 0))
+				while(newLogFile != null && newLength > (newLogFile.IndexLength ?? 0))
 				{
-					NewLogFile = await LogFiles.TryUpdateIndexAsync(NewLogFile, NewLength);
-					if(NewLogFile != null)
+					newLogFile = await _logFiles.TryUpdateIndexAsync(newLogFile, newLength);
+					if(newLogFile != null)
 					{
 						break;
 					}
-					NewLogFile = await LogFiles.GetLogFileAsync(LogFile.Id);
+					newLogFile = await _logFiles.GetLogFileAsync(logFile.Id);
 				}
 			}
-			return NewLogFile;
+			return newLogFile;
 		}
 
 		/// <summary>
 		/// Reads a chunk from storage
 		/// </summary>
-		/// <param name="LogFile">Log file to read from</param>
-		/// <param name="ChunkIdx">The chunk to read</param>
+		/// <param name="logFile">Log file to read from</param>
+		/// <param name="chunkIdx">The chunk to read</param>
 		/// <returns>Chunk data</returns>
-		[SuppressMessage("Design", "CA1031:Do not catch general exception types", Justification = "<Pending>")]
-		private async Task<LogChunkData> ReadChunkAsync(ILogFile LogFile, int ChunkIdx)
+		private async Task<LogChunkData> ReadChunkAsync(ILogFile logFile, int chunkIdx)
 		{
-			ILogChunk Chunk = LogFile.Chunks[ChunkIdx];
+			ILogChunk chunk = logFile.Chunks[chunkIdx];
 
 			// Try to read the chunk data from storage
-			LogChunkData? ChunkData = null;
+			LogChunkData? chunkData = null;
 			try
 			{
 				// If the chunk is not yet complete, query the log builder
-				if (Chunk.Length == 0)
+				if (chunk.Length == 0)
 				{
-					ChunkData = await Builder.GetChunkAsync(LogFile.Id, Chunk.Offset, Chunk.LineIndex);
+					chunkData = await _builder.GetChunkAsync(logFile.Id, chunk.Offset, chunk.LineIndex);
 				}
 
 				// Otherwise go directly to the log storage
-				if (ChunkData == null)
+				if (chunkData == null)
 				{
-					ChunkData = await Storage.ReadChunkAsync(LogFile.Id, Chunk.Offset, Chunk.LineIndex);
+					chunkData = await _storage.ReadChunkAsync(logFile.Id, chunk.Offset, chunk.LineIndex);
 				}
 			}
-			catch (Exception Ex)
+			catch (Exception ex)
 			{
-				Logger.LogError(Ex, "Unable to read log {LogId} at offset {Offset}", LogFile.Id, Chunk.Offset);
+				_logger.LogError(ex, "Unable to read log {LogId} at offset {Offset}", logFile.Id, chunk.Offset);
 			}
 
 			// Get the minimum length and line count for the chunk
-			if (ChunkIdx + 1 < LogFile.Chunks.Count)
+			if (chunkIdx + 1 < logFile.Chunks.Count)
 			{
-				ILogChunk NextChunk = LogFile.Chunks[ChunkIdx + 1];
-				ChunkData = RepairChunkData(LogFile, ChunkIdx, ChunkData, (int)(NextChunk.Offset - Chunk.Offset), NextChunk.LineIndex - Chunk.LineIndex);
+				ILogChunk nextChunk = logFile.Chunks[chunkIdx + 1];
+				chunkData = RepairChunkData(logFile, chunkIdx, chunkData, (int)(nextChunk.Offset - chunk.Offset), nextChunk.LineIndex - chunk.LineIndex);
 			}
 			else
 			{
-				if (LogFile.MaxLineIndex != null && Chunk.Length != 0)
+				if (logFile.MaxLineIndex != null && chunk.Length != 0)
 				{
-					ChunkData = RepairChunkData(LogFile, ChunkIdx, ChunkData, Chunk.Length, LogFile.MaxLineIndex.Value - Chunk.LineIndex);
+					chunkData = RepairChunkData(logFile, chunkIdx, chunkData, chunk.Length, logFile.MaxLineIndex.Value - chunk.LineIndex);
 				}
-				else if(ChunkData == null)
+				else if(chunkData == null)
 				{
-					ChunkData = RepairChunkData(LogFile, ChunkIdx, ChunkData, 1024, 1);
+					chunkData = RepairChunkData(logFile, chunkIdx, chunkData, 1024, 1);
 				}
 			}
 
-			return ChunkData;
+			return chunkData;
 		}
 
 		/// <summary>
 		/// Validates the given chunk data, and fix it up if necessary
 		/// </summary>
-		/// <param name="LogFile">The log file instance</param>
-		/// <param name="ChunkIdx">Index of the chunk within the logfile</param>
-		/// <param name="ChunkData">The chunk data that was read</param>
-		/// <param name="Length">Expected length of the data</param>
-		/// <param name="LineCount">Expected number of lines in the data</param>
+		/// <param name="logFile">The log file instance</param>
+		/// <param name="chunkIdx">Index of the chunk within the logfile</param>
+		/// <param name="chunkData">The chunk data that was read</param>
+		/// <param name="length">Expected length of the data</param>
+		/// <param name="lineCount">Expected number of lines in the data</param>
 		/// <returns>Repaired chunk data</returns>
-		LogChunkData RepairChunkData(ILogFile LogFile, int ChunkIdx, LogChunkData? ChunkData, int Length, int LineCount)
+		LogChunkData RepairChunkData(ILogFile logFile, int chunkIdx, LogChunkData? chunkData, int length, int lineCount)
 		{
-			int CurrentLength = 0;
-			int CurrentLineCount = 0;
-			if(ChunkData != null)
+			int currentLength = 0;
+			int currentLineCount = 0;
+			if(chunkData != null)
 			{
-				CurrentLength = ChunkData.Length;
-				CurrentLineCount = ChunkData.LineCount;
+				currentLength = chunkData.Length;
+				currentLineCount = chunkData.LineCount;
 			}
 
-			if (ChunkData == null || CurrentLength < Length || CurrentLineCount < LineCount)
+			if (chunkData == null || currentLength < length || currentLineCount < lineCount)
 			{
-				Logger.LogWarning("Creating placeholder subchunk for log {LogId} chunk {ChunkIdx} (length {Length} vs expected {ExpLength}, lines {LineCount} vs expected {ExpLineCount})", LogFile.Id, ChunkIdx, CurrentLength, Length, CurrentLineCount, LineCount);
+				_logger.LogWarning("Creating placeholder subchunk for log {LogId} chunk {ChunkIdx} (length {Length} vs expected {ExpLength}, lines {LineCount} vs expected {ExpLineCount})", logFile.Id, chunkIdx, currentLength, length, currentLineCount, lineCount);
 
-				List<LogSubChunkData> SubChunks = new List<LogSubChunkData>();
-				if (ChunkData != null && ChunkData.Length < Length && ChunkData.LineCount < LineCount)
+				List<LogSubChunkData> subChunks = new List<LogSubChunkData>();
+				if (chunkData != null && chunkData.Length < length && chunkData.LineCount < lineCount)
 				{
-					SubChunks.AddRange(ChunkData.SubChunks);
+					subChunks.AddRange(chunkData.SubChunks);
 				}
 
-				LogText Text = new LogText();
-				Text.AppendMissingDataInfo(ChunkIdx, LogFile.Chunks[ChunkIdx].Server, Length - CurrentLength, LineCount - CurrentLineCount);
-				SubChunks.Add(new LogSubChunkData(LogFile.Type, CurrentLength, CurrentLineCount, Text));
+				LogText text = new LogText();
+				text.AppendMissingDataInfo(chunkIdx, logFile.Chunks[chunkIdx].Server, length - currentLength, lineCount - currentLineCount);
+				subChunks.Add(new LogSubChunkData(logFile.Type, currentLength, currentLineCount, text));
 
-				ILogChunk Chunk = LogFile.Chunks[ChunkIdx];
-				ChunkData = new LogChunkData(Chunk.Offset, Chunk.LineIndex, SubChunks);
+				ILogChunk chunk = logFile.Chunks[chunkIdx];
+				chunkData = new LogChunkData(chunk.Offset, chunk.LineIndex, subChunks);
 			}
-			return ChunkData;
+			return chunkData;
 		}
 
 		/// <summary>
 		/// Writes a set of chunks to the database
 		/// </summary>
-		/// <param name="LogFileId">Unique id of the log file</param>
-		/// <param name="Offset">Offset of the chunk to write</param>
-		/// <param name="LineIndex">First line index of the chunk</param>
+		/// <param name="logFileId">Unique id of the log file</param>
+		/// <param name="offset">Offset of the chunk to write</param>
+		/// <param name="lineIndex">First line index of the chunk</param>
 		/// <returns>Chunk daata</returns>
-		[SuppressMessage("Design", "CA1031:Do not catch general exception types", Justification = "<Pending>")]
-		private async Task<LogChunkData?> WriteChunkAsync(LogId LogFileId, long Offset, int LineIndex)
+		private async Task<LogChunkData?> WriteChunkAsync(LogId logFileId, long offset, int lineIndex)
 		{
 			// Write the chunk to storage
-			LogChunkData? ChunkData = await Builder.GetChunkAsync(LogFileId, Offset, LineIndex);
-			if (ChunkData == null)
+			LogChunkData? chunkData = await _builder.GetChunkAsync(logFileId, offset, lineIndex);
+			if (chunkData == null)
 			{
-				Logger.LogDebug("Log {LogId} offset {Offset} not found in log builder", LogFileId, Offset);
+				_logger.LogDebug("Log {LogId} offset {Offset} not found in log builder", logFileId, offset);
 			}
 			else
 			{
 				try
 				{
-					await Storage.WriteChunkAsync(LogFileId, ChunkData.Offset, ChunkData);
+					await _storage.WriteChunkAsync(logFileId, chunkData.Offset, chunkData);
 				}
-				catch (Exception Ex)
+				catch (Exception ex)
 				{
-					Logger.LogError(Ex, "Unable to write log {LogId} at offset {Offset}", LogFileId, ChunkData.Offset);
+					_logger.LogError(ex, "Unable to write log {LogId} at offset {Offset}", logFileId, chunkData.Offset);
 				}
 			}
 
 			// Remove it from the log builder
 			try
 			{
-				await Builder.RemoveChunkAsync(LogFileId, Offset);
+				await _builder.RemoveChunkAsync(logFileId, offset);
 			}
-			catch (Exception Ex)
+			catch (Exception ex)
 			{
-				Logger.LogError(Ex, "Unable to remove log {LogId} at offset {Offset} from log builder", LogFileId, Offset);
+				_logger.LogError(ex, "Unable to remove log {LogId} at offset {Offset} from log builder", logFileId, offset);
 			}
 
-			return ChunkData;
+			return chunkData;
 		}
 
 		/// <summary>
 		/// Reads a chunk from storage
 		/// </summary>
-		/// <param name="LogFile">Log file to read from</param>
-		/// <param name="Length">Length of the log covered by the index</param>
+		/// <param name="logFile">Log file to read from</param>
+		/// <param name="length">Length of the log covered by the index</param>
 		/// <returns>Chunk data</returns>
-		[SuppressMessage("Design", "CA1031:Do not catch general exception types", Justification = "<Pending>")]
-		private async Task<LogIndexData?> ReadIndexAsync(ILogFile LogFile, long Length)
+		private async Task<LogIndexData?> ReadIndexAsync(ILogFile logFile, long length)
 		{
 			try
 			{
-				LogIndexData? Index = await Storage.ReadIndexAsync(LogFile.Id, Length);
-				return Index;
+				LogIndexData? index = await _storage.ReadIndexAsync(logFile.Id, length);
+				return index;
 			}
-			catch (Exception Ex)
+			catch (Exception ex)
 			{
-				Logger.LogError(Ex, "Unable to read log {LogId} index at length {Length}", LogFile.Id, Length);
+				_logger.LogError(ex, "Unable to read log {LogId} index at length {Length}", logFile.Id, length);
 				return null;
 			}
 		}
@@ -1483,34 +1434,33 @@ namespace Horde.Build.Services
 		/// <summary>
 		/// Writes an index to the database
 		/// </summary>
-		/// <param name="LogFileId">Unique id of the log file</param>
-		/// <param name="Length">Length of the data covered by the index</param>
-		/// <param name="Index">Index to write</param>
+		/// <param name="logFileId">Unique id of the log file</param>
+		/// <param name="length">Length of the data covered by the index</param>
+		/// <param name="index">Index to write</param>
 		/// <returns>Async task</returns>
-		[SuppressMessage("Design", "CA1031:Do not catch general exception types", Justification = "<Pending>")]
-		private async Task WriteIndexAsync(LogId LogFileId, long Length, LogIndexData Index)
+		private async Task WriteIndexAsync(LogId logFileId, long length, LogIndexData index)
 		{
 			try
 			{
-				await Storage.WriteIndexAsync(LogFileId, Length, Index);
+				await _storage.WriteIndexAsync(logFileId, length, index);
 			}
-			catch(Exception Ex)
+			catch(Exception ex)
 			{
-				Logger.LogError(Ex, "Unable to write index for log {LogId}", LogFileId);
+				_logger.LogError(ex, "Unable to write index for log {LogId}", logFileId);
 			}
 		}
 
 		/// <summary>
 		/// Determines if the user is authorized to perform an action on a particular template
 		/// </summary>
-		/// <param name="LogFile">The template to check</param>
-		/// <param name="User">The principal to authorize</param>
+		/// <param name="logFile">The template to check</param>
+		/// <param name="user">The principal to authorize</param>
 		/// <returns>True if the action is authorized</returns>
-		public static bool AuthorizeForSession(ILogFile LogFile, ClaimsPrincipal User)
+		public static bool AuthorizeForSession(ILogFile logFile, ClaimsPrincipal user)
 		{
-			if(LogFile.SessionId != null)
+			if(logFile.SessionId != null)
 			{
-				return User.HasClaim(HordeClaimTypes.AgentSessionId, LogFile.SessionId.Value.ToString());
+				return user.HasClaim(HordeClaimTypes.AgentSessionId, logFile.SessionId.Value.ToString());
 			}
 			else
 			{
@@ -1519,80 +1469,80 @@ namespace Horde.Build.Services
 		}
 
 		/// <inheritdoc/>
-		public async Task<List<int>> SearchLogDataAsync(ILogFile LogFile, string Text, int FirstLine, int Count, LogSearchStats SearchStats)
+		public async Task<List<int>> SearchLogDataAsync(ILogFile logFile, string text, int firstLine, int count, LogSearchStats searchStats)
 		{
-			Stopwatch Timer = Stopwatch.StartNew();
+			Stopwatch timer = Stopwatch.StartNew();
 
-			using IScope Scope = GlobalTracer.Instance.BuildSpan("SearchLogDataAsync").StartActive();
-			Scope.Span.SetTag("LogId", LogFile.Id.ToString());
-			Scope.Span.SetTag("Text", Text);
-			Scope.Span.SetTag("Count", Count.ToString(CultureInfo.InvariantCulture));
+			using IScope scope = GlobalTracer.Instance.BuildSpan("SearchLogDataAsync").StartActive();
+			scope.Span.SetTag("LogId", logFile.Id.ToString());
+			scope.Span.SetTag("Text", text);
+			scope.Span.SetTag("Count", count.ToString(CultureInfo.InvariantCulture));
 
-			List<int> Results = new List<int>();
-			if (Count > 0)
+			List<int> results = new List<int>();
+			if (count > 0)
 			{
-				IAsyncEnumerator<int> Enumerator = SearchLogDataInternalAsync(LogFile, Text, FirstLine, SearchStats).GetAsyncEnumerator();
-				while (await Enumerator.MoveNextAsync() && Results.Count < Count)
+				IAsyncEnumerator<int> enumerator = SearchLogDataInternalAsync(logFile, text, firstLine, searchStats).GetAsyncEnumerator();
+				while (await enumerator.MoveNextAsync() && results.Count < count)
 				{
-					Results.Add(Enumerator.Current);
+					results.Add(enumerator.Current);
 				}
 			}
 
-			Logger.LogDebug("Search for \"{SearchText}\" in log {LogId} found {NumResults}/{MaxResults} results, took {Time}ms ({@Stats})", Text, LogFile.Id, Results.Count, Count, Timer.ElapsedMilliseconds, SearchStats);
-			return Results;
+			_logger.LogDebug("Search for \"{SearchText}\" in log {LogId} found {NumResults}/{MaxResults} results, took {Time}ms ({@Stats})", text, logFile.Id, results.Count, count, timer.ElapsedMilliseconds, searchStats);
+			return results;
 		}
 
 		/// <inheritdoc/>
-		public async IAsyncEnumerable<int> SearchLogDataInternalAsync(ILogFile LogFile, string Text, int FirstLine, LogSearchStats SearchStats)
+		public async IAsyncEnumerable<int> SearchLogDataInternalAsync(ILogFile logFile, string text, int firstLine, LogSearchStats searchStats)
 		{
-			SearchText SearchText = new SearchText(Text);
+			SearchText searchText = new SearchText(text);
 
 			// Read the index for this log file
-			if (LogFile.IndexLength != null)
+			if (logFile.IndexLength != null)
 			{
-				LogIndexData? IndexData = await ReadIndexAsync(LogFile, LogFile.IndexLength.Value);
-				if(IndexData != null && FirstLine < IndexData.LineCount)
+				LogIndexData? indexData = await ReadIndexAsync(logFile, logFile.IndexLength.Value);
+				if(indexData != null && firstLine < indexData.LineCount)
 				{
-					using IScope IndexScope = GlobalTracer.Instance.BuildSpan("Indexed").StartActive();
-					IndexScope.Span.SetTag("LineCount", IndexData.LineCount.ToString(CultureInfo.InvariantCulture));
+					using IScope indexScope = GlobalTracer.Instance.BuildSpan("Indexed").StartActive();
+					indexScope.Span.SetTag("LineCount", indexData.LineCount.ToString(CultureInfo.InvariantCulture));
 
-					foreach(int LineIndex in IndexData.Search(FirstLine, SearchText, SearchStats))
+					foreach(int lineIndex in indexData.Search(firstLine, searchText, searchStats))
 					{
-						yield return LineIndex;
+						yield return lineIndex;
 					}
 
-					FirstLine = IndexData.LineCount;
+					firstLine = indexData.LineCount;
 				}
 			}
 
 			// Manually search through the rest of the log
-			int ChunkIdx = LogFile.Chunks.GetChunkForLine(FirstLine);
-			for (; ChunkIdx < LogFile.Chunks.Count; ChunkIdx++)
+			int chunkIdx = logFile.Chunks.GetChunkForLine(firstLine);
+			for (; chunkIdx < logFile.Chunks.Count; chunkIdx++)
 			{
-				ILogChunk Chunk = LogFile.Chunks[ChunkIdx];
+				ILogChunk chunk = logFile.Chunks[chunkIdx];
 
 				// Read the chunk data
-				LogChunkData ChunkData = await ReadChunkAsync(LogFile, ChunkIdx);
-				if (FirstLine < ChunkData.LineIndex + ChunkData.LineCount)
+				LogChunkData chunkData = await ReadChunkAsync(logFile, chunkIdx);
+				if (firstLine < chunkData.LineIndex + chunkData.LineCount)
 				{
 					// Find the first sub-chunk we're looking for
-					int SubChunkIdx = 0;
-					if (FirstLine > Chunk.LineIndex)
+					int subChunkIdx = 0;
+					if (firstLine > chunk.LineIndex)
 					{
-						SubChunkIdx = ChunkData.GetSubChunkForLine(FirstLine - Chunk.LineIndex);
+						subChunkIdx = chunkData.GetSubChunkForLine(firstLine - chunk.LineIndex);
 					}
 
 					// Search through the sub-chunks
-					for (; SubChunkIdx < ChunkData.SubChunks.Count; SubChunkIdx++)
+					for (; subChunkIdx < chunkData.SubChunks.Count; subChunkIdx++)
 					{
-						LogSubChunkData SubChunkData = ChunkData.SubChunks[SubChunkIdx];
-						if (FirstLine < SubChunkData.LineIndex + SubChunkData.LineCount)
+						LogSubChunkData subChunkData = chunkData.SubChunks[subChunkIdx];
+						if (firstLine < subChunkData.LineIndex + subChunkData.LineCount)
 						{
 							// Create an index containing just this sub-chunk
-							LogIndexData Index = SubChunkData.BuildIndex();
-							foreach (int LineIndex in Index.Search(FirstLine, SearchText, SearchStats))
+							LogIndexData index = subChunkData.BuildIndex();
+							foreach (int lineIndex in index.Search(firstLine, searchText, searchStats))
 							{
-								yield return LineIndex;
+								yield return lineIndex;
 							}
 						}
 					}
diff --git a/Engine/Source/Programs/Horde/Horde.Build/Logs/Storage/ILogStorage.cs b/Engine/Source/Programs/Horde/Horde.Build/Logs/Storage/ILogStorage.cs
--- a/Engine/Source/Programs/Horde/Horde.Build/Logs/Storage/ILogStorage.cs
+++ b/Engine/Source/Programs/Horde/Horde.Build/Logs/Storage/ILogStorage.cs
@@ -1,15 +1,9 @@
 // Copyright Epic Games, Inc. All Rights Reserved.
 
-using Horde.Build.Models;
-using Horde.Build.Services;
-using Horde.Build.Storage;
-using Horde.Build.Utilities;
-using Microsoft.Extensions.Caching.Memory;
-using MongoDB.Bson;
 using System;
-using System.Collections.Generic;
-using System.Linq;
 using System.Threading.Tasks;
+using Horde.Build.Models;
+using Horde.Build.Utilities;
 
 namespace Horde.Build.Logs
 {
@@ -23,36 +17,36 @@ namespace Horde.Build.Logs
 		/// <summary>
 		/// Attempts to read an index for the given log file
 		/// </summary>
-		/// <param name="LogId">Unique id of the log file</param>
-		/// <param name="Length">Length of the file covered by the index</param>
+		/// <param name="logId">Unique id of the log file</param>
+		/// <param name="length">Length of the file covered by the index</param>
 		/// <returns>Index for the log file</returns>
-		Task<LogIndexData?> ReadIndexAsync(LogId LogId, long Length);
+		Task<LogIndexData?> ReadIndexAsync(LogId logId, long length);
 
 		/// <summary>
 		/// Log file to write an index for
 		/// </summary>
-		/// <param name="LogId">Unique id of the log file</param>
-		/// <param name="Length">Length of the file covered by the index</param>
-		/// <param name="Index">The log file index</param>
+		/// <param name="logId">Unique id of the log file</param>
+		/// <param name="length">Length of the file covered by the index</param>
+		/// <param name="index">The log file index</param>
 		/// <returns>Async task</returns>
-		Task WriteIndexAsync(LogId LogId, long Length, LogIndexData Index);
+		Task WriteIndexAsync(LogId logId, long length, LogIndexData index);
 
 		/// <summary>
 		/// Retrieves an item from the cache
 		/// </summary>
-		/// <param name="LogId">Unique id of the log file</param>
-		/// <param name="Offset">Offset of the chunk to read</param>
-		/// <param name="LineIndex">First line of the chunk</param>
+		/// <param name="logId">Unique id of the log file</param>
+		/// <param name="offset">Offset of the chunk to read</param>
+		/// <param name="lineIndex">First line of the chunk</param>
 		/// <returns>Data for the given key, or null if it's not present</returns>
-		Task<LogChunkData?> ReadChunkAsync(LogId LogId, long Offset, int LineIndex);
+		Task<LogChunkData?> ReadChunkAsync(LogId logId, long offset, int lineIndex);
 
 		/// <summary>
 		/// Writes a chunk to storage
 		/// </summary>
-		/// <param name="LogId">Unique id of the log file</param>
-		/// <param name="Offset">Offset of the chunk to write</param>
-		/// <param name="ChunkData">Information about the chunk data</param>
+		/// <param name="logId">Unique id of the log file</param>
+		/// <param name="offset">Offset of the chunk to write</param>
+		/// <param name="chunkData">Information about the chunk data</param>
 		/// <returns>Async task</returns>
-		Task WriteChunkAsync(LogId LogId, long Offset, LogChunkData ChunkData);
+		Task WriteChunkAsync(LogId logId, long offset, LogChunkData chunkData);
 	}
 }
