Horde: Fix coding conventions within Horde.Build.

diff --git a/Engine/Source/Programs/Horde/Horde.Build/Logs/Storage/Impl/LocalLogStorage.cs b/Engine/Source/Programs/Horde/Horde.Build/Logs/Storage/Impl/LocalLogStorage.cs
--- a/Engine/Source/Programs/Horde/Horde.Build/Logs/Storage/Impl/LocalLogStorage.cs
+++ b/Engine/Source/Programs/Horde/Horde.Build/Logs/Storage/Impl/LocalLogStorage.cs
@@ -1,13 +1,10 @@
 // Copyright Epic Games, Inc. All Rights Reserved.
 
+using System;
+using System.Threading.Tasks;
 using Horde.Build.Models;
 using Horde.Build.Utilities;
 using Microsoft.Extensions.Caching.Memory;
-using MongoDB.Bson;
-using System;
-using System.Collections.Generic;
-using System.Linq;
-using System.Threading.Tasks;
 
 namespace Horde.Build.Logs.Storage
 {
@@ -21,76 +18,76 @@ namespace Horde.Build.Logs.Storage
 		/// <summary>
 		/// The memory cache instance
 		/// </summary>
-		MemoryCache Cache;
+		readonly MemoryCache _cache;
 
 		/// <summary>
 		/// The inner storage
 		/// </summary>
-		ILogStorage Inner;
+		readonly ILogStorage _inner;
 
 		/// <summary>
 		/// Constructor
 		/// </summary>
-		/// <param name="NumItems">Maximum size for the memory cache</param>
-		/// <param name="Inner">The inner storage provider</param>
-		public LocalLogStorage(int NumItems, ILogStorage Inner)
+		/// <param name="numItems">Maximum size for the memory cache</param>
+		/// <param name="inner">The inner storage provider</param>
+		public LocalLogStorage(int numItems, ILogStorage inner)
 		{
-			MemoryCacheOptions Options = new MemoryCacheOptions();
-			Options.SizeLimit = NumItems;
+			MemoryCacheOptions options = new MemoryCacheOptions();
+			options.SizeLimit = numItems;
 
-			this.Cache = new MemoryCache(Options);
-			this.Inner = Inner;
+			_cache = new MemoryCache(options);
+			_inner = inner;
 		}
 
 		/// <inheritdoc/>
 		public void Dispose()
 		{
-			Cache.Dispose();
-			Inner.Dispose();
+			_cache.Dispose();
+			_inner.Dispose();
 		}
 
 		/// <summary>
 		/// Gets the cache key for a particular index
 		/// </summary>
-		/// <param name="LogId">The log file to retrieve an index for</param>
-		/// <param name="Length">Length of the file covered by the index</param>
+		/// <param name="logId">The log file to retrieve an index for</param>
+		/// <param name="length">Length of the file covered by the index</param>
 		/// <returns>Cache key for the index</returns>
-		static string IndexKey(LogId LogId, long Length)
+		static string IndexKey(LogId logId, long length)
 		{
-			return $"{LogId}/index-{Length}";
+			return $"{logId}/index-{length}";
 		}
 
 		/// <summary>
 		/// Gets the cache key for a particular chunk
 		/// </summary>
-		/// <param name="LogId">The log file to retrieve an index for</param>
-		/// <param name="Offset">The chunk offset</param>
+		/// <param name="logId">The log file to retrieve an index for</param>
+		/// <param name="offset">The chunk offset</param>
 		/// <returns>Cache key for the chunk</returns>
-		static string ChunkKey(LogId LogId, long Offset)
+		static string ChunkKey(LogId logId, long offset)
 		{
-			return $"{LogId}/chunk-{Offset}";
+			return $"{logId}/chunk-{offset}";
 		}
 
 		/// <summary>
 		/// Adds an entry to the cache
 		/// </summary>
-		/// <param name="Key">The cache key</param>
-		/// <param name="Value">Value to store</param>
-		void AddEntry(string Key, object? Value)
+		/// <param name="key">The cache key</param>
+		/// <param name="value">Value to store</param>
+		void AddEntry(string key, object? value)
 		{
-			using (ICacheEntry Entry = Cache.CreateEntry(Key))
+			using (ICacheEntry entry = _cache.CreateEntry(key))
 			{
-				if (Value == null)
+				if (value == null)
 				{
-					Entry.SetAbsoluteExpiration(TimeSpan.FromSeconds(30.0));
-					Entry.SetSize(0);
+					entry.SetAbsoluteExpiration(TimeSpan.FromSeconds(30.0));
+					entry.SetSize(0);
 				}
 				else
 				{
-					Entry.SetSlidingExpiration(TimeSpan.FromHours(4));
-					Entry.SetSize(1);
+					entry.SetSlidingExpiration(TimeSpan.FromHours(4));
+					entry.SetSize(1);
 				}
-				Entry.SetValue(Value);
+				entry.SetValue(value);
 			}
 		}
 
@@ -98,50 +95,50 @@ namespace Horde.Build.Logs.Storage
 		/// Reads a value from the cache, or from the inner storage provider, adding the result to the cache
 		/// </summary>
 		/// <typeparam name="T">Type of object to read</typeparam>
-		/// <param name="Key">The cache key</param>
-		/// <param name="ReadInner">Delegate to read from the inner storage provider</param>
+		/// <param name="key">The cache key</param>
+		/// <param name="readInner">Delegate to read from the inner storage provider</param>
 		/// <returns>The retrieved value</returns>
-		async Task<T?> ReadValueAsync<T>(string Key, Func<Task<T?>> ReadInner) where T : class
+		async Task<T?> ReadValueAsync<T>(string key, Func<Task<T?>> readInner) where T : class
 		{
-			T? Value;
-			if (!Cache.TryGetValue(Key, out Value))
+			T? value;
+			if (!_cache.TryGetValue(key, out value))
 			{
 				try
 				{
-					Value = await ReadInner();
+					value = await readInner();
 				}
 				finally
 				{
-					AddEntry(Key, Value);
+					AddEntry(key, value);
 				}
 			}
-			return Value;
+			return value;
 		}
 
 		/// <inheritdoc/>
-		public Task<LogIndexData?> ReadIndexAsync(LogId LogId, long Length)
+		public Task<LogIndexData?> ReadIndexAsync(LogId logId, long length)
 		{
-			return ReadValueAsync(IndexKey(LogId, Length), () => Inner.ReadIndexAsync(LogId, Length));
+			return ReadValueAsync(IndexKey(logId, length), () => _inner.ReadIndexAsync(logId, length));
 		}
 
 		/// <inheritdoc/>
-		public Task WriteIndexAsync(LogId LogId, long Length, LogIndexData IndexData)
+		public Task WriteIndexAsync(LogId logId, long length, LogIndexData indexData)
 		{
-			AddEntry(IndexKey(LogId, Length), IndexData);
-			return Inner.WriteIndexAsync(LogId, Length, IndexData);
+			AddEntry(IndexKey(logId, length), indexData);
+			return _inner.WriteIndexAsync(logId, length, indexData);
 		}
 
 		/// <inheritdoc/>
-		public Task<LogChunkData?> ReadChunkAsync(LogId LogId, long Offset, int LineIndex)
+		public Task<LogChunkData?> ReadChunkAsync(LogId logId, long offset, int lineIndex)
 		{
-			return ReadValueAsync(ChunkKey(LogId, Offset), () => Inner.ReadChunkAsync(LogId, Offset, LineIndex));
+			return ReadValueAsync(ChunkKey(logId, offset), () => _inner.ReadChunkAsync(logId, offset, lineIndex));
 		}
 
 		/// <inheritdoc/>
-		public Task WriteChunkAsync(LogId LogId, long Offset, LogChunkData ChunkData)
+		public Task WriteChunkAsync(LogId logId, long offset, LogChunkData chunkData)
 		{
-			AddEntry(ChunkKey(LogId, Offset), ChunkData);
-			return Inner.WriteChunkAsync(LogId, Offset, ChunkData);
+			AddEntry(ChunkKey(logId, offset), chunkData);
+			return _inner.WriteChunkAsync(logId, offset, chunkData);
 		}
 	}
 }
diff --git a/Engine/Source/Programs/Horde/Horde.Build/Logs/Storage/Impl/NullLogStorage.cs b/Engine/Source/Programs/Horde/Horde.Build/Logs/Storage/Impl/NullLogStorage.cs
--- a/Engine/Source/Programs/Horde/Horde.Build/Logs/Storage/Impl/NullLogStorage.cs
+++ b/Engine/Source/Programs/Horde/Horde.Build/Logs/Storage/Impl/NullLogStorage.cs
@@ -1,12 +1,8 @@
 // Copyright Epic Games, Inc. All Rights Reserved.
 
+using System.Threading.Tasks;
 using Horde.Build.Models;
 using Horde.Build.Utilities;
-using MongoDB.Bson;
-using System;
-using System.Collections.Generic;
-using System.Linq;
-using System.Threading.Tasks;
 
 namespace Horde.Build.Logs.Storage.Impl
 {
@@ -23,25 +19,25 @@ namespace Horde.Build.Logs.Storage.Impl
 		}
 
 		/// <inheritdoc/>
-		public Task<LogIndexData?> ReadIndexAsync(LogId LogId, long Length)
+		public Task<LogIndexData?> ReadIndexAsync(LogId logId, long length)
 		{
 			return Task.FromResult<LogIndexData?>(null);
 		}
 
 		/// <inheritdoc/>
-		public Task WriteIndexAsync(LogId LogId, long Length, LogIndexData Index)
+		public Task WriteIndexAsync(LogId logId, long length, LogIndexData index)
 		{
 			return Task.CompletedTask;
 		}
 
 		/// <inheritdoc/>
-		public Task<LogChunkData?> ReadChunkAsync(LogId LogId, long Offset, int LineIndex)
+		public Task<LogChunkData?> ReadChunkAsync(LogId logId, long offset, int lineIndex)
 		{
 			return Task.FromResult<LogChunkData?>(null);
 		}
 
 		/// <inheritdoc/>
-		public Task WriteChunkAsync(LogId LogId, long Offset, LogChunkData ChunkData)
+		public Task WriteChunkAsync(LogId logId, long offset, LogChunkData chunkData)
 		{
 			return Task.CompletedTask;
 		}
diff --git a/Engine/Source/Programs/Horde/Horde.Build/Logs/Storage/Impl/PersistentLogStorage.cs b/Engine/Source/Programs/Horde/Horde.Build/Logs/Storage/Impl/PersistentLogStorage.cs
--- a/Engine/Source/Programs/Horde/Horde.Build/Logs/Storage/Impl/PersistentLogStorage.cs
+++ b/Engine/Source/Programs/Horde/Horde.Build/Logs/Storage/Impl/PersistentLogStorage.cs
@@ -1,16 +1,12 @@
 // Copyright Epic Games, Inc. All Rights Reserved.
 
+using System;
+using System.Threading.Tasks;
 using EpicGames.Core;
 using Horde.Build.Models;
-using Horde.Build.Services;
 using Horde.Build.Storage;
 using Horde.Build.Utilities;
 using Microsoft.Extensions.Logging;
-using MongoDB.Bson;
-using System;
-using System.Collections.Generic;
-using System.Linq;
-using System.Threading.Tasks;
 
 namespace Horde.Build.Logs.Readers
 {
@@ -24,22 +20,22 @@ namespace Horde.Build.Logs.Readers
 		/// <summary>
 		/// The bulk storage provider to use
 		/// </summary>
-		IStorageBackend StorageProvider;
+		readonly IStorageBackend _storageProvider;
 
 		/// <summary>
 		/// Log provider
 		/// </summary>
-		ILogger Logger;
+		readonly ILogger _logger;
 
 		/// <summary>
 		/// Constructor
 		/// </summary>
-		/// <param name="StorageProvider">The storage provider</param>
-		/// <param name="Logger">Logging provider</param>
-		public PersistentLogStorage(IStorageBackend<PersistentLogStorage> StorageProvider, ILogger<PersistentLogStorage> Logger)
+		/// <param name="storageProvider">The storage provider</param>
+		/// <param name="logger">Logging provider</param>
+		public PersistentLogStorage(IStorageBackend<PersistentLogStorage> storageProvider, ILogger<PersistentLogStorage> logger)
 		{
-			this.StorageProvider = StorageProvider;
-			this.Logger = Logger;
+			_storageProvider = storageProvider;
+			_logger = logger;
 		}
 
 		/// <inheritdoc/>
@@ -48,64 +44,64 @@ namespace Horde.Build.Logs.Readers
 		}
 
 		/// <inheritdoc/>
-		public async Task<LogIndexData?> ReadIndexAsync(LogId LogId, long Length)
+		public async Task<LogIndexData?> ReadIndexAsync(LogId logId, long length)
 		{
-			Logger.LogDebug("Reading log {LogId} index length {Length} from persistent storage", LogId, Length);
+			_logger.LogDebug("Reading log {LogId} index length {Length} from persistent storage", logId, length);
 
-			string Path = $"{LogId}/index_{Length}";
-			ReadOnlyMemory<byte>? Data = await StorageProvider.ReadBytesAsync(Path);
-			if (Data == null)
+			string path = $"{logId}/index_{length}";
+			ReadOnlyMemory<byte>? data = await _storageProvider.ReadBytesAsync(path);
+			if (data == null)
 			{
 				return null;
 			}
-			return LogIndexData.FromMemory(Data.Value);
+			return LogIndexData.FromMemory(data.Value);
 		}
 
 		/// <inheritdoc/>
-		public Task WriteIndexAsync(LogId LogId, long Length, LogIndexData IndexData)
+		public Task WriteIndexAsync(LogId logId, long length, LogIndexData indexData)
 		{
-			Logger.LogDebug("Writing log {LogId} index length {Length} to persistent storage", LogId, Length);
+			_logger.LogDebug("Writing log {LogId} index length {Length} to persistent storage", logId, length);
 
-			string Path = $"{LogId}/index_{Length}";
-			ReadOnlyMemory<byte> Data = IndexData.ToByteArray();
-			return StorageProvider.WriteBytesAsync(Path, Data);
+			string path = $"{logId}/index_{length}";
+			ReadOnlyMemory<byte> data = indexData.ToByteArray();
+			return _storageProvider.WriteBytesAsync(path, data);
 		}
 
 		/// <inheritdoc/>
-		public async Task<LogChunkData?> ReadChunkAsync(LogId LogId, long Offset, int LineIndex)
+		public async Task<LogChunkData?> ReadChunkAsync(LogId logId, long offset, int lineIndex)
 		{
-			Logger.LogDebug("Reading log {LogId} chunk offset {Offset} from persistent storage", LogId, Offset);
+			_logger.LogDebug("Reading log {LogId} chunk offset {Offset} from persistent storage", logId, offset);
 
-			string Path = $"{LogId}/offset_{Offset}";
-			ReadOnlyMemory<byte>? Data = await StorageProvider.ReadBytesAsync(Path);
-			if(Data == null)
+			string path = $"{logId}/offset_{offset}";
+			ReadOnlyMemory<byte>? data = await _storageProvider.ReadBytesAsync(path);
+			if(data == null)
 			{
 				return null;
 			}
 
-			MemoryReader Reader = new MemoryReader(Data.Value);
-			LogChunkData ChunkData = Reader.ReadLogChunkData(Offset, LineIndex);
+			MemoryReader reader = new MemoryReader(data.Value);
+			LogChunkData chunkData = reader.ReadLogChunkData(offset, lineIndex);
 
-			if (Reader.Offset != Data.Value.Length)
+			if (reader.Offset != data.Value.Length)
 			{
-				throw new Exception($"Serialization of persistent chunk {Path} is not at expected offset (expected {Data.Value.Length}, actual {Reader.Offset})");
+				throw new Exception($"Serialization of persistent chunk {path} is not at expected offset (expected {data.Value.Length}, actual {reader.Offset})");
 			}
 
-			return ChunkData;
+			return chunkData;
 		}
 
 		/// <inheritdoc/>
-		public Task WriteChunkAsync(LogId LogId, long Offset, LogChunkData ChunkData)
+		public Task WriteChunkAsync(LogId logId, long offset, LogChunkData chunkData)
 		{
-			Logger.LogDebug("Writing log {LogId} chunk offset {Offset} to persistent storage", LogId, Offset);
+			_logger.LogDebug("Writing log {LogId} chunk offset {Offset} to persistent storage", logId, offset);
 
-			string Path = $"{LogId}/offset_{Offset}";
-			byte[] Data = new byte[ChunkData.GetSerializedSize()];
-			MemoryWriter Writer = new MemoryWriter(Data);
-			Writer.WriteLogChunkData(ChunkData);
-			Writer.CheckOffset(Data.Length);
+			string path = $"{logId}/offset_{offset}";
+			byte[] data = new byte[chunkData.GetSerializedSize()];
+			MemoryWriter writer = new MemoryWriter(data);
+			writer.WriteLogChunkData(chunkData);
+			writer.CheckOffset(data.Length);
 
-			return StorageProvider.WriteBytesAsync(Path, Data);
+			return _storageProvider.WriteBytesAsync(path, data);
 		}
 	}
 }
diff --git a/Engine/Source/Programs/Horde/Horde.Build/Logs/Storage/Impl/RedisLogStorage.cs b/Engine/Source/Programs/Horde/Horde.Build/Logs/Storage/Impl/RedisLogStorage.cs
--- a/Engine/Source/Programs/Horde/Horde.Build/Logs/Storage/Impl/RedisLogStorage.cs
+++ b/Engine/Source/Programs/Horde/Horde.Build/Logs/Storage/Impl/RedisLogStorage.cs
@@ -1,15 +1,11 @@
 // Copyright Epic Games, Inc. All Rights Reserved.
 
+using System;
+using System.Threading.Tasks;
 using Horde.Build.Models;
 using Horde.Build.Utilities;
 using Microsoft.Extensions.Logging;
-using MongoDB.Bson;
 using StackExchange.Redis;
-using System;
-using System.Collections.Generic;
-using System.Diagnostics.CodeAnalysis;
-using System.Linq;
-using System.Threading.Tasks;
 
 namespace Horde.Build.Logs.Storage.Impl
 {
@@ -23,153 +19,151 @@ namespace Horde.Build.Logs.Storage.Impl
 		/// <summary>
 		/// The Redis database
 		/// </summary>
-		IDatabase RedisDb;
+		readonly IDatabase _redisDb;
 
 		/// <summary>
 		/// Logging device
 		/// </summary>
-		ILogger Logger;
+		readonly ILogger _logger;
 
 		/// <summary>
 		/// Inner storage layer
 		/// </summary>
-		ILogStorage Inner;
+		readonly ILogStorage _inner;
 
 		/// <summary>
 		/// Constructor
 		/// </summary>
-		/// <param name="RedisDb">Redis database</param>
-		/// <param name="Logger">Logger instance</param>
-		/// <param name="Inner">Inner storage layer</param>
-		public RedisLogStorage(IDatabase RedisDb, ILogger Logger, ILogStorage Inner)
+		/// <param name="redisDb">Redis database</param>
+		/// <param name="logger">Logger instance</param>
+		/// <param name="inner">Inner storage layer</param>
+		public RedisLogStorage(IDatabase redisDb, ILogger logger, ILogStorage inner)
 		{
-			this.RedisDb = RedisDb;
-			this.Logger = Logger;
-			this.Inner = Inner;
+			_redisDb = redisDb;
+			_logger = logger;
+			_inner = inner;
 		}
 
 		/// <inheritdoc/>
 		public void Dispose()
 		{
-			Inner.Dispose();
+			_inner.Dispose();
 		}
 
 		/// <summary>
 		/// Gets the key for a log file's index
 		/// </summary>
-		/// <param name="LogId">The log file id</param>
-		/// <param name="Length">Length of the file covered by the index</param>
+		/// <param name="logId">The log file id</param>
+		/// <param name="length">Length of the file covered by the index</param>
 		/// <returns>The index key</returns>
-		static string IndexKey(LogId LogId, long Length)
+		static string IndexKey(LogId logId, long length)
 		{
-			return $"log-{LogId}-index-{Length}";
+			return $"log-{logId}-index-{length}";
 		}
 
 		/// <summary>
 		/// Gets the key for a chunk's data
 		/// </summary>
-		/// <param name="LogId">The log file id</param>
-		/// <param name="Offset">Offset of the chunk within the log file</param>
+		/// <param name="logId">The log file id</param>
+		/// <param name="offset">Offset of the chunk within the log file</param>
 		/// <returns>The chunk key</returns>
-		static string ChunkKey(LogId LogId, long Offset)
+		static string ChunkKey(LogId logId, long offset)
 		{
-			return $"log-{LogId}-chunk-{Offset}";
+			return $"log-{logId}-chunk-{offset}";
 		}
 
 		/// <summary>
 		/// Adds an index to the cache
 		/// </summary>
-		/// <param name="Key">Key for the item to store</param>
-		/// <param name="Value">Value to be stored</param>
+		/// <param name="key">Key for the item to store</param>
+		/// <param name="value">Value to be stored</param>
 		/// <returns>Async task</returns>
-		[SuppressMessage("Design", "CA1031:Do not catch general exception types")]
-		async Task AddAsync(string Key, ReadOnlyMemory<byte> Value)
+		async Task AddAsync(string key, ReadOnlyMemory<byte> value)
 		{
 			try
 			{
-				await RedisDb.StringSetAsync(Key, Value, expiry: TimeSpan.FromHours(1.0));
-				Logger.LogDebug("Added key {Key} to Redis cache ({Size} bytes)", Key, Value.Length);
+				await _redisDb.StringSetAsync(key, value, expiry: TimeSpan.FromHours(1.0));
+				_logger.LogDebug("Added key {Key} to Redis cache ({Size} bytes)", key, value.Length);
 			}
-			catch (Exception Ex)
+			catch (Exception ex)
 			{
-				Logger.LogError(Ex, "Error writing Redis key {Key}", Key);
+				_logger.LogError(ex, "Error writing Redis key {Key}", key);
 			}
 		}
 
 		/// <summary>
 		/// Adds an index to the cache
 		/// </summary>
-		/// <param name="Key">Key for the item to store</param>
-		/// <param name="Deserialize">Delegate to deserialize the item</param>
+		/// <param name="key">Key for the item to store</param>
+		/// <param name="deserialize">Delegate to deserialize the item</param>
 		/// <returns>Async task</returns>
-		[SuppressMessage("Design", "CA1031:Do not catch general exception types")]
-		async Task<T?> GetAsync<T>(string Key, Func<ReadOnlyMemory<byte>, T> Deserialize) where T : class
+		async Task<T?> GetAsync<T>(string key, Func<ReadOnlyMemory<byte>, T> deserialize) where T : class
 		{
 			try
 			{
-				RedisValue Value = await RedisDb.StringGetAsync(Key);
-				if (Value.IsNullOrEmpty)
+				RedisValue value = await _redisDb.StringGetAsync(key);
+				if (value.IsNullOrEmpty)
 				{
-					Logger.LogDebug("Redis cache miss for {Key}", Key);
+					_logger.LogDebug("Redis cache miss for {Key}", key);
 					return null;
 				}
 				else
 				{
-					Logger.LogDebug("Redis cache hit for {Key}", Key);
-					return Deserialize(Value);
+					_logger.LogDebug("Redis cache hit for {Key}", key);
+					return deserialize(value);
 				}
 			}
-			catch (Exception Ex)
+			catch (Exception ex)
 			{
-				Logger.LogError(Ex, "Error reading Redis key {Key}", Key);
+				_logger.LogError(ex, "Error reading Redis key {Key}", key);
 				return null;
 			}
 		}
 
 		/// <inheritdoc/>
-		public async Task<LogIndexData?> ReadIndexAsync(LogId LogId, long Length)
+		public async Task<LogIndexData?> ReadIndexAsync(LogId logId, long length)
 		{
-			string Key = IndexKey(LogId, Length);
+			string key = IndexKey(logId, length);
 
-			LogIndexData? IndexData = await GetAsync(Key, Memory => LogIndexData.FromMemory(Memory));
-			if(IndexData == null)
+			LogIndexData? indexData = await GetAsync(key, memory => LogIndexData.FromMemory(memory));
+			if(indexData == null)
 			{
-				IndexData = await Inner.ReadIndexAsync(LogId, Length);
-				if(IndexData != null)
+				indexData = await _inner.ReadIndexAsync(logId, length);
+				if(indexData != null)
 				{
-					await AddAsync(Key, IndexData.ToByteArray());
+					await AddAsync(key, indexData.ToByteArray());
 				}
 			}
-			return IndexData;
+			return indexData;
 		}
 
 		/// <inheritdoc/>
-		public async Task WriteIndexAsync(LogId LogId, long Length, LogIndexData Index)
+		public async Task WriteIndexAsync(LogId logId, long length, LogIndexData index)
 		{
-			await Inner.WriteIndexAsync(LogId, Length, Index);
+			await _inner.WriteIndexAsync(logId, length, index);
 		}
 
 		/// <inheritdoc/>
-		public async Task<LogChunkData?> ReadChunkAsync(LogId LogId, long Offset, int LineIndex)
+		public async Task<LogChunkData?> ReadChunkAsync(LogId logId, long offset, int lineIndex)
 		{
-			string Key = ChunkKey(LogId, Offset);
+			string key = ChunkKey(logId, offset);
 
-			LogChunkData? ChunkData = await GetAsync(Key, Memory => LogChunkData.FromMemory(Memory, Offset, LineIndex));
-			if (ChunkData == null)
+			LogChunkData? chunkData = await GetAsync(key, memory => LogChunkData.FromMemory(memory, offset, lineIndex));
+			if (chunkData == null)
 			{
-				ChunkData = await Inner.ReadChunkAsync(LogId, Offset, LineIndex);
-				if (ChunkData != null)
+				chunkData = await _inner.ReadChunkAsync(logId, offset, lineIndex);
+				if (chunkData != null)
 				{
-					await AddAsync(Key, ChunkData.ToByteArray());
+					await AddAsync(key, chunkData.ToByteArray());
 				}
 			}
-			return ChunkData;
+			return chunkData;
 		}
 
 		/// <inheritdoc/>
-		public async Task WriteChunkAsync(LogId LogId, long Offset, LogChunkData ChunkData)
+		public async Task WriteChunkAsync(LogId logId, long offset, LogChunkData chunkData)
 		{
-			await Inner.WriteChunkAsync(LogId, Offset, ChunkData);
+			await _inner.WriteChunkAsync(logId, offset, chunkData);
 		}
 	}
 }
diff --git a/Engine/Source/Programs/Horde/Horde.Build/Logs/Storage/Impl/SequencedLogStorage.cs b/Engine/Source/Programs/Horde/Horde.Build/Logs/Storage/Impl/SequencedLogStorage.cs
--- a/Engine/Source/Programs/Horde/Horde.Build/Logs/Storage/Impl/SequencedLogStorage.cs
+++ b/Engine/Source/Programs/Horde/Horde.Build/Logs/Storage/Impl/SequencedLogStorage.cs
@@ -1,12 +1,9 @@
 // Copyright Epic Games, Inc. All Rights Reserved.
 
-using Horde.Build.Models;
-using Horde.Build.Utilities;
-using MongoDB.Bson;
-using System;
 using System.Collections.Generic;
-using System.Linq;
 using System.Threading.Tasks;
+using Horde.Build.Models;
+using Horde.Build.Utilities;
 
 namespace Horde.Build.Logs.Storage.Impl
 {
@@ -20,177 +17,177 @@ namespace Horde.Build.Logs.Storage.Impl
 		/// <summary>
 		/// Inner storage implementation
 		/// </summary>
-		ILogStorage Inner;
+		readonly ILogStorage _inner;
 
 		/// <summary>
 		/// Pending index reads
 		/// </summary>
-		Dictionary<(LogId, long), Task<LogIndexData?>> IndexReadTasks = new Dictionary<(LogId, long), Task<LogIndexData?>>();
+		readonly Dictionary<(LogId, long), Task<LogIndexData?>> _indexReadTasks = new Dictionary<(LogId, long), Task<LogIndexData?>>();
 
 		/// <summary>
 		/// Pending index reads
 		/// </summary>
-		Dictionary<(LogId, long), Task> IndexWriteTasks = new Dictionary<(LogId, long), Task>();
+		readonly Dictionary<(LogId, long), Task> _indexWriteTasks = new Dictionary<(LogId, long), Task>();
 
 		/// <summary>
 		/// Pending chunk reads
 		/// </summary>
-		Dictionary<(LogId, long), Task<LogChunkData?>> ChunkReadTasks = new Dictionary<(LogId, long), Task<LogChunkData?>>();
+		readonly Dictionary<(LogId, long), Task<LogChunkData?>> _chunkReadTasks = new Dictionary<(LogId, long), Task<LogChunkData?>>();
 
 		/// <summary>
 		/// Pending chunk reads
 		/// </summary>
-		Dictionary<(LogId, long), Task> ChunkWriteTasks = new Dictionary<(LogId, long), Task>();
+		readonly Dictionary<(LogId, long), Task> _chunkWriteTasks = new Dictionary<(LogId, long), Task>();
 
 		/// <summary>
 		/// Constructor
 		/// </summary>
-		/// <param name="Inner">The inner storage provider</param>
-		public SequencedLogStorage(ILogStorage Inner)
+		/// <param name="inner">The inner storage provider</param>
+		public SequencedLogStorage(ILogStorage inner)
 		{
-			this.Inner = Inner;
+			_inner = inner;
 		}
 
 		/// <inheritdoc/>
 		public void Dispose()
 		{
-			Inner.Dispose();
+			_inner.Dispose();
 		}
 
 		/// <inheritdoc/>
-		public Task<LogIndexData?> ReadIndexAsync(LogId LogId, long Length)
+		public Task<LogIndexData?> ReadIndexAsync(LogId logId, long length)
 		{
-			Task<LogIndexData?>? Task;
-			lock (IndexReadTasks)
+			Task<LogIndexData?>? task;
+			lock (_indexReadTasks)
 			{
-				if (!IndexReadTasks.TryGetValue((LogId, Length), out Task))
+				if (!_indexReadTasks.TryGetValue((logId, length), out task))
 				{
-					Task = InnerReadIndexAsync(LogId, Length);
-					IndexReadTasks.Add((LogId, Length), Task);
+					task = InnerReadIndexAsync(logId, length);
+					_indexReadTasks.Add((logId, length), task);
 				}
 			}
-			return Task;
+			return task;
 		}
 
 		/// <inheritdoc/>
-		public Task WriteIndexAsync(LogId LogId, long Length, LogIndexData IndexData)
+		public Task WriteIndexAsync(LogId logId, long length, LogIndexData indexData)
 		{
-			Task? Task;
-			lock (IndexWriteTasks)
+			Task? task;
+			lock (_indexWriteTasks)
 			{
-				if (!IndexWriteTasks.TryGetValue((LogId, Length), out Task))
+				if (!_indexWriteTasks.TryGetValue((logId, length), out task))
 				{
-					Task = InnerWriteIndexAsync(LogId, Length, IndexData);
-					IndexWriteTasks.Add((LogId, Length), Task);
+					task = InnerWriteIndexAsync(logId, length, indexData);
+					_indexWriteTasks.Add((logId, length), task);
 				}
 			}
-			return Task;
+			return task;
 		}
 
 		/// <inheritdoc/>
-		public Task<LogChunkData?> ReadChunkAsync(LogId LogId, long Offset, int LineIndex)
+		public Task<LogChunkData?> ReadChunkAsync(LogId logId, long offset, int lineIndex)
 		{
-			Task<LogChunkData?>? Task;
-			lock (ChunkReadTasks)
+			Task<LogChunkData?>? task;
+			lock (_chunkReadTasks)
 			{
-				if (!ChunkReadTasks.TryGetValue((LogId, Offset), out Task))
+				if (!_chunkReadTasks.TryGetValue((logId, offset), out task))
 				{
-					Task = InnerReadChunkAsync(LogId, Offset, LineIndex);
-					ChunkReadTasks[(LogId, Offset)] = Task;
+					task = InnerReadChunkAsync(logId, offset, lineIndex);
+					_chunkReadTasks[(logId, offset)] = task;
 				}
 			}
-			return Task;
+			return task;
 		}
 
 		/// <inheritdoc/>
-		public Task WriteChunkAsync(LogId LogId, long Offset, LogChunkData ChunkData)
+		public Task WriteChunkAsync(LogId logId, long offset, LogChunkData chunkData)
 		{
-			Task? Task;
-			lock (ChunkWriteTasks)
+			Task? task;
+			lock (_chunkWriteTasks)
 			{
-				if (!ChunkWriteTasks.TryGetValue((LogId, Offset), out Task))
+				if (!_chunkWriteTasks.TryGetValue((logId, offset), out task))
 				{
-					Task = InnerWriteChunkAsync(LogId, Offset, ChunkData);
-					ChunkWriteTasks[(LogId, Offset)] = Task;
+					task = InnerWriteChunkAsync(logId, offset, chunkData);
+					_chunkWriteTasks[(logId, offset)] = task;
 				}
 			}
-			return Task;
+			return task;
 		}
 
 		/// <summary>
 		/// Wrapper for reading an index from the inner storage provider
 		/// </summary>
-		/// <param name="LogId">The log file to read the index for</param>
-		/// <param name="Length">Length of the file that's indexed</param>
+		/// <param name="logId">The log file to read the index for</param>
+		/// <param name="length">Length of the file that's indexed</param>
 		/// <returns>The index data</returns>
-		async Task<LogIndexData?> InnerReadIndexAsync(LogId LogId, long Length)
+		async Task<LogIndexData?> InnerReadIndexAsync(LogId logId, long length)
 		{
 			await Task.Yield();
 
-			LogIndexData? IndexData = await Inner.ReadIndexAsync(LogId, Length);
-			lock (IndexReadTasks)
+			LogIndexData? indexData = await _inner.ReadIndexAsync(logId, length);
+			lock (_indexReadTasks)
 			{
-				IndexReadTasks.Remove((LogId, Length));
+				_indexReadTasks.Remove((logId, length));
 			}
-			return IndexData;
+			return indexData;
 		}
 
 		/// <summary>
 		/// Wrapper for reading an index from the inner storage provider
 		/// </summary>
-		/// <param name="LogId">The log file to read the index for</param>
-		/// <param name="Length">Length of the indexed data</param>
-		/// <param name="IndexData">The index data to write</param>
+		/// <param name="logId">The log file to read the index for</param>
+		/// <param name="length">Length of the indexed data</param>
+		/// <param name="indexData">The index data to write</param>
 		/// <returns>The index data</returns>
-		async Task InnerWriteIndexAsync(LogId LogId, long Length, LogIndexData IndexData)
+		async Task InnerWriteIndexAsync(LogId logId, long length, LogIndexData indexData)
 		{
 			await Task.Yield();
 
-			await Inner.WriteIndexAsync(LogId, Length, IndexData);
+			await _inner.WriteIndexAsync(logId, length, indexData);
 
-			lock (IndexWriteTasks)
+			lock (_indexWriteTasks)
 			{
-				IndexWriteTasks.Remove((LogId, Length));
+				_indexWriteTasks.Remove((logId, length));
 			}
 		}
 
 		/// <summary>
 		/// Wrapper for reading a chunk from the inner storage provider
 		/// </summary>
-		/// <param name="LogId">The log file to read the index for</param>
-		/// <param name="Offset">Offset of the chunk to read</param>
-		/// <param name="LineIndex">Index of the first line in this chunk</param>
+		/// <param name="logId">The log file to read the index for</param>
+		/// <param name="offset">Offset of the chunk to read</param>
+		/// <param name="lineIndex">Index of the first line in this chunk</param>
 		/// <returns>The index data</returns>
-		async Task<LogChunkData?> InnerReadChunkAsync(LogId LogId, long Offset, int LineIndex)
+		async Task<LogChunkData?> InnerReadChunkAsync(LogId logId, long offset, int lineIndex)
 		{
 			await Task.Yield();
 
-			LogChunkData? ChunkData = await Inner.ReadChunkAsync(LogId, Offset, LineIndex);
-			lock (ChunkReadTasks)
+			LogChunkData? chunkData = await _inner.ReadChunkAsync(logId, offset, lineIndex);
+			lock (_chunkReadTasks)
 			{
-				ChunkReadTasks.Remove((LogId, Offset));
+				_chunkReadTasks.Remove((logId, offset));
 			}
-			return ChunkData;
+			return chunkData;
 		}
 
 		/// <summary>
 		/// Wrapper for reading a chunk from the inner storage provider
 		/// </summary>
-		/// <param name="LogId">The log file to write the chunk for</param>
-		/// <param name="Offset">Offset of the chunk within the log</param>
-		/// <param name="ChunkData">The chunk data</param>
+		/// <param name="logId">The log file to write the chunk for</param>
+		/// <param name="offset">Offset of the chunk within the log</param>
+		/// <param name="chunkData">The chunk data</param>
 		/// <returns>The index data</returns>
-		async Task<LogChunkData> InnerWriteChunkAsync(LogId LogId, long Offset, LogChunkData ChunkData)
+		async Task<LogChunkData> InnerWriteChunkAsync(LogId logId, long offset, LogChunkData chunkData)
 		{
 			await Task.Yield();
 
-			await Inner.WriteChunkAsync(LogId, Offset, ChunkData);
+			await _inner.WriteChunkAsync(logId, offset, chunkData);
 
-			lock (ChunkWriteTasks)
+			lock (_chunkWriteTasks)
 			{
-				ChunkWriteTasks.Remove((LogId, Offset));
+				_chunkWriteTasks.Remove((logId, offset));
 			}
-			return ChunkData;
+			return chunkData;
 		}
 	}
 }
diff --git a/Engine/Source/Programs/Horde/Horde.Build/Models/Globals.cs b/Engine/Source/Programs/Horde/Horde.Build/Models/Globals.cs
--- a/Engine/Source/Programs/Horde/Horde.Build/Models/Globals.cs
+++ b/Engine/Source/Programs/Horde/Horde.Build/Models/Globals.cs
@@ -1,24 +1,17 @@
-﻿// Copyright Epic Games, Inc. All Rights Reserved.
+// Copyright Epic Games, Inc. All Rights Reserved.
 
-using Google.Protobuf.WellKnownTypes;
-using Horde.Build.Acls;
-using Horde.Build.Services;
-using Horde.Build.Utilities;
-using Microsoft.IdentityModel.Tokens;
-using MongoDB.Bson;
-using MongoDB.Bson.Serialization.Attributes;
 using System;
 using System.Collections.Generic;
 using System.ComponentModel.DataAnnotations;
-using System.Diagnostics;
 using System.Linq;
 using System.Security.Cryptography;
-using System.Threading.Tasks;
+using Horde.Build.Acls;
+using Horde.Build.Utilities;
+using MongoDB.Bson;
+using MongoDB.Bson.Serialization.Attributes;
 
 namespace Horde.Build.Models
 {
-	using UserId = ObjectId<IUser>;
-
 	/// <summary>
 	/// Base class for singleton documents
 	/// </summary>
@@ -87,65 +80,65 @@ namespace Horde.Build.Models
 		/// <summary>
 		/// Gets the next scheduled downtime
 		/// </summary>
-		/// <param name="Now">The current time</param>
+		/// <param name="now">The current time</param>
 		/// <returns>Start and finish time</returns>
-		public (DateTimeOffset StartTime, DateTimeOffset FinishTime) GetNext(DateTimeOffset Now)
+		public (DateTimeOffset StartTime, DateTimeOffset FinishTime) GetNext(DateTimeOffset now)
 		{
-			TimeSpan Offset = TimeSpan.Zero;
+			TimeSpan offset = TimeSpan.Zero;
 			if (Frequency == ScheduledDowntimeFrequency.Daily)
 			{
-				double Days = (Now - StartTime).TotalDays;
-				if (Days >= 1.0)
+				double days = (now - StartTime).TotalDays;
+				if (days >= 1.0)
 				{
-					Days -= Days % 1.0;
+					days -= days % 1.0;
 				}
-				Offset = TimeSpan.FromDays(Days);
+				offset = TimeSpan.FromDays(days);
 			}
 			else if (Frequency == ScheduledDowntimeFrequency.Weekly)
 			{
-				double Days = (Now - StartTime).TotalDays;
-				if (Days >= 7.0)
+				double days = (now - StartTime).TotalDays;
+				if (days >= 7.0)
 				{
-					Days -= Days % 7.0;
+					days -= days % 7.0;
 				}
-				Offset = TimeSpan.FromDays(Days);
+				offset = TimeSpan.FromDays(days);
 			}
-			return (StartTime + Offset, FinishTime + Offset);
+			return (StartTime + offset, FinishTime + offset);
 		}
 
 		/// <summary>
 		/// Determines if this schedule is active
 		/// </summary>
-		/// <param name="Now">The current time</param>
+		/// <param name="now">The current time</param>
 		/// <returns>True if downtime is active</returns>
-		public bool IsActive(DateTimeOffset Now)
+		public bool IsActive(DateTimeOffset now)
 		{
 			if (Frequency == ScheduledDowntimeFrequency.Once)
 			{
-				return Now >= StartTime && Now < FinishTime;
+				return now >= StartTime && now < FinishTime;
 			}
 			else if (Frequency == ScheduledDowntimeFrequency.Daily)
 			{
-				double Days = (Now - StartTime).TotalDays;
-				if (Days < 0.0)
+				double days = (now - StartTime).TotalDays;
+				if (days < 0.0)
 				{
 					return false;
 				}
 				else
 				{
-					return (Days % 1.0) < (FinishTime - StartTime).TotalDays;
+					return (days % 1.0) < (FinishTime - StartTime).TotalDays;
 				}
 			}
 			else if (Frequency == ScheduledDowntimeFrequency.Weekly)
 			{
-				double Days = (Now - StartTime).TotalDays;
-				if(Days < 0.0)
+				double days = (now - StartTime).TotalDays;
+				if(days < 0.0)
 				{
 					return false;
 				}
 				else
 				{
-					return (Days % 7.0) < (FinishTime - StartTime).TotalDays;
+					return (days % 7.0) < (FinishTime - StartTime).TotalDays;
 				}
 			}
 			else
@@ -346,30 +339,30 @@ namespace Horde.Build.Models
 		/// <summary>
 		/// Finds a perforce cluster with the given name or that contains the provided server
 		/// </summary>
-		/// <param name="Name">Name of the cluster</param>
-		/// <param name="ServerAndPort">Find cluster which contains server</param>
+		/// <param name="name">Name of the cluster</param>
+		/// <param name="serverAndPort">Find cluster which contains server</param>
 		/// <returns></returns>
-		public PerforceCluster? FindPerforceCluster(string? Name, string? ServerAndPort = null)
+		public PerforceCluster? FindPerforceCluster(string? name, string? serverAndPort = null)
 		{
-			List<PerforceCluster> Clusters = PerforceClusters;
+			List<PerforceCluster> clusters = PerforceClusters;
 
-			if (ServerAndPort != null)
+			if (serverAndPort != null)
 			{
-				return Clusters.FirstOrDefault(x => x.Servers.FirstOrDefault( server => string.Equals(server.ServerAndPort, ServerAndPort, StringComparison.OrdinalIgnoreCase)) != null);
+				return clusters.FirstOrDefault(x => x.Servers.FirstOrDefault( server => String.Equals(server.ServerAndPort, serverAndPort, StringComparison.OrdinalIgnoreCase)) != null);
 			}
 
-			if (Clusters.Count == 0)
+			if (clusters.Count == 0)
 			{
-				Clusters = DefaultClusters;
+				clusters = DefaultClusters;
 			}
 
-			if (Name == null)
+			if (name == null)
 			{
-				return Clusters.FirstOrDefault();
+				return clusters.FirstOrDefault();
 			}
 			else
 			{
-				return Clusters.FirstOrDefault(x => String.Equals(x.Name, Name, StringComparison.OrdinalIgnoreCase));
+				return clusters.FirstOrDefault(x => String.Equals(x.Name, name, StringComparison.OrdinalIgnoreCase));
 			}
 		}
 
@@ -377,15 +370,15 @@ namespace Horde.Build.Models
 
 		static List<PerforceCluster> GetDefaultClusters()
 		{
-			PerforceServer Server = new PerforceServer();
-			Server.ServerAndPort = Perforce.P4.P4Server.Get("P4PORT") ?? "perforce:1666";
+			PerforceServer server = new PerforceServer();
+			server.ServerAndPort = Perforce.P4.P4Server.Get("P4PORT") ?? "perforce:1666";
 
-			PerforceCluster Cluster = new PerforceCluster();
-			Cluster.Name = "Default";
-			Cluster.CanImpersonate = false;
-			Cluster.Servers.Add(Server);
+			PerforceCluster cluster = new PerforceCluster();
+			cluster.Name = "Default";
+			cluster.CanImpersonate = false;
+			cluster.Servers.Add(server);
 
-			return new List<PerforceCluster> { Cluster };
+			return new List<PerforceCluster> { cluster };
 		}
 	}
 }
diff --git a/Engine/Source/Programs/Horde/Horde.Build/Notices/INoticeCollection.cs b/Engine/Source/Programs/Horde/Horde.Build/Notices/INoticeCollection.cs
--- a/Engine/Source/Programs/Horde/Horde.Build/Notices/INoticeCollection.cs
+++ b/Engine/Source/Programs/Horde/Horde.Build/Notices/INoticeCollection.cs
@@ -1,11 +1,11 @@
 // Copyright Epic Games, Inc. All Rights Reserved.
 
-using Horde.Build.Models;
-using Horde.Build.Utilities;
-using MongoDB.Bson;
 using System;
 using System.Collections.Generic;
 using System.Threading.Tasks;
+using Horde.Build.Models;
+using Horde.Build.Utilities;
+using MongoDB.Bson;
 
 namespace Horde.Build.Collections
 {
@@ -19,29 +19,29 @@ namespace Horde.Build.Collections
 		/// <summary>
 		/// Add a notice to the collection
 		/// </summary>
-		/// <param name="Message"></param>
-		/// <param name="UserId"></param>
-		/// <param name="StartTime"></param>
-		/// <param name="FinishTime"></param>
+		/// <param name="message"></param>
+		/// <param name="userId"></param>
+		/// <param name="startTime"></param>
+		/// <param name="finishTime"></param>
 		/// <returns></returns>
-		Task<INotice?> AddNoticeAsync(string Message, UserId? UserId, DateTime? StartTime, DateTime? FinishTime);
+		Task<INotice?> AddNoticeAsync(string message, UserId? userId, DateTime? startTime, DateTime? finishTime);
 
 		/// <summary>
 		/// Update an existing notice
 		/// </summary>
-		/// <param name="Id"></param>
-		/// <param name="Message"></param>
-		/// <param name="StartTime"></param>
-		/// <param name="FinishTime"></param>
+		/// <param name="id"></param>
+		/// <param name="message"></param>
+		/// <param name="startTime"></param>
+		/// <param name="finishTime"></param>
 		/// <returns></returns>
-		Task<bool> UpdateNoticeAsync(ObjectId Id, string? Message, DateTime? StartTime, DateTime? FinishTime);
+		Task<bool> UpdateNoticeAsync(ObjectId id, string? message, DateTime? startTime, DateTime? finishTime);
 
 		/// <summary>
 		/// Get a notice by id
 		/// </summary>
-		/// <param name="NoticeId"></param>
+		/// <param name="noticeId"></param>
 		/// <returns></returns>
-		Task<INotice?> GetNoticeAsync(ObjectId NoticeId);
+		Task<INotice?> GetNoticeAsync(ObjectId noticeId);
 
 		/// <summary>
 		/// Get all notices
@@ -52,11 +52,8 @@ namespace Horde.Build.Collections
 		/// <summary>
 		/// Remove a notice
 		/// </summary>
-		/// <param name="Id"></param>
+		/// <param name="id"></param>
 		/// <returns></returns>
-		Task<bool> RemoveNoticeAsync(ObjectId Id);
-
-
+		Task<bool> RemoveNoticeAsync(ObjectId id);
 	}
-
 }
\ No newline at end of file
diff --git a/Engine/Source/Programs/Horde/Horde.Build/Notices/Impl/NoticeCollection.cs b/Engine/Source/Programs/Horde/Horde.Build/Notices/Impl/NoticeCollection.cs
--- a/Engine/Source/Programs/Horde/Horde.Build/Notices/Impl/NoticeCollection.cs
+++ b/Engine/Source/Programs/Horde/Horde.Build/Notices/Impl/NoticeCollection.cs
@@ -1,17 +1,15 @@
 // Copyright Epic Games, Inc. All Rights Reserved.
 
-using Horde.Build.Services;
-using MongoDB.Bson.Serialization.Attributes;
+using System;
+using System.Collections.Generic;
+using System.Linq;
 using System.Threading.Tasks;
-
 using Horde.Build.Models;
-using MongoDB.Driver;
-using System;
+using Horde.Build.Services;
 using Horde.Build.Utilities;
-using System.Linq;
-
 using MongoDB.Bson;
-using System.Collections.Generic;
+using MongoDB.Bson.Serialization.Attributes;
+using MongoDB.Driver;
 
 namespace Horde.Build.Collections.Impl
 {
@@ -45,81 +43,76 @@ namespace Horde.Build.Collections.Impl
 			[BsonConstructor]
 			private NoticeDocument()
 			{
-				
 			}
 
-			public NoticeDocument(ObjectId Id, string Message, UserId? UserId, DateTime? StartTime, DateTime? FinishTime)
+			public NoticeDocument(ObjectId id, string message, UserId? userId, DateTime? startTime, DateTime? finishTime)
 			{
-				this.Id = Id;
-				this.Message = Message;
-				this.UserId = UserId;
-				this.StartTime = StartTime;
-				this.FinishTime = FinishTime;
+				Id = id;
+				Message = message;
+				UserId = userId;
+				StartTime = startTime;
+				FinishTime = finishTime;
 			}
-
 		}
 
-		readonly IMongoCollection<NoticeDocument> Notices;
+		readonly IMongoCollection<NoticeDocument> _notices;
 
 		/// <summary>
 		/// Constructor
 		/// </summary>
-		public NoticeCollection(DatabaseService DatabaseService)
+		public NoticeCollection(DatabaseService databaseService)
 		{
-			Notices = DatabaseService.GetCollection<NoticeDocument>("Notices");
+			_notices = databaseService.GetCollection<NoticeDocument>("Notices");
 		}
 
 		/// <inheritdoc/>
-		public async Task<INotice?> AddNoticeAsync(string Message, UserId? UserId, DateTime? StartTime, DateTime? FinishTime)
+		public async Task<INotice?> AddNoticeAsync(string message, UserId? userId, DateTime? startTime, DateTime? finishTime)
 		{
-			NoticeDocument NewNotice = new NoticeDocument(ObjectId.GenerateNewId(), Message, UserId, StartTime, FinishTime);
-			await Notices.InsertOneAsync(NewNotice);
-			return NewNotice;
+			NoticeDocument newNotice = new NoticeDocument(ObjectId.GenerateNewId(), message, userId, startTime, finishTime);
+			await _notices.InsertOneAsync(newNotice);
+			return newNotice;
 		}
 
 		/// <inheritdoc/>
-		public async Task<bool> UpdateNoticeAsync(ObjectId Id, string? Message, DateTime? StartTime, DateTime? FinishTime)
+		public async Task<bool> UpdateNoticeAsync(ObjectId id, string? message, DateTime? startTime, DateTime? finishTime)
 		{
 
-			UpdateDefinitionBuilder<NoticeDocument> UpdateBuilder = Builders<NoticeDocument>.Update;
-			List<UpdateDefinition<NoticeDocument>> Updates = new List<UpdateDefinition<NoticeDocument>>();
+			UpdateDefinitionBuilder<NoticeDocument> updateBuilder = Builders<NoticeDocument>.Update;
+			List<UpdateDefinition<NoticeDocument>> updates = new List<UpdateDefinition<NoticeDocument>>();
 
-			if (Message != null)
+			if (message != null)
 			{
-				Updates.Add(UpdateBuilder.Set(x => x.Message, Message));
+				updates.Add(updateBuilder.Set(x => x.Message, message));
 			}
 
-			Updates.Add(UpdateBuilder.Set(x => x.StartTime, StartTime));
+			updates.Add(updateBuilder.Set(x => x.StartTime, startTime));
 
-			Updates.Add(UpdateBuilder.Set(x => x.FinishTime, FinishTime));
+			updates.Add(updateBuilder.Set(x => x.FinishTime, finishTime));
 
-			NoticeDocument? Document = await Notices.FindOneAndUpdateAsync(x => x.Id == Id, UpdateBuilder.Combine(Updates));
+			NoticeDocument? document = await _notices.FindOneAndUpdateAsync(x => x.Id == id, updateBuilder.Combine(updates));
 
-			return Document != null;
+			return document != null;
 		}
 
 		/// <inheritdoc/>
-		public async Task<INotice?> GetNoticeAsync(ObjectId NoticeId)
+		public async Task<INotice?> GetNoticeAsync(ObjectId noticeId)
 		{
-			return await Notices.Find(x => x.Id == NoticeId).FirstOrDefaultAsync();
+			return await _notices.Find(x => x.Id == noticeId).FirstOrDefaultAsync();
 		}
 
 		/// <inheritdoc/>
 		public async Task<List<INotice>> GetNoticesAsync()
 		{
-			List<NoticeDocument> Results = await Notices.Find(x => true).ToListAsync();
-			return Results.Select<NoticeDocument, INotice>(x => x).ToList();
+			List<NoticeDocument> results = await _notices.Find(x => true).ToListAsync();
+			return results.Select<NoticeDocument, INotice>(x => x).ToList();
 		}
 
-
 		/// <inheritdoc/>
-		public async Task<bool> RemoveNoticeAsync(ObjectId Id)
+		public async Task<bool> RemoveNoticeAsync(ObjectId id)
 		{
-			DeleteResult Result = await Notices.DeleteOneAsync(x => x.Id == Id);
-			return Result.DeletedCount > 0;
+			DeleteResult result = await _notices.DeleteOneAsync(x => x.Id == id);
+			return result.DeletedCount > 0;
 		}
-
 	}
-
 }
 
diff --git a/Engine/Source/Programs/Horde/Horde.Build/Notices/Impl/NoticeService.cs b/Engine/Source/Programs/Horde/Horde.Build/Notices/Impl/NoticeService.cs
--- a/Engine/Source/Programs/Horde/Horde.Build/Notices/Impl/NoticeService.cs
+++ b/Engine/Source/Programs/Horde/Horde.Build/Notices/Impl/NoticeService.cs
@@ -1,23 +1,12 @@
 // Copyright Epic Games, Inc. All Rights Reserved.
 
+using System;
+using System.Collections.Generic;
 using System.Threading.Tasks;
-
-using Horde.Build.Models;
 using Horde.Build.Collections;
-using System;
+using Horde.Build.Models;
 using Horde.Build.Utilities;
-using System.Threading;
-
-using Microsoft.Extensions.Logging;
-using HordeCommon;
 using MongoDB.Bson;
-using System.Collections.Generic;
-using System.Security.Claims;
-using System.Diagnostics.CodeAnalysis;
-using Horde.Build.Notifications;
-using System.Linq;
-using Horde.Build.Api;
-using Microsoft.Extensions.Hosting;
 
 namespace Horde.Build.Services
 {
@@ -28,60 +17,53 @@ namespace Horde.Build.Services
 	/// </summary>
 	public sealed class NoticeService 
 	{
-		/// <summary>
-		/// Log output writer
-		/// </summary>
-		ILogger<NoticeService> Logger;
-
 		/// <summary>
 		/// Notice collection
 		/// </summary>
-		INoticeCollection Notices;
+		readonly INoticeCollection _notices;
 
 		/// <summary>
 		/// Device service constructor
 		/// </summary>
-		public NoticeService(INoticeCollection Notices, ILogger<NoticeService> Logger)
+		public NoticeService(INoticeCollection notices)
 		{
-			this.Notices = Notices;
-            this.Logger = Logger;			
+			_notices = notices;
 		}
 
 		/// <summary>
 		/// Add a notice
 		/// </summary>
-		/// <param name="Message"></param>
-		/// <param name="UserId"></param>
-		/// <param name="StartTime"></param>
-		/// <param name="FinishTime"></param>
+		/// <param name="message"></param>
+		/// <param name="userId"></param>
+		/// <param name="startTime"></param>
+		/// <param name="finishTime"></param>
 		/// <returns></returns>
-		public Task<INotice?> AddNoticeAsync(string Message, UserId? UserId, DateTime? StartTime, DateTime? FinishTime)
+		public Task<INotice?> AddNoticeAsync(string message, UserId? userId, DateTime? startTime, DateTime? finishTime)
 		{
-			return this.Notices.AddNoticeAsync(Message, UserId, StartTime, FinishTime);
+			return _notices.AddNoticeAsync(message, userId, startTime, finishTime);
 		}
 
 		/// <summary>
 		/// Update an existing notice
 		/// </summary>
-		/// <param name="Id"></param>
-		/// <param name="Message"></param>
-		/// <param name="StartTime"></param>
-		/// <param name="FinishTime"></param>
+		/// <param name="id"></param>
+		/// <param name="message"></param>
+		/// <param name="startTime"></param>
+		/// <param name="finishTime"></param>
 		/// <returns></returns>
-		public Task<bool> UpdateNoticeAsync(ObjectId Id, string? Message, DateTime? StartTime, DateTime? FinishTime)
+		public Task<bool> UpdateNoticeAsync(ObjectId id, string? message, DateTime? startTime, DateTime? finishTime)
 		{
-			return this.Notices.UpdateNoticeAsync(Id, Message, StartTime, FinishTime);
+			return _notices.UpdateNoticeAsync(id, message, startTime, finishTime);
 		}
 
-
 		/// <summary>
 		/// Get a notice by id
 		/// </summary>
-		/// <param name="NoticeId"></param>
+		/// <param name="noticeId"></param>
 		/// <returns></returns>
-		public Task<INotice?> GetNoticeAsync(ObjectId NoticeId)
+		public Task<INotice?> GetNoticeAsync(ObjectId noticeId)
 		{
-			return this.Notices.GetNoticeAsync(NoticeId);
+			return _notices.GetNoticeAsync(noticeId);
 		}
 
 		/// <summary>
@@ -90,18 +72,17 @@ namespace Horde.Build.Services
 		/// <returns></returns>
 		public Task<List<INotice>> GetNoticesAsync()
 		{
-			return this.Notices.GetNoticesAsync();
+			return _notices.GetNoticesAsync();
 		}
 
 		/// <summary>
 		/// Remove an existing notice
 		/// </summary>
-		/// <param name="Id"></param>
+		/// <param name="id"></param>
 		/// <returns></returns>
-		public Task<bool> RemoveNoticeAsync(ObjectId Id)
+		public Task<bool> RemoveNoticeAsync(ObjectId id)
 		{
-			return this.Notices.RemoveNoticeAsync(Id);
+			return _notices.RemoveNoticeAsync(id);
 		}
-
 	}
 }
\ No newline at end of file
diff --git a/Engine/Source/Programs/Horde/Horde.Build/Notifications/IEvent.cs b/Engine/Source/Programs/Horde/Horde.Build/Notifications/IEvent.cs
--- a/Engine/Source/Programs/Horde/Horde.Build/Notifications/IEvent.cs
+++ b/Engine/Source/Programs/Horde/Horde.Build/Notifications/IEvent.cs
@@ -2,11 +2,6 @@
 
 using Horde.Build.Api;
 using Horde.Build.Utilities;
-using System;
-using System.Collections.Generic;
-using System.Diagnostics.CodeAnalysis;
-using System.Linq;
-using System.Threading.Tasks;
 
 namespace Horde.Build.Models
 {
diff --git a/Engine/Source/Programs/Horde/Horde.Build/Notifications/INotificationService.cs b/Engine/Source/Programs/Horde/Horde.Build/Notifications/INotificationService.cs
--- a/Engine/Source/Programs/Horde/Horde.Build/Notifications/INotificationService.cs
+++ b/Engine/Source/Programs/Horde/Horde.Build/Notifications/INotificationService.cs
@@ -1,12 +1,12 @@
 // Copyright Epic Games, Inc. All Rights Reserved.
 
-using Horde.Build.Api;
-using Horde.Build.Models;
-using MongoDB.Bson;
 using System.Collections.Generic;
 using System.Diagnostics.CodeAnalysis;
 using System.Security.Claims;
 using System.Threading.Tasks;
+using Horde.Build.Api;
+using Horde.Build.Models;
+using MongoDB.Bson;
 
 namespace Horde.Build.Notifications
 {
@@ -33,14 +33,14 @@ namespace Horde.Build.Notifications
 		/// <summary>
 		/// Constructor
 		/// </summary>
-		/// <param name="JobId"></param>
-		/// <param name="JobName"></param>
-		/// <param name="PoolName"></param>
-		public JobScheduledNotification(string JobId, string JobName, string PoolName)
+		/// <param name="jobId"></param>
+		/// <param name="jobName"></param>
+		/// <param name="poolName"></param>
+		public JobScheduledNotification(string jobId, string jobName, string poolName)
 		{
-			this.JobId = JobId;
-			this.JobName = JobName;
-			this.PoolName = PoolName;
+			JobId = jobId;
+			JobName = jobName;
+			PoolName = poolName;
 		}
 	}
 	
@@ -52,67 +52,67 @@ namespace Horde.Build.Notifications
 		/// <summary>
 		/// Updates a subscription to a trigger
 		/// </summary>
-		/// <param name="TriggerId"></param>
-		/// <param name="User"></param>
-		/// <param name="Email">Whether to receive email notifications</param>
-		/// <param name="Slack">Whether to receive Slack notifications</param>
+		/// <param name="triggerId"></param>
+		/// <param name="user"></param>
+		/// <param name="email">Whether to receive email notifications</param>
+		/// <param name="slack">Whether to receive Slack notifications</param>
 		/// <returns></returns>
-		Task<bool> UpdateSubscriptionsAsync(ObjectId TriggerId, ClaimsPrincipal User, bool? Email, bool? Slack);
+		Task<bool> UpdateSubscriptionsAsync(ObjectId triggerId, ClaimsPrincipal user, bool? email, bool? slack);
 
 		/// <summary>
 		/// Gets the current subscriptions for a user
 		/// </summary>
-		/// <param name="TriggerId"></param>
-		/// <param name="User"></param>
+		/// <param name="triggerId"></param>
+		/// <param name="user"></param>
 		/// <returns>Subscriptions for that user</returns>
-		Task<INotificationSubscription?> GetSubscriptionsAsync(ObjectId TriggerId, ClaimsPrincipal User);
+		Task<INotificationSubscription?> GetSubscriptionsAsync(ObjectId triggerId, ClaimsPrincipal user);
 
 		/// <summary>
 		/// Notify all subscribers that a job step has finished
 		/// </summary>
-		/// <param name="Job">The job containing the step that has finished</param>
-		/// <param name="Graph"></param>
-		/// <param name="BatchId">The batch id</param>
-		/// <param name="StepId">The step id</param>
+		/// <param name="job">The job containing the step that has finished</param>
+		/// <param name="graph"></param>
+		/// <param name="batchId">The batch id</param>
+		/// <param name="stepId">The step id</param>
 		/// <returns>Async task</returns>
-		void NotifyJobStepComplete(IJob Job, IGraph Graph, SubResourceId BatchId, SubResourceId StepId);
+		void NotifyJobStepComplete(IJob job, IGraph graph, SubResourceId batchId, SubResourceId stepId);
 		
 		/// <summary>
 		/// Notify all subscribers that a job step's outcome has changed
 		/// </summary>
-		/// <param name="Job">The job containing the step that has changed</param>
-		/// <param name="OldLabelStates">The old label states for the job</param>
-		/// <param name="NewLabelStates">The new label states for the job</param>
+		/// <param name="job">The job containing the step that has changed</param>
+		/// <param name="oldLabelStates">The old label states for the job</param>
+		/// <param name="newLabelStates">The new label states for the job</param>
 		/// <returns>Async task</returns>
-		void NotifyLabelUpdate(IJob Job, IReadOnlyList<(LabelState, LabelOutcome)> OldLabelStates, IReadOnlyList<(LabelState, LabelOutcome)> NewLabelStates);
+		void NotifyLabelUpdate(IJob job, IReadOnlyList<(LabelState, LabelOutcome)> oldLabelStates, IReadOnlyList<(LabelState, LabelOutcome)> newLabelStates);
 
 		/// <summary>
 		/// Notify slack channel about a stream update failure
 		/// </summary>
-		/// <param name="ErrorMessage">Error message passed back</param>
-		/// <param name="FileName"></param>
-		/// <param name="Change">Latest change number of the file</param>
-		/// <param name="Author">Author of the change, if found from p4 service</param>
-		/// <param name="Description">Description of the change, if found from p4 service</param>
-		void NotifyConfigUpdateFailure(string ErrorMessage, string FileName, int? Change = null, IUser? Author = null, string? Description = null);
+		/// <param name="errorMessage">Error message passed back</param>
+		/// <param name="fileName"></param>
+		/// <param name="change">Latest change number of the file</param>
+		/// <param name="author">Author of the change, if found from p4 service</param>
+		/// <param name="description">Description of the change, if found from p4 service</param>
+		void NotifyConfigUpdateFailure(string errorMessage, string fileName, int? change = null, IUser? author = null, string? description = null);
 
 		/// <summary>
 		/// Sends a notification to a user regarding a build health issue.
 		/// </summary>
-		/// <param name="Issue">The new issue that was created</param>
-		void NotifyIssueUpdated(IIssue Issue);
+		/// <param name="issue">The new issue that was created</param>
+		void NotifyIssueUpdated(IIssue issue);
 
 		/// <summary>
 		/// Send a device service notification
 		/// </summary>
-		/// <param name="Message">The message to send</param>
-		/// <param name="Device">The device</param>
-		/// <param name="Pool">The pool</param>
-		/// <param name="Stream"></param>
-		/// <param name="Job"></param>
-		/// <param name="Step"></param>
-		/// <param name="Node"></param>
-		/// <param name="User"></param>
-		void NotifyDeviceService(string Message, IDevice? Device = null, IDevicePool? Pool = null, IStream? Stream = null, IJob? Job = null, IJobStep? Step = null, INode? Node = null, IUser? User = null);
+		/// <param name="message">The message to send</param>
+		/// <param name="device">The device</param>
+		/// <param name="pool">The pool</param>
+		/// <param name="stream"></param>
+		/// <param name="job"></param>
+		/// <param name="step"></param>
+		/// <param name="node"></param>
+		/// <param name="user"></param>
+		void NotifyDeviceService(string message, IDevice? device = null, IDevicePool? pool = null, IStream? stream = null, IJob? job = null, IJobStep? step = null, INode? node = null, IUser? user = null);
     }
 }
diff --git a/Engine/Source/Programs/Horde/Horde.Build/Notifications/INotificationSink.cs b/Engine/Source/Programs/Horde/Horde.Build/Notifications/INotificationSink.cs
--- a/Engine/Source/Programs/Horde/Horde.Build/Notifications/INotificationSink.cs
+++ b/Engine/Source/Programs/Horde/Horde.Build/Notifications/INotificationSink.cs
@@ -1,11 +1,11 @@
 // Copyright Epic Games, Inc. All Rights Reserved.
 
-using HordeCommon;
-using Horde.Build.Api;
-using Horde.Build.Models;
 using System;
 using System.Collections.Generic;
 using System.Threading.Tasks;
+using Horde.Build.Api;
+using Horde.Build.Models;
+using HordeCommon;
 
 namespace Horde.Build.Notifications
 {
@@ -17,88 +17,88 @@ namespace Horde.Build.Notifications
 		/// <summary>
 		/// Send notifications that a job has been scheduled
 		/// </summary>
-		/// <param name="Notifications">List of notifications to send</param>
+		/// <param name="notifications">List of notifications to send</param>
 		/// <returns>Async task</returns>
-		Task NotifyJobScheduledAsync(List<JobScheduledNotification> Notifications);
+		Task NotifyJobScheduledAsync(List<JobScheduledNotification> notifications);
 		
 		/// <summary>
 		/// Send notifications that a job has completed
 		/// </summary>
-		/// <param name="JobStream"></param>
-		/// <param name="Job">The job containing the step</param>
-		/// <param name="Graph"></param>
-		/// <param name="Outcome"></param>
+		/// <param name="jobStream"></param>
+		/// <param name="job">The job containing the step</param>
+		/// <param name="graph"></param>
+		/// <param name="outcome"></param>
 		/// <returns>Async task</returns>
-		Task NotifyJobCompleteAsync(IStream JobStream, IJob Job, IGraph Graph, LabelOutcome Outcome);
+		Task NotifyJobCompleteAsync(IStream jobStream, IJob job, IGraph graph, LabelOutcome outcome);
 
 		/// <summary>
 		/// Send notifications that a job has completed
 		/// </summary>
-		/// <param name="User">User to notify</param>
-		/// <param name="JobStream"></param>
-		/// <param name="Job">The job containing the step</param>
-		/// <param name="Graph"></param>
-		/// <param name="Outcome"></param>
+		/// <param name="user">User to notify</param>
+		/// <param name="jobStream"></param>
+		/// <param name="job">The job containing the step</param>
+		/// <param name="graph"></param>
+		/// <param name="outcome"></param>
 		/// <returns>Async task</returns>
-		Task NotifyJobCompleteAsync(IUser User, IStream JobStream, IJob Job, IGraph Graph, LabelOutcome Outcome);
+		Task NotifyJobCompleteAsync(IUser user, IStream jobStream, IJob job, IGraph graph, LabelOutcome outcome);
 
 		/// <summary>
 		/// Send notifications that a job step has completed
 		/// </summary>
-		/// <param name="User">User to notify</param>
-		/// <param name="JobStream">Stream containing the job</param>
-		/// <param name="Job">The job containing the step</param>
-		/// <param name="Batch">Unique id of the batch</param>
-		/// <param name="Step">The step id</param>
-		/// <param name="Node">Corresponding node for the step</param>
-		/// <param name="JobStepEventData"></param>
+		/// <param name="user">User to notify</param>
+		/// <param name="jobStream">Stream containing the job</param>
+		/// <param name="job">The job containing the step</param>
+		/// <param name="batch">Unique id of the batch</param>
+		/// <param name="step">The step id</param>
+		/// <param name="node">Corresponding node for the step</param>
+		/// <param name="jobStepEventData"></param>
 		/// <returns>Async task</returns>
-		Task NotifyJobStepCompleteAsync(IUser User, IStream JobStream, IJob Job, IJobStepBatch Batch, IJobStep Step, INode Node, List<ILogEventData> JobStepEventData);
+		Task NotifyJobStepCompleteAsync(IUser user, IStream jobStream, IJob job, IJobStepBatch batch, IJobStep step, INode node, List<ILogEventData> jobStepEventData);
 
 		/// <summary>
 		/// Send notifications that a job step has completed
 		/// </summary>
-		/// <param name="User">User to notify</param>
-		/// <param name="Job">The job containing the step</param>
-		/// <param name="Stream"></param>
-		/// <param name="Label"></param>
-		/// <param name="LabelIdx"></param>
-		/// <param name="Outcome"></param>
-		/// <param name="StepData"></param>
+		/// <param name="user">User to notify</param>
+		/// <param name="job">The job containing the step</param>
+		/// <param name="stream"></param>
+		/// <param name="label"></param>
+		/// <param name="labelIdx"></param>
+		/// <param name="outcome"></param>
+		/// <param name="stepData"></param>
 		/// <returns>Async task</returns>
-		Task NotifyLabelCompleteAsync(IUser User, IJob Job, IStream Stream, ILabel Label, int LabelIdx, LabelOutcome Outcome, List<(string, JobStepOutcome, Uri)> StepData);
+		Task NotifyLabelCompleteAsync(IUser user, IJob job, IStream stream, ILabel label, int labelIdx, LabelOutcome outcome, List<(string, JobStepOutcome, Uri)> stepData);
 
 		/// <summary>
 		/// Send notifications that a new issue has been created or assigned
 		/// </summary>
-		/// <param name="Issue"></param>
+		/// <param name="issue"></param>
 		/// <returns></returns>
-		Task NotifyIssueUpdatedAsync(IIssue Issue);
+		Task NotifyIssueUpdatedAsync(IIssue issue);
 
 		/// <summary>
 		/// Notification that a stream has failed to update
 		/// </summary>
-		/// <param name="ErrorMessage"></param>
-		/// <param name="FileName"></param>
-		/// <param name="Change"></param>
-		/// <param name="Author"></param>
-		/// <param name="Description"></param>
+		/// <param name="errorMessage"></param>
+		/// <param name="fileName"></param>
+		/// <param name="change"></param>
+		/// <param name="author"></param>
+		/// <param name="description"></param>
 		/// <returns></returns>
-		Task NotifyConfigUpdateFailureAsync(string ErrorMessage, string FileName, int? Change = null, IUser? Author = null, string? Description = null);
+		Task NotifyConfigUpdateFailureAsync(string errorMessage, string fileName, int? change = null, IUser? author = null, string? description = null);
 
 		/// <summary>
 		/// Notification for device service
 		/// </summary>
-		/// <param name="Message"></param>
-		/// <param name="Device"></param>
-		/// <param name="Pool"></param>
-		/// <param name="Stream"></param>
-		/// <param name="Job"></param>
-		/// <param name="Step"></param>
-		/// <param name="Node"></param>
-		/// <param name="User"></param>
+		/// <param name="message"></param>
+		/// <param name="device"></param>
+		/// <param name="pool"></param>
+		/// <param name="stream"></param>
+		/// <param name="job"></param>
+		/// <param name="step"></param>
+		/// <param name="node"></param>
+		/// <param name="user"></param>
 		/// <returns></returns>
-		Task NotifyDeviceServiceAsync(string Message, IDevice? Device = null, IDevicePool? Pool = null, IStream? Stream = null, IJob? Job = null, IJobStep? Step = null, INode? Node = null, IUser? User = null);
+		Task NotifyDeviceServiceAsync(string message, IDevice? device = null, IDevicePool? pool = null, IStream? stream = null, IJob? job = null, IJobStep? step = null, INode? node = null, IUser? user = null);
     }
 }
 
diff --git a/Engine/Source/Programs/Horde/Horde.Build/Notifications/INotificationTrigger.cs b/Engine/Source/Programs/Horde/Horde.Build/Notifications/INotificationTrigger.cs
--- a/Engine/Source/Programs/Horde/Horde.Build/Notifications/INotificationTrigger.cs
+++ b/Engine/Source/Programs/Horde/Horde.Build/Notifications/INotificationTrigger.cs
@@ -1,10 +1,8 @@
 // Copyright Epic Games, Inc. All Rights Reserved.
 
-using Horde.Build.Api;
+using System.Collections.Generic;
 using Horde.Build.Utilities;
 using MongoDB.Bson;
-using MongoDB.Bson.Serialization.Attributes;
-using System.Collections.Generic;
 
 namespace Horde.Build.Models
 {
diff --git a/Engine/Source/Programs/Horde/Horde.Build/Notifications/INotificationTriggerCollection.cs b/Engine/Source/Programs/Horde/Horde.Build/Notifications/INotificationTriggerCollection.cs
--- a/Engine/Source/Programs/Horde/Horde.Build/Notifications/INotificationTriggerCollection.cs
+++ b/Engine/Source/Programs/Horde/Horde.Build/Notifications/INotificationTriggerCollection.cs
@@ -1,15 +1,10 @@
 // Copyright Epic Games, Inc. All Rights Reserved.
 
+using System.Collections.Generic;
+using System.Threading.Tasks;
 using Horde.Build.Models;
-using Horde.Build.Services;
 using Horde.Build.Utilities;
 using MongoDB.Bson;
-using MongoDB.Bson.Serialization.Attributes;
-using MongoDB.Driver;
-using System;
-using System.Collections.Generic;
-using System.Linq;
-using System.Threading.Tasks;
 
 namespace Horde.Build.Collections
 {
@@ -23,46 +18,46 @@ namespace Horde.Build.Collections
 		/// <summary>
 		/// Finds or adds a trigger
 		/// </summary>
-		/// <param name="TriggerId">The trigger id</param>
+		/// <param name="triggerId">The trigger id</param>
 		/// <returns>New trigger document</returns>
-		Task<INotificationTrigger> FindOrAddAsync(ObjectId TriggerId);
+		Task<INotificationTrigger> FindOrAddAsync(ObjectId triggerId);
 
 		/// <summary>
 		/// Finds an existing trigger with the given id, or adds one if it does not exist
 		/// </summary>
-		/// <param name="TriggerId"></param>
+		/// <param name="triggerId"></param>
 		/// <returns></returns>
-		Task<INotificationTrigger?> GetAsync(ObjectId TriggerId);
+		Task<INotificationTrigger?> GetAsync(ObjectId triggerId);
 
 		/// <summary>
 		/// Deletes a trigger
 		/// </summary>
-		/// <param name="TriggerId">The unique trigger id</param>
+		/// <param name="triggerId">The unique trigger id</param>
 		/// <returns>Async task</returns>
-		Task DeleteAsync(ObjectId TriggerId);
+		Task DeleteAsync(ObjectId triggerId);
 
 		/// <summary>
 		/// Deletes a trigger
 		/// </summary>
-		/// <param name="TriggerIds">The unique trigger id</param>
+		/// <param name="triggerIds">The unique trigger id</param>
 		/// <returns>Async task</returns>
-		Task DeleteAsync(List<ObjectId> TriggerIds);
+		Task DeleteAsync(List<ObjectId> triggerIds);
 
 		/// <summary>
 		/// Fires the trigger, and marks notifications has having been sent
 		/// </summary>
-		/// <param name="Trigger">The trigger to fire</param>
+		/// <param name="trigger">The trigger to fire</param>
 		/// <returns>The trigger document</returns>
-		Task<INotificationTrigger?> FireAsync(INotificationTrigger Trigger);
+		Task<INotificationTrigger?> FireAsync(INotificationTrigger trigger);
 
 		/// <summary>
 		/// Adds a subscriber to a particular trigger
 		/// </summary>
-		/// <param name="Trigger">The trigger to subscribe to</param>
-		/// <param name="UserId">The user name</param>
-		/// <param name="Email">Whether to receive email notifications</param>
-		/// <param name="Slack">Whether to receive Slack notifications</param>
+		/// <param name="trigger">The trigger to subscribe to</param>
+		/// <param name="userId">The user name</param>
+		/// <param name="email">Whether to receive email notifications</param>
+		/// <param name="slack">Whether to receive Slack notifications</param>
 		/// <returns>The new trigger state</returns>
-		Task<INotificationTrigger?> UpdateSubscriptionsAsync(INotificationTrigger Trigger, UserId UserId, bool? Email, bool? Slack);
+		Task<INotificationTrigger?> UpdateSubscriptionsAsync(INotificationTrigger trigger, UserId userId, bool? email, bool? slack);
 	}
 }
diff --git a/Engine/Source/Programs/Horde/Horde.Build/Notifications/ISubscription.cs b/Engine/Source/Programs/Horde/Horde.Build/Notifications/ISubscription.cs
--- a/Engine/Source/Programs/Horde/Horde.Build/Notifications/ISubscription.cs
+++ b/Engine/Source/Programs/Horde/Horde.Build/Notifications/ISubscription.cs
@@ -1,13 +1,6 @@
 // Copyright Epic Games, Inc. All Rights Reserved.
 
-using HordeCommon;
-using Horde.Build.Api;
 using Horde.Build.Utilities;
-using MongoDB.Bson;
-using System;
-using System.Collections.Generic;
-using System.Linq;
-using System.Threading.Tasks;
 
 namespace Horde.Build.Models
 {
diff --git a/Engine/Source/Programs/Horde/Horde.Build/Notifications/ISubscriptionCollection.cs b/Engine/Source/Programs/Horde/Horde.Build/Notifications/ISubscriptionCollection.cs
--- a/Engine/Source/Programs/Horde/Horde.Build/Notifications/ISubscriptionCollection.cs
+++ b/Engine/Source/Programs/Horde/Horde.Build/Notifications/ISubscriptionCollection.cs
@@ -1,26 +1,13 @@
 // Copyright Epic Games, Inc. All Rights Reserved.
 
-using EpicGames.Core;
-using HordeCommon;
+using System.Collections.Generic;
+using System.Threading.Tasks;
 using Horde.Build.Api;
 using Horde.Build.Models;
-using Horde.Build.Services;
 using Horde.Build.Utilities;
-using MongoDB.Bson;
-using MongoDB.Bson.Serialization.Attributes;
-using MongoDB.Driver;
-using Polly;
-using System;
-using System.Collections.Generic;
-using System.Linq;
-using System.Reflection.Metadata.Ecma335;
-using System.Text;
-using System.Threading.Tasks;
 
 namespace Horde.Build.Collections
 {
-	using StreamId = StringId<IStream>;
-	using TemplateRefId = StringId<TemplateRef>;
 	using UserId = ObjectId<IUser>;
 
 	/// <summary>
@@ -46,14 +33,14 @@ namespace Horde.Build.Collections
 		/// <summary>
 		/// Constructor
 		/// </summary>
-		/// <param name="Event">Name of the event</param>
-		/// <param name="UserId">User name</param>
-		/// <param name="NotificationType">Notification type</param>
-		public NewSubscription(EventRecord Event, UserId UserId, NotificationType NotificationType)
+		/// <param name="eventRecord">Name of the event</param>
+		/// <param name="userId">User name</param>
+		/// <param name="notificationType">Notification type</param>
+		public NewSubscription(EventRecord eventRecord, UserId userId, NotificationType notificationType)
 		{
-			this.Event = Event;
-			this.UserId = UserId;
-			this.NotificationType = NotificationType;
+			Event = eventRecord;
+			UserId = userId;
+			NotificationType = notificationType;
 		}
 	}
 
@@ -65,36 +52,36 @@ namespace Horde.Build.Collections
 		/// <summary>
 		/// Add new subscription documents
 		/// </summary>
-		/// <param name="Subscriptions">The new subscriptions to add</param>
+		/// <param name="subscriptions">The new subscriptions to add</param>
 		/// <returns>The subscriptions that were added</returns>
-		public Task<List<ISubscription>> AddAsync(IEnumerable<NewSubscription> Subscriptions);
+		public Task<List<ISubscription>> AddAsync(IEnumerable<NewSubscription> subscriptions);
 
 		/// <summary>
 		/// Remove a set of existing subscriptions
 		/// </summary>
-		/// <param name="Subscriptions">Subscriptions to remove</param>
+		/// <param name="subscriptions">Subscriptions to remove</param>
 		/// <returns>Async task</returns>
-		public Task RemoveAsync(IEnumerable<ISubscription> Subscriptions);
+		public Task RemoveAsync(IEnumerable<ISubscription> subscriptions);
 
 		/// <summary>
 		/// Gets a subscription by id
 		/// </summary>
-		/// <param name="SubscriptionId">Subscription to remove</param>
+		/// <param name="subscriptionId">Subscription to remove</param>
 		/// <returns>Async task</returns>
-		public Task<ISubscription?> GetAsync(string SubscriptionId);
+		public Task<ISubscription?> GetAsync(string subscriptionId);
 
 		/// <summary>
 		/// Find all subscribers of a certain event
 		/// </summary>
-		/// <param name="Event">Name of the event</param>
+		/// <param name="eventRecord">Name of the event</param>
 		/// <returns>Name of the event to find subscribers for</returns>
-		public Task<List<ISubscription>> FindSubscribersAsync(EventRecord Event);
+		public Task<List<ISubscription>> FindSubscribersAsync(EventRecord eventRecord);
 
 		/// <summary>
 		/// Find subscriptions for a particular user
 		/// </summary>
-		/// <param name="UserId">The user to search for</param>
+		/// <param name="userId">The user to search for</param>
 		/// <returns>List of subscriptions</returns>
-		public Task<List<ISubscription>> FindSubscriptionsAsync(UserId UserId);
+		public Task<List<ISubscription>> FindSubscriptionsAsync(UserId userId);
 	}
 }
