Horde: Fix coding conventions within Horde.Build.

diff --git a/Engine/Source/Programs/Horde/Horde.Build/Notifications/Impl/NotificationService.cs b/Engine/Source/Programs/Horde/Horde.Build/Notifications/Impl/NotificationService.cs
--- a/Engine/Source/Programs/Horde/Horde.Build/Notifications/Impl/NotificationService.cs
+++ b/Engine/Source/Programs/Horde/Horde.Build/Notifications/Impl/NotificationService.cs
@@ -3,35 +3,25 @@
 using System;
 using System.Collections.Concurrent;
 using System.Collections.Generic;
-using System.Drawing;
 using System.Linq;
-using System.Net;
-using System.Net.Http;
-using System.Net.Mail;
 using System.Security.Claims;
-using System.Text;
 using System.Threading;
 using System.Threading.Tasks;
 using EpicGames.Core;
 using Horde.Build.Api;
 using Horde.Build.Collections;
-using HordeCommon;
 using Horde.Build.Models;
+using Horde.Build.Services;
 using Horde.Build.Utilities;
-using Horde.Build.Utilities.BlockKit;
-using Horde.Build.Utilities.Slack.BlockKit;
+using HordeCommon;
 using Microsoft.Extensions.Hosting;
 using Microsoft.Extensions.Logging;
 using Microsoft.Extensions.Options;
 using MongoDB.Bson;
 using MongoDB.Driver;
-using Newtonsoft.Json;
-using Newtonsoft.Json.Linq;
+using StackExchange.Redis;
 using StatsdClient;
 using JsonSerializer = System.Text.Json.JsonSerializer;
-using Horde.Build.Notifications;
-using Horde.Build.Services;
-using StackExchange.Redis;
 
 namespace Horde.Build.Notifications.Impl
 {
@@ -45,115 +35,115 @@ namespace Horde.Build.Notifications.Impl
 		/// <summary>
 		/// The available notification sinks
 		/// </summary>
-		private List<INotificationSink> Sinks;
+		private readonly List<INotificationSink> _sinks;
 
 		/// <summary>
 		/// Collection of subscriptions
 		/// </summary>
-		private readonly ISubscriptionCollection SubscriptionCollection;
+		private readonly ISubscriptionCollection _subscriptionCollection;
 
 		/// <summary>
 		/// Collection of notification request documents.
 		/// </summary>
-		private readonly INotificationTriggerCollection TriggerCollection;
+		private readonly INotificationTriggerCollection _triggerCollection;
 
 		/// <summary>
-		/// Instance of the <see cref="GraphCollection"/>.
+		/// Instance of the <see cref="_graphCollection"/>.
 		/// </summary>
-		private readonly IGraphCollection GraphCollection;
+		private readonly IGraphCollection _graphCollection;
 
 		/// <summary>
 		/// 
 		/// </summary>
-		private readonly IUserCollection UserCollection;
+		private readonly IUserCollection _userCollection;
 
 		/// <summary>
 		/// Job service instance
 		/// </summary>
-		private readonly JobService JobService;
+		private readonly JobService _jobService;
 
 		/// <summary>
-		/// Instance of the <see cref="StreamService"/>.
+		/// Instance of the <see cref="_streamService"/>.
 		/// </summary>
-		private readonly StreamService StreamService;
+		private readonly StreamService _streamService;
 
 		/// <summary>
 		/// 
 		/// </summary>
-		private readonly IIssueService IssueService;
+		private readonly IIssueService _issueService;
 
 		/// <summary>
-		/// Instance of the <see cref="LogFileService"/>.
+		/// Instance of the <see cref="_logFileService"/>.
 		/// </summary>
-		private readonly ILogFileService LogFileService;
+		private readonly ILogFileService _logFileService;
 		
 		/// <summary>
 		/// Instance of the <see cref="IDogStatsd"/>.
 		/// </summary>
-		private readonly IDogStatsd DogStatsd;
+		private readonly IDogStatsd _dogStatsd;
 
 		/// <summary>
 		/// Redis database
 		/// </summary>
-		private readonly IDatabase RedisDatabase;
+		private readonly IDatabase _redisDatabase;
 		
 		/// <summary>
 		/// Settings for the application.
 		/// </summary>
-		private readonly IOptionsMonitor<ServerSettings> Settings;
+		private readonly IOptionsMonitor<ServerSettings> _settings;
 
 		/// <summary>
 		/// List of asychronous tasks currently executing
 		/// </summary>
-		private readonly ConcurrentQueue<Task> NewTasks = new ConcurrentQueue<Task>();
+		private readonly ConcurrentQueue<Task> _newTasks = new ConcurrentQueue<Task>();
 
 		/// <summary>
 		/// Set when there are new tasks to wait for
 		/// </summary>
-		private AsyncEvent NewTaskEvent = new AsyncEvent();
+		private AsyncEvent _newTaskEvent = new AsyncEvent();
 
 		/// <summary>
 		/// Settings for the application.
 		/// </summary>
-		private readonly ILogger<NotificationService> Logger;
+		private readonly ILogger<NotificationService> _logger;
 		
 		/// <summary>
 		/// Ticker for running batch sender method
 		/// </summary>
-		internal ITicker Ticker;
+		internal ITicker _ticker;
 
 		/// <summary>
 		/// Interval at which queued notifications should be sent as a batch 
 		/// </summary>
-		internal TimeSpan NotificationBatchInterval = TimeSpan.FromHours(12);
+		internal TimeSpan _notificationBatchInterval = TimeSpan.FromHours(12);
 		
-		static string RedisQueueListKey(string NotificationType) => "NotificationService.queued." + NotificationType;
+		static string RedisQueueListKey(string notificationType) => "NotificationService.queued." + notificationType;
 
 		/// <summary>
 		/// Constructor
 		/// </summary>
-		public NotificationService(IEnumerable<INotificationSink> Sinks, IOptionsMonitor<ServerSettings> Settings, ILogger<NotificationService> Logger, IGraphCollection GraphCollection, ISubscriptionCollection SubscriptionCollection, INotificationTriggerCollection TriggerCollection, IUserCollection UserCollection, JobService JobService, StreamService StreamService, IIssueService IssueService, ILogFileService LogFileService, IDogStatsd DogStatsd, IDatabase RedisDatabase, IClock Clock)
+		public NotificationService(IEnumerable<INotificationSink> sinks, IOptionsMonitor<ServerSettings> settings, ILogger<NotificationService> logger, IGraphCollection graphCollection, ISubscriptionCollection subscriptionCollection, INotificationTriggerCollection triggerCollection, IUserCollection userCollection, JobService jobService, StreamService streamService, IIssueService issueService, ILogFileService logFileService, IDogStatsd dogStatsd, IDatabase redisDatabase, IClock clock)
 		{
-			this.Sinks = Sinks.ToList();
-			this.Settings = Settings;
-			this.Logger = Logger;
-			this.GraphCollection = GraphCollection;
-			this.SubscriptionCollection = SubscriptionCollection;
-			this.TriggerCollection = TriggerCollection;
-			this.UserCollection = UserCollection;
-			this.JobService = JobService;
-			this.StreamService = StreamService;
-			this.IssueService = IssueService;
-			this.LogFileService = LogFileService;
-			this.DogStatsd = DogStatsd;
-			this.RedisDatabase = RedisDatabase;
-
-			IssueService.OnIssueUpdated += NotifyIssueUpdated;
-			JobService.OnJobStepComplete += NotifyJobStepComplete;
-			JobService.OnJobScheduled += NotifyJobScheduled;
-			JobService.OnLabelUpdate += NotifyLabelUpdate;
+			_sinks = sinks.ToList();
+			_settings = settings;
+			_logger = logger;
+			_graphCollection = graphCollection;
+			_subscriptionCollection = subscriptionCollection;
+			_triggerCollection = triggerCollection;
+			_userCollection = userCollection;
+			_jobService = jobService;
+			_streamService = streamService;
+			_issueService = issueService;
+			_logFileService = logFileService;
+			_dogStatsd = dogStatsd;
+			_redisDatabase = redisDatabase;
+
+			issueService.OnIssueUpdated += NotifyIssueUpdated;
+			jobService.OnJobStepComplete += NotifyJobStepComplete;
+			jobService.OnJobScheduled += NotifyJobScheduled;
+			jobService.OnLabelUpdate += NotifyLabelUpdate;
 			
-			this.Ticker = Clock.AddSharedTicker<NotificationService>(NotificationBatchInterval, TickEveryTwelveHoursAsync, Logger);
+			_ticker = clock.AddSharedTicker<NotificationService>(_notificationBatchInterval, TickEveryTwelveHoursAsync, logger);
 		}
 
 		/// <inheritdoc/>
@@ -161,155 +151,155 @@ namespace Horde.Build.Notifications.Impl
 		{
 			base.Dispose();
 
-			IssueService.OnIssueUpdated -= NotifyIssueUpdated;
-			JobService.OnJobStepComplete -= NotifyJobStepComplete;
-			JobService.OnJobScheduled += NotifyJobScheduled;
-			JobService.OnLabelUpdate -= NotifyLabelUpdate;
+			_issueService.OnIssueUpdated -= NotifyIssueUpdated;
+			_jobService.OnJobStepComplete -= NotifyJobStepComplete;
+			_jobService.OnJobScheduled += NotifyJobScheduled;
+			_jobService.OnLabelUpdate -= NotifyLabelUpdate;
 
 			GC.SuppressFinalize(this);
-			Ticker.Dispose();
+			_ticker.Dispose();
 		}
 
 		/// <inheritdoc/>
-		public async Task<bool> UpdateSubscriptionsAsync(ObjectId TriggerId, ClaimsPrincipal User, bool? Email, bool? Slack)
+		public async Task<bool> UpdateSubscriptionsAsync(ObjectId triggerId, ClaimsPrincipal user, bool? email, bool? slack)
 		{
-			UserId? UserId = User.GetUserId();
-			if (UserId == null)
+			UserId? userId = user.GetUserId();
+			if (userId == null)
 			{
-				Logger.LogWarning("Unable to find username for principal {User}", User.Identity?.Name);
+				_logger.LogWarning("Unable to find username for principal {User}", user.Identity?.Name);
 				return false;
 			}
 
-			INotificationTrigger Trigger = await TriggerCollection.FindOrAddAsync(TriggerId);
-			await TriggerCollection.UpdateSubscriptionsAsync(Trigger, UserId.Value, Email, Slack);
+			INotificationTrigger trigger = await _triggerCollection.FindOrAddAsync(triggerId);
+			await _triggerCollection.UpdateSubscriptionsAsync(trigger, userId.Value, email, slack);
 			return true;
 		}
 
 		/// <inheritdoc/>
-		public async Task<INotificationSubscription?> GetSubscriptionsAsync(ObjectId TriggerId, ClaimsPrincipal User)
+		public async Task<INotificationSubscription?> GetSubscriptionsAsync(ObjectId triggerId, ClaimsPrincipal user)
 		{
-			UserId? UserId = User.GetUserId();
-			if (UserId == null)
+			UserId? userId = user.GetUserId();
+			if (userId == null)
 			{
 				return null;
 			}
 
-			INotificationTrigger? Trigger = await TriggerCollection.GetAsync(TriggerId);
-			if(Trigger == null)
+			INotificationTrigger? trigger = await _triggerCollection.GetAsync(triggerId);
+			if(trigger == null)
 			{
 				return null;
 			}
 
-			return Trigger.Subscriptions.FirstOrDefault(x => x.UserId == UserId.Value);
+			return trigger.Subscriptions.FirstOrDefault(x => x.UserId == userId.Value);
 		}
 
 		/// <inheritdoc/>
-		public void NotifyJobStepComplete(IJob Job, IGraph Graph, SubResourceId BatchId, SubResourceId StepId)
+		public void NotifyJobStepComplete(IJob job, IGraph graph, SubResourceId batchId, SubResourceId stepId)
 		{
 			// Enqueue job step complete notifications if needed
-			if (Job.TryGetStep(BatchId, StepId, out IJobStep? Step))
+			if (job.TryGetStep(batchId, stepId, out IJobStep? step))
 			{
-				Logger.LogInformation("Queuing step notifications for {JobId}:{BatchId}:{StepId}", Job.Id, BatchId, StepId);
-				EnqueueTask(() => SendJobStepNotificationsAsync(Job, BatchId, StepId));
+				_logger.LogInformation("Queuing step notifications for {JobId}:{BatchId}:{StepId}", job.Id, batchId, stepId);
+				EnqueueTask(() => SendJobStepNotificationsAsync(job, batchId, stepId));
 			}
 
 			// Enqueue job complete notifications if needed
-			if (Job.GetState() == JobState.Complete)
+			if (job.GetState() == JobState.Complete)
 			{
-				Logger.LogInformation("Queuing job notifications for {JobId}:{BatchId}:{StepId}", Job.Id, BatchId, StepId);
-				EnqueueTask(() => SendJobNotificationsAsync(Job, Graph));
-				EnqueueTask(() => RecordJobCompleteMetrics(Job));
+				_logger.LogInformation("Queuing job notifications for {JobId}:{BatchId}:{StepId}", job.Id, batchId, stepId);
+				EnqueueTask(() => SendJobNotificationsAsync(job, graph));
+				EnqueueTask(() => RecordJobCompleteMetrics(job));
 			}
 		}
 		
 		/// <inheritdoc/>
-		public void NotifyJobScheduled(IPool Pool, bool PoolHasAgentsOnline, IJob Job, IGraph Graph, SubResourceId BatchId)
+		public void NotifyJobScheduled(IPool pool, bool poolHasAgentsOnline, IJob job, IGraph graph, SubResourceId batchId)
 		{
-			if (Pool.EnableAutoscaling && !PoolHasAgentsOnline)
+			if (pool.EnableAutoscaling && !poolHasAgentsOnline)
 			{
-				EnqueueTasks(Sink => EnqueueNotificationForBatchSending(new JobScheduledNotification(Job.Id.ToString(), Job.Name, Pool.Name)));
+				EnqueueTasks(sink => EnqueueNotificationForBatchSending(new JobScheduledNotification(job.Id.ToString(), job.Name, pool.Name)));
 			}
 		}
 
 		/// <inheritdoc/>
-		public void NotifyLabelUpdate(IJob Job, IReadOnlyList<(LabelState, LabelOutcome)> OldLabelStates, IReadOnlyList<(LabelState, LabelOutcome)> NewLabelStates)
+		public void NotifyLabelUpdate(IJob job, IReadOnlyList<(LabelState, LabelOutcome)> oldLabelStates, IReadOnlyList<(LabelState, LabelOutcome)> newLabelStates)
 		{
 			// If job has any label trigger IDs, send label complete notifications if needed
-			for (int Idx = 0; Idx < OldLabelStates.Count && Idx < NewLabelStates.Count; Idx++)
+			for (int idx = 0; idx < oldLabelStates.Count && idx < newLabelStates.Count; idx++)
 			{
-				if (OldLabelStates[Idx] != NewLabelStates[Idx])
+				if (oldLabelStates[idx] != newLabelStates[idx])
 				{
-					EnqueueTask(() => SendAllLabelNotificationsAsync(Job, OldLabelStates, NewLabelStates));
+					EnqueueTask(() => SendAllLabelNotificationsAsync(job, oldLabelStates, newLabelStates));
 					break;
 				}
 			}
 		}
 
 		/// <inheritdoc/>
-		public void NotifyIssueUpdated(IIssue Issue)
+		public void NotifyIssueUpdated(IIssue issue)
 		{
-			Logger.LogInformation("Issue {IssueId} updated", Issue.Id);
-			EnqueueTasks(Sink => Sink.NotifyIssueUpdatedAsync(Issue));
+			_logger.LogInformation("Issue {IssueId} updated", issue.Id);
+			EnqueueTasks(sink => sink.NotifyIssueUpdatedAsync(issue));
 		}
 
 		/// <inheritdoc/>
-		public void NotifyConfigUpdateFailure(string ErrorMessage, string FileName, int? Change = null, IUser? Author = null, string? Description = null)
+		public void NotifyConfigUpdateFailure(string errorMessage, string fileName, int? change = null, IUser? author = null, string? description = null)
 		{
-			EnqueueTasks(Sink => Sink.NotifyConfigUpdateFailureAsync(ErrorMessage, FileName, Change, Author, Description));
+			EnqueueTasks(sink => sink.NotifyConfigUpdateFailureAsync(errorMessage, fileName, change, author, description));
 		}
 
 		/// <inheritdoc/>
-		public void NotifyDeviceService(string Message, IDevice? Device = null, IDevicePool? Pool = null, IStream? Stream = null, IJob? Job = null, IJobStep? Step = null, INode? Node = null, IUser? User = null)
+		public void NotifyDeviceService(string message, IDevice? device = null, IDevicePool? pool = null, IStream? stream = null, IJob? job = null, IJobStep? step = null, INode? node = null, IUser? user = null)
 		{
-			EnqueueTasks(Sink => Sink.NotifyDeviceServiceAsync(Message, Device, Pool, Stream, Job, Step, Node, User));
+			EnqueueTasks(sink => sink.NotifyDeviceServiceAsync(message, device, pool, stream, job, step, node, user));
 		}
 
 		/// <summary>
 		/// Enqueues an async task
 		/// </summary>
-		/// <param name="TaskFunc">Function to generate an async task</param>
-		void EnqueueTask(Func<Task> TaskFunc)
+		/// <param name="taskFunc">Function to generate an async task</param>
+		void EnqueueTask(Func<Task> taskFunc)
 		{
-			NewTasks.Enqueue(Task.Run(TaskFunc));
+			_newTasks.Enqueue(Task.Run(taskFunc));
 		}
 
 		/// <summary>
 		/// Enqueues an async task
 		/// </summary>
-		/// <param name="TaskFunc">Function to generate an async task</param>
-		void EnqueueTasks(Func<INotificationSink, Task> TaskFunc)
+		/// <param name="taskFunc">Function to generate an async task</param>
+		void EnqueueTasks(Func<INotificationSink, Task> taskFunc)
 		{
-			foreach (INotificationSink Sink in Sinks)
+			foreach (INotificationSink sink in _sinks)
 			{
-				EnqueueTask(() => TaskFunc(Sink));
+				EnqueueTask(() => taskFunc(sink));
 			}
 		}
 
 		/// <summary>
 		/// Enqueue a notification in Redis for batch sending later on 
 		/// </summary>
-		/// <param name="Notification"></param>
+		/// <param name="notification"></param>
 		/// <typeparam name="T">Any INotification type</typeparam>
-		private async Task EnqueueNotificationForBatchSending<T>(T Notification) where T : INotification
+		private async Task EnqueueNotificationForBatchSending<T>(T notification) where T : INotification
 		{
 			try
 			{
-				byte[] Data = JsonSerializer.SerializeToUtf8Bytes(Notification);
-				await RedisDatabase.ListRightPushAsync(RedisQueueListKey(typeof(T).ToString()), Data);
+				byte[] data = JsonSerializer.SerializeToUtf8Bytes(notification);
+				await _redisDatabase.ListRightPushAsync(RedisQueueListKey(typeof(T).ToString()), data);
 			}
 			catch (Exception e)
 			{
-				Logger.LogError(e, "Unable to serialize and queue notification {Type} for batch sending in Redis", Notification.GetType());
+				_logger.LogError(e, "Unable to serialize and queue notification {Type} for batch sending in Redis", notification.GetType());
 			}
 		}
 
-		private async ValueTask TickEveryTwelveHoursAsync(CancellationToken StoppingToken)
+		private async ValueTask TickEveryTwelveHoursAsync(CancellationToken stoppingToken)
 		{
-			List<JobScheduledNotification> JobScheduledNotifications = await GetAllQueuedNotificationsAsync<JobScheduledNotification>();
+			List<JobScheduledNotification> jobScheduledNotifications = await GetAllQueuedNotificationsAsync<JobScheduledNotification>();
 
-			foreach (INotificationSink Sink in Sinks)
+			foreach (INotificationSink sink in _sinks)
 			{
-				await Sink.NotifyJobScheduledAsync(JobScheduledNotifications);
+				await sink.NotifyJobScheduledAsync(jobScheduledNotifications);
 			}
 		}
 		
@@ -323,55 +313,57 @@ namespace Horde.Build.Notifications.Impl
 		{
 			// Reading and deleting the entire list from Redis in this way is not thread-safe.
 			// But likely not a big deal considering the alternative of distributed locks.
-			string RedisKey = RedisQueueListKey(typeof(T).ToString());
-			RedisValue[] RawNotifications = await RedisDatabase.ListRangeAsync(RedisKey);
-			await RedisDatabase.KeyDeleteAsync(RedisKey);
+			string redisKey = RedisQueueListKey(typeof(T).ToString());
+			RedisValue[] rawNotifications = await _redisDatabase.ListRangeAsync(redisKey);
+			await _redisDatabase.KeyDeleteAsync(redisKey);
 
-			List<T> Notifications = new List<T>();
-			foreach (byte[] Data in RawNotifications)
+			List<T> notifications = new List<T>();
+			foreach (byte[] data in rawNotifications)
 			{
 				try
 				{
-					T? Notification = JsonSerializer.Deserialize<T>(Data);
-					if (Notification == null) throw new Exception("Unable to deserialize");
-					Notifications.Add(Notification);
+					T? notification = JsonSerializer.Deserialize<T>(data);
+					if (notification == null)
+					{
+						throw new Exception("Unable to deserialize");
+					}
+					notifications.Add(notification);
 				}
 				catch (Exception e)
 				{
-					Logger.LogError(e, "Unable to deserialize notification data of {Type}", typeof(T));
+					_logger.LogError(e, "Unable to deserialize notification data of {Type}", typeof(T));
 				}
 			}
 
-			return Notifications;
+			return notifications;
 		}
 
 		/// <inheritdoc/>
-		[System.Diagnostics.CodeAnalysis.SuppressMessage("Design", "CA1031:Do not catch general exception types", Justification = "<Pending>")]
-		protected async override Task ExecuteAsync(CancellationToken StoppingToken)
+		protected override async Task ExecuteAsync(CancellationToken stoppingToken)
 		{
-			await Ticker.StartAsync();
+			await _ticker.StartAsync();
 			
 			// This background service just waits for tasks to finish and prints any exception info. The only reason to do this is to
 			// ensure we finish processing everything before shutdown.
-			using (CancellationTask StoppingTask = new CancellationTask(StoppingToken))
+			using (CancellationTask stoppingTask = new CancellationTask(stoppingToken))
 			{
-				List<Task> Tasks = new List<Task>();
-				Tasks.Add(NewTaskEvent.Task);
+				List<Task> tasks = new List<Task>();
+				tasks.Add(_newTaskEvent.Task);
 
 				for (; ; )
 				{
 					// Add any new tasks to be monitored
-					Task? NewTask;
-					while (NewTasks.TryDequeue(out NewTask))
+					Task? newTask;
+					while (_newTasks.TryDequeue(out newTask))
 					{
-						Tasks.Add(NewTask);
+						tasks.Add(newTask);
 					}
 
 					// If we don't have any
-					if (Tasks.Count == 1)
+					if (tasks.Count == 1)
 					{
-						await Task.WhenAny(NewTaskEvent.Task, StoppingTask.Task);
-						if (StoppingToken.IsCancellationRequested)
+						await Task.WhenAny(_newTaskEvent.Task, stoppingTask.Task);
+						if (stoppingToken.IsCancellationRequested)
 						{
 							break;
 						}
@@ -379,98 +371,98 @@ namespace Horde.Build.Notifications.Impl
 					else
 					{
 						// Wait for something to finish
-						Task Task = await Task.WhenAny(Tasks);
-						if (Task == NewTaskEvent.Task)
+						Task task = await Task.WhenAny(tasks);
+						if (task == _newTaskEvent.Task)
 						{
-							NewTaskEvent = new AsyncEvent();
-							Tasks[0] = NewTaskEvent.Task;
+							_newTaskEvent = new AsyncEvent();
+							tasks[0] = _newTaskEvent.Task;
 						}
 						else
 						{
 							try
 							{
-								await Task;
+								await task;
 							}
-							catch (Exception Ex)
+							catch (Exception ex)
 							{
-								Logger.LogError(Ex, "Exception while executing notification");
+								_logger.LogError(ex, "Exception while executing notification");
 							}
-							Tasks.Remove(Task);
+							tasks.Remove(task);
 						}
 					}
 				}
 			}
 			
-			await Ticker.StopAsync();
+			await _ticker.StopAsync();
 		}
 
-		internal Task ExecuteBackgroundForTest(CancellationToken StoppingToken)
+		internal Task ExecuteBackgroundForTest(CancellationToken stoppingToken)
 		{
-			return ExecuteAsync(StoppingToken);
+			return ExecuteAsync(stoppingToken);
 		}
 
 		/// <summary>
 		/// Gets the <see cref="INotificationTrigger"/> for a given trigger ID, if any.
 		/// </summary>
-		/// <param name="TriggerId"></param>
+		/// <param name="triggerId"></param>
 		/// <param name="bFireTrigger">If true, the trigger is fired and cannot be reused</param>
 		/// <returns></returns>
-		private async Task<INotificationTrigger?> GetNotificationTrigger(ObjectId? TriggerId, bool bFireTrigger)
+		private async Task<INotificationTrigger?> GetNotificationTrigger(ObjectId? triggerId, bool bFireTrigger)
 		{
-			if (TriggerId == null)
+			if (triggerId == null)
 			{
 				return null;
 			}
 
-			INotificationTrigger? Trigger = await TriggerCollection.GetAsync(TriggerId.Value);
-			if (Trigger == null)
+			INotificationTrigger? trigger = await _triggerCollection.GetAsync(triggerId.Value);
+			if (trigger == null)
 			{
 				return null;
 			}
 
-			return bFireTrigger ? await TriggerCollection.FireAsync(Trigger) : Trigger;
+			return bFireTrigger ? await _triggerCollection.FireAsync(trigger) : trigger;
 		}
 	
-		private async Task SendJobNotificationsAsync(IJob Job, IGraph Graph)
+		private async Task SendJobNotificationsAsync(IJob job, IGraph graph)
 		{
-			using IDisposable Scope = Logger.BeginScope("Sending notifications for job {JobId}", Job.Id);
+			using IDisposable scope = _logger.BeginScope("Sending notifications for job {JobId}", job.Id);
 
-			Job.GetJobState(Job.GetStepForNodeMap(), out _, out LabelOutcome Outcome);
-			JobCompleteEventRecord JobCompleteEvent = new JobCompleteEventRecord(Job.StreamId, Job.TemplateId, Outcome);
+			job.GetJobState(job.GetStepForNodeMap(), out _, out LabelOutcome outcome);
+			JobCompleteEventRecord jobCompleteEvent = new JobCompleteEventRecord(job.StreamId, job.TemplateId, outcome);
 
-			IStream? JobStream = await StreamService.GetStreamAsync(Job.StreamId);
-			if (JobStream == null)
+			IStream? jobStream = await _streamService.GetStreamAsync(job.StreamId);
+			if (jobStream == null)
 			{
-				Logger.LogError("Unable to get stream {StreamId}", Job.StreamId);
+				_logger.LogError("Unable to get stream {StreamId}", job.StreamId);
 				return;
 			}
 
-			List<IUser> UsersToNotify = await GetUsersToNotify(JobCompleteEvent, Job.NotificationTriggerId, true);
-			foreach (IUser UserToNotify in UsersToNotify)
+			List<IUser> usersToNotify = await GetUsersToNotify(jobCompleteEvent, job.NotificationTriggerId, true);
+			foreach (IUser userToNotify in usersToNotify)
 			{
-				if(Job.PreflightChange != 0)
+				if(job.PreflightChange != 0)
 				{
-					if(UserToNotify.Id != Job.StartedByUserId)
+					if(userToNotify.Id != job.StartedByUserId)
 					{
 						continue;
 					}
 				}
-				EnqueueTasks(Sink => Sink.NotifyJobCompleteAsync(UserToNotify, JobStream, Job, Graph, Outcome));
+				EnqueueTasks(sink => sink.NotifyJobCompleteAsync(userToNotify, jobStream, job, graph, outcome));
 			}
 
-			if (Job.PreflightChange == 0)
+			if (job.PreflightChange == 0)
 			{
-				EnqueueTasks(Sink => Sink.NotifyJobCompleteAsync(JobStream, Job, Graph, Outcome));
+				EnqueueTasks(sink => sink.NotifyJobCompleteAsync(jobStream, job, graph, outcome));
 			}
 
-			Logger.LogDebug("Finished sending notifications for job {JobId}", Job.Id);
+			_logger.LogDebug("Finished sending notifications for job {JobId}", job.Id);
 		}
 
-		private Task RecordJobCompleteMetrics(IJob Job)
+		private Task RecordJobCompleteMetrics(IJob job)
 		{
-			void RecordMetric(string Type, JobStepOutcome Outcome, DateTimeOffset? StartTime, DateTimeOffset? FinishTime)
+			void RecordMetric(string type, JobStepOutcome outcome, DateTimeOffset? startTime, DateTimeOffset? finishTime)
 			{
-				string OutcomeStr = Outcome switch
+				string outcomeStr = outcome switch
 				{
 					JobStepOutcome.Unspecified => "unspecified",
 					JobStepOutcome.Failure => "failure",
@@ -479,23 +471,23 @@ namespace Horde.Build.Notifications.Impl
 					_ => "unspecified"
 				};
 				
-				string[] Tags = {"stream:" + Job.StreamId, "template:" + Job.TemplateId};
-				DogStatsd.Increment($"horde.{Type}.{OutcomeStr}.count", 1, tags: Tags);
+				string[] tags = {"stream:" + job.StreamId, "template:" + job.TemplateId};
+				_dogStatsd.Increment($"horde.{type}.{outcomeStr}.count", 1, tags: tags);
 
-				if (StartTime == null || FinishTime == null)
+				if (startTime == null || finishTime == null)
 				{
-					Logger.LogDebug("Completed job or step is missing start or finish time, cannot record duration metric. Job ID={JobId}", Job.Id);
+					_logger.LogDebug("Completed job or step is missing start or finish time, cannot record duration metric. Job ID={JobId}", job.Id);
 					return;
 				}
 
-				TimeSpan Duration = FinishTime.Value - StartTime.Value;
-				DogStatsd.Timer($"horde.{Type}.{OutcomeStr}.duration", Duration.TotalSeconds, tags: Tags);
+				TimeSpan duration = finishTime.Value - startTime.Value;
+				_dogStatsd.Timer($"horde.{type}.{outcomeStr}.duration", duration.TotalSeconds, tags: tags);
 			}
 
-			JobStepOutcome JobOutcome = Job.Batches.SelectMany(x => x.Steps).Min(x => x.Outcome);
-			DateTime? StartTime = Job.Batches.Select(x => x.StartTimeUtc).Min();
-			DateTime? FinishTime = Job.Batches.Select(x => x.FinishTimeUtc).Max();
-			RecordMetric("job", JobOutcome, StartTime, FinishTime);
+			JobStepOutcome jobOutcome = job.Batches.SelectMany(x => x.Steps).Min(x => x.Outcome);
+			DateTime? startTime = job.Batches.Select(x => x.StartTimeUtc).Min();
+			DateTime? finishTime = job.Batches.Select(x => x.FinishTimeUtc).Max();
+			RecordMetric("job", jobOutcome, startTime, finishTime);
 
 			// TODO: record metrics for individual steps
 			// foreach (IJobStepBatch Batch in Job.Batches)
@@ -508,231 +500,231 @@ namespace Horde.Build.Notifications.Impl
 			return Task.CompletedTask;
 		}
 
-		private async Task<List<IUser>> GetUsersToNotify(EventRecord? Event, ObjectId? NotificationTriggerId, bool bFireTrigger)
+		private async Task<List<IUser>> GetUsersToNotify(EventRecord? eventRecord, ObjectId? notificationTriggerId, bool bFireTrigger)
 		{
-			List<UserId> UserIds = new List<UserId>();
+			List<UserId> userIds = new List<UserId>();
 
 			// Find the notifications for all steps of this type
-			if (Event != null)
+			if (eventRecord != null)
 			{
-				List<ISubscription> Subscriptions = await SubscriptionCollection.FindSubscribersAsync(Event);
-				foreach (ISubscription Subscription in Subscriptions)
+				List<ISubscription> subscriptions = await _subscriptionCollection.FindSubscribersAsync(eventRecord);
+				foreach (ISubscription subscription in subscriptions)
 				{
-					if (Subscription.NotificationType == NotificationType.Slack)
+					if (subscription.NotificationType == NotificationType.Slack)
 					{
-						UserIds.Add(Subscription.UserId);
+						userIds.Add(subscription.UserId);
 					}
 				}
 			}
 
 			// Find the notifications for this particular step
-			if (NotificationTriggerId != null)
+			if (notificationTriggerId != null)
 			{
-				INotificationTrigger? Trigger = await GetNotificationTrigger(NotificationTriggerId, bFireTrigger);
-				if (Trigger != null)
+				INotificationTrigger? trigger = await GetNotificationTrigger(notificationTriggerId, bFireTrigger);
+				if (trigger != null)
 				{
-					foreach (INotificationSubscription Subscription in Trigger.Subscriptions)
+					foreach (INotificationSubscription subscription in trigger.Subscriptions)
 					{
-						if (Subscription.Email)
+						if (subscription.Email)
 						{
 							// TODO?
 						}
-						if (Subscription.Slack)
+						if (subscription.Slack)
 						{
-							UserIds.Add(Subscription.UserId);
+							userIds.Add(subscription.UserId);
 						}
 					}
 				}
 			}
-			return await UserCollection.FindUsersAsync(UserIds);
+			return await _userCollection.FindUsersAsync(userIds);
 		}
 
-		private async Task SendJobStepNotificationsAsync(IJob Job, SubResourceId BatchId, SubResourceId StepId)
+		private async Task SendJobStepNotificationsAsync(IJob job, SubResourceId batchId, SubResourceId stepId)
 		{
-			using IDisposable Scope = Logger.BeginScope("Sending notifications for step {JobId}:{BatchId}:{StepId}", Job.Id, BatchId, StepId);
+			using IDisposable scope = _logger.BeginScope("Sending notifications for step {JobId}:{BatchId}:{StepId}", job.Id, batchId, stepId);
 
-			IJobStepBatch? Batch;
-			if(!Job.TryGetBatch(BatchId, out Batch))
+			IJobStepBatch? batch;
+			if(!job.TryGetBatch(batchId, out batch))
 			{
-				Logger.LogError("Unable to find batch {BatchId} in job {JobId}", BatchId, Job.Id);
+				_logger.LogError("Unable to find batch {BatchId} in job {JobId}", batchId, job.Id);
 				return;
 			}
 
-			IJobStep? Step;
-			if (!Batch.TryGetStep(StepId, out Step))
+			IJobStep? step;
+			if (!batch.TryGetStep(stepId, out step))
 			{
-				Logger.LogError("Unable to find step {StepId} in batch {JobId}:{BatchId}", StepId, Job.Id, BatchId);
+				_logger.LogError("Unable to find step {StepId} in batch {JobId}:{BatchId}", stepId, job.Id, batchId);
 				return;
 			}
 
-			IGraph JobGraph = await GraphCollection.GetAsync(Job.GraphHash);
-			INode Node = JobGraph.GetNode(new NodeRef(Batch.GroupIdx, Step.NodeIdx));
+			IGraph jobGraph = await _graphCollection.GetAsync(job.GraphHash);
+			INode node = jobGraph.GetNode(new NodeRef(batch.GroupIdx, step.NodeIdx));
 
 			// Find the notifications for this particular step
-			EventRecord EventRecord = new StepCompleteEventRecord(Job.StreamId, Job.TemplateId, Node.Name, Step.Outcome);
+			EventRecord eventRecord = new StepCompleteEventRecord(job.StreamId, job.TemplateId, node.Name, step.Outcome);
 
-			List<IUser> UsersToNotify = await GetUsersToNotify(EventRecord, Step.NotificationTriggerId, true);
+			List<IUser> usersToNotify = await GetUsersToNotify(eventRecord, step.NotificationTriggerId, true);
 
 			// If this is not a success notification and the author isn't in the list to notify, add them manually if this is the outcome has gotten worse.
-			int Failures = Job.Batches.Sum(x => x.Steps.Count(y => y.Outcome == JobStepOutcome.Failure));
-			bool FirstFailure = Step.Outcome == JobStepOutcome.Failure && Failures == 1;
-			bool FirstWarning = Step.Outcome == JobStepOutcome.Warnings && Failures == 0 && Job.Batches.Sum(x => x.Steps.Count(y => y.Outcome == JobStepOutcome.Warnings)) == 1;
-			if (Job.StartedByUserId.HasValue && !UsersToNotify.Any(x => x.Id == Job.StartedByUserId) && (FirstFailure || FirstWarning))
-			{
-				Logger.LogInformation("Author {AuthorUserId} is not in notify list but step outcome is {JobStepOutcome}, adding them to the list...", Job.StartedByUserId, Step.Outcome);
-				IUser? AuthorUser = await UserCollection.GetUserAsync(Job.StartedByUserId.Value);
-				if (AuthorUser != null)
+			int failures = job.Batches.Sum(x => x.Steps.Count(y => y.Outcome == JobStepOutcome.Failure));
+			bool firstFailure = step.Outcome == JobStepOutcome.Failure && failures == 1;
+			bool firstWarning = step.Outcome == JobStepOutcome.Warnings && failures == 0 && job.Batches.Sum(x => x.Steps.Count(y => y.Outcome == JobStepOutcome.Warnings)) == 1;
+			if (job.StartedByUserId.HasValue && !usersToNotify.Any(x => x.Id == job.StartedByUserId) && (firstFailure || firstWarning))
+			{
+				_logger.LogInformation("Author {AuthorUserId} is not in notify list but step outcome is {JobStepOutcome}, adding them to the list...", job.StartedByUserId, step.Outcome);
+				IUser? authorUser = await _userCollection.GetUserAsync(job.StartedByUserId.Value);
+				if (authorUser != null)
 				{
-					UsersToNotify.Add(AuthorUser);
+					usersToNotify.Add(authorUser);
 				}
 			}
 
-			if (UsersToNotify.Count == 0)
+			if (usersToNotify.Count == 0)
 			{
-				Logger.LogInformation("No users to notify for step {JobId}:{BatchId}:{StepId}", Job.Id, BatchId, StepId);
+				_logger.LogInformation("No users to notify for step {JobId}:{BatchId}:{StepId}", job.Id, batchId, stepId);
 				return;
 			}
 
-			IStream? JobStream = await StreamService.GetStreamAsync(Job.StreamId);
-			if (JobStream == null)
+			IStream? jobStream = await _streamService.GetStreamAsync(job.StreamId);
+			if (jobStream == null)
 			{
-				Logger.LogError("Unable to find stream {StreamId}", Job.StreamId);
+				_logger.LogError("Unable to find stream {StreamId}", job.StreamId);
 				return;
 			}
 
-			if(Step.LogId == null)
+			if(step.LogId == null)
 			{
-				Logger.LogError("Step does not have a log file");
+				_logger.LogError("Step does not have a log file");
 				return;
 			}
 
-			ILogFile? LogFile = await LogFileService.GetLogFileAsync(Step.LogId.Value);
-			if(LogFile == null)
+			ILogFile? logFile = await _logFileService.GetLogFileAsync(step.LogId.Value);
+			if(logFile == null)
 			{
-				Logger.LogError("Step does not have a log file");
+				_logger.LogError("Step does not have a log file");
 				return;
 			}
 
-			List<ILogEvent> JobStepEvents = await LogFileService.FindLogEventsAsync(LogFile);
-			List<ILogEventData> JobStepEventData = new List<ILogEventData>();
-			foreach (ILogEvent Event in JobStepEvents)
+			List<ILogEvent> jobStepEvents = await _logFileService.FindLogEventsAsync(logFile);
+			List<ILogEventData> jobStepEventData = new List<ILogEventData>();
+			foreach (ILogEvent logEvent in jobStepEvents)
 			{
-				ILogEventData EventData = await LogFileService.GetEventDataAsync(LogFile, Event.LineIndex, Event.LineCount);
-				JobStepEventData.Add(EventData);
+				ILogEventData eventData = await _logFileService.GetEventDataAsync(logFile, logEvent.LineIndex, logEvent.LineCount);
+				jobStepEventData.Add(eventData);
 			}
 
-			foreach (IUser SlackUser in UsersToNotify)
+			foreach (IUser slackUser in usersToNotify)
 			{
-				if(Job.PreflightChange != 0)
+				if(job.PreflightChange != 0)
 				{
-					if(SlackUser.Id != Job.StartedByUserId)
+					if(slackUser.Id != job.StartedByUserId)
 					{
 						continue;
 					}
 				}
-				EnqueueTasks(Sink => Sink.NotifyJobStepCompleteAsync(SlackUser, JobStream, Job, Batch, Step, Node, JobStepEventData));
+				EnqueueTasks(sink => sink.NotifyJobStepCompleteAsync(slackUser, jobStream, job, batch, step, node, jobStepEventData));
 			}
-			Logger.LogDebug("Finished sending notifications for step {JobId}:{BatchId}:{StepId}", Job.Id, BatchId, StepId);
+			_logger.LogDebug("Finished sending notifications for step {JobId}:{BatchId}:{StepId}", job.Id, batchId, stepId);
 		}
 
-		private async Task SendAllLabelNotificationsAsync(IJob Job, IReadOnlyList<(LabelState State, LabelOutcome Outcome)> OldLabelStates, IReadOnlyList<(LabelState, LabelOutcome)> NewLabelStates)
+		private async Task SendAllLabelNotificationsAsync(IJob job, IReadOnlyList<(LabelState State, LabelOutcome Outcome)> oldLabelStates, IReadOnlyList<(LabelState, LabelOutcome)> newLabelStates)
 		{
-			IStream? Stream = await StreamService.GetStreamAsync(Job.StreamId);
-			if (Stream == null)
+			IStream? stream = await _streamService.GetStreamAsync(job.StreamId);
+			if (stream == null)
 			{
-				Logger.LogError("Unable to find stream {StreamId} for job {JobId}", Job.StreamId, Job.Id);
+				_logger.LogError("Unable to find stream {StreamId} for job {JobId}", job.StreamId, job.Id);
 				return;
 			}
 
-			IGraph? Graph = await GraphCollection.GetAsync(Job.GraphHash);
-			if (Graph == null)
+			IGraph? graph = await _graphCollection.GetAsync(job.GraphHash);
+			if (graph == null)
 			{
-				Logger.LogError("Unable to find graph {GraphHash} for job {JobId}", Job.GraphHash, Job.Id);
+				_logger.LogError("Unable to find graph {GraphHash} for job {JobId}", job.GraphHash, job.Id);
 				return;
 			}
 
-			IReadOnlyDictionary<NodeRef, IJobStep> StepForNode = Job.GetStepForNodeMap();
-			for (int LabelIdx = 0; LabelIdx < Graph.Labels.Count; ++LabelIdx)
+			IReadOnlyDictionary<NodeRef, IJobStep> stepForNode = job.GetStepForNodeMap();
+			for (int labelIdx = 0; labelIdx < graph.Labels.Count; ++labelIdx)
 			{
-				(LabelState State, LabelOutcome Outcome) OldLabel = OldLabelStates[LabelIdx];
-				(LabelState State, LabelOutcome Outcome) NewLabel = NewLabelStates[LabelIdx];
-				if (OldLabel != NewLabel)
+				(LabelState State, LabelOutcome Outcome) oldLabel = oldLabelStates[labelIdx];
+				(LabelState State, LabelOutcome Outcome) newLabel = newLabelStates[labelIdx];
+				if (oldLabel != newLabel)
 				{
 					// If the state transitioned from Unspecified to Running, don't update unless the outcome also changed.
-					if (OldLabel.State == LabelState.Unspecified && NewLabel.State == LabelState.Running && OldLabel.Outcome == NewLabel.Outcome)
+					if (oldLabel.State == LabelState.Unspecified && newLabel.State == LabelState.Running && oldLabel.Outcome == newLabel.Outcome)
 					{
 						continue;
 					}
 
 					// If the label isn't complete, don't report on outcome changing to success, this will be reported when the label state becomes complete.
-					if (NewLabel.State != LabelState.Complete && NewLabel.Outcome == LabelOutcome.Success)
+					if (newLabel.State != LabelState.Complete && newLabel.Outcome == LabelOutcome.Success)
 					{
 						return;
 					}
 
-					ILabel Label = Graph.Labels[LabelIdx];
+					ILabel label = graph.Labels[labelIdx];
 
-					EventRecord? EventId;
-					if (String.IsNullOrEmpty(Label.DashboardName))
+					EventRecord? eventId;
+					if (String.IsNullOrEmpty(label.DashboardName))
 					{
-						EventId = null;
+						eventId = null;
 					}
 					else
 					{
-						EventId = new LabelCompleteEventRecord(Job.StreamId, Job.TemplateId, Label.DashboardCategory, Label.DashboardName, NewLabel.Outcome);
+						eventId = new LabelCompleteEventRecord(job.StreamId, job.TemplateId, label.DashboardCategory, label.DashboardName, newLabel.Outcome);
 					}
 
-					ObjectId? TriggerId;
-					if (Job.LabelIdxToTriggerId.TryGetValue(LabelIdx, out ObjectId NewTriggerId))
+					ObjectId? triggerId;
+					if (job.LabelIdxToTriggerId.TryGetValue(labelIdx, out ObjectId newTriggerId))
 					{
-						TriggerId = NewTriggerId;
+						triggerId = newTriggerId;
 					}
 					else
 					{
-						TriggerId = null;
+						triggerId = null;
 					}
 
-					bool bFireTrigger = NewLabel.State == LabelState.Complete;
+					bool bFireTrigger = newLabel.State == LabelState.Complete;
 
-					List<IUser> UsersToNotify = await GetUsersToNotify(EventId, TriggerId, bFireTrigger);
+					List<IUser> usersToNotify = await GetUsersToNotify(eventId, triggerId, bFireTrigger);
 
 					// filter preflight label notifications to only include initiator
-					if (UsersToNotify.Count > 0 && Job.PreflightChange != 0 && Job.StartedByUserId != null)
+					if (usersToNotify.Count > 0 && job.PreflightChange != 0 && job.StartedByUserId != null)
 					{
-						UsersToNotify = UsersToNotify.Where(x => x.Id == Job.StartedByUserId).ToList();
+						usersToNotify = usersToNotify.Where(x => x.Id == job.StartedByUserId).ToList();
 					}
 
-					if (UsersToNotify.Count > 0)
+					if (usersToNotify.Count > 0)
 					{
-						SendLabelUpdateNotifications(Job, Stream, Graph, StepForNode, Graph.Labels[LabelIdx], LabelIdx, NewLabel.Outcome, UsersToNotify);
+						SendLabelUpdateNotifications(job, stream, graph, stepForNode, graph.Labels[labelIdx], labelIdx, newLabel.Outcome, usersToNotify);
 					}
 					else
 					{
-						Logger.LogDebug("No users to notify for label {DashboardName}/{UgsName} in job {JobId}", Graph.Labels[LabelIdx].DashboardName, Graph.Labels[LabelIdx].UgsName, Job.Id);
+						_logger.LogDebug("No users to notify for label {DashboardName}/{UgsName} in job {JobId}", graph.Labels[labelIdx].DashboardName, graph.Labels[labelIdx].UgsName, job.Id);
 					}
 				}
 			}
 		}
 
-		private void SendLabelUpdateNotifications(IJob Job, IStream Stream, IGraph Graph, IReadOnlyDictionary<NodeRef, IJobStep> StepForNode, ILabel Label, int LabelIdx, LabelOutcome Outcome, List<IUser> SlackUsers)
+		private void SendLabelUpdateNotifications(IJob job, IStream stream, IGraph graph, IReadOnlyDictionary<NodeRef, IJobStep> stepForNode, ILabel label, int labelIdx, LabelOutcome outcome, List<IUser> slackUsers)
 		{
-			List<(string, JobStepOutcome, Uri)> StepData = new List<(string, JobStepOutcome, Uri)>();
-			if (Outcome != LabelOutcome.Success)
+			List<(string, JobStepOutcome, Uri)> stepData = new List<(string, JobStepOutcome, Uri)>();
+			if (outcome != LabelOutcome.Success)
 			{
-				foreach (NodeRef IncludedNodeRef in Label.IncludedNodes)
+				foreach (NodeRef includedNodeRef in label.IncludedNodes)
 				{
-					INode IncludedNode = Graph.GetNode(IncludedNodeRef);
-					IJobStep IncludedStep = StepForNode[IncludedNodeRef];
-					StepData.Add((IncludedNode.Name, IncludedStep.Outcome, new Uri($"{Settings.CurrentValue.DashboardUrl}job/{Job.Id}?step={IncludedStep.Id}")));
+					INode includedNode = graph.GetNode(includedNodeRef);
+					IJobStep includedStep = stepForNode[includedNodeRef];
+					stepData.Add((includedNode.Name, includedStep.Outcome, new Uri($"{_settings.CurrentValue.DashboardUrl}job/{job.Id}?step={includedStep.Id}")));
 				}
 			}
 
-			foreach (IUser SlackUser in SlackUsers)
+			foreach (IUser slackUser in slackUsers)
 			{
-				EnqueueTasks(Sink => Sink.NotifyLabelCompleteAsync(SlackUser, Job, Stream, Label, LabelIdx, Outcome, StepData));
+				EnqueueTasks(sink => sink.NotifyLabelCompleteAsync(slackUser, job, stream, label, labelIdx, outcome, stepData));
 			}
 
-			Logger.LogDebug("Finished sending label notifications for label {DashboardName}/{UgsName} in job {JobId}", Label.DashboardName, Label.UgsName, Job.Id);
+			_logger.LogDebug("Finished sending label notifications for label {DashboardName}/{UgsName} in job {JobId}", label.DashboardName, label.UgsName, job.Id);
 		}
 	}
 }
diff --git a/Engine/Source/Programs/Horde/Horde.Build/Notifications/Impl/NotificationTriggerCollection.cs b/Engine/Source/Programs/Horde/Horde.Build/Notifications/Impl/NotificationTriggerCollection.cs
--- a/Engine/Source/Programs/Horde/Horde.Build/Notifications/Impl/NotificationTriggerCollection.cs
+++ b/Engine/Source/Programs/Horde/Horde.Build/Notifications/Impl/NotificationTriggerCollection.cs
@@ -1,15 +1,14 @@
 // Copyright Epic Games, Inc. All Rights Reserved.
 
+using System.Collections.Generic;
+using System.Linq;
+using System.Threading.Tasks;
 using Horde.Build.Models;
 using Horde.Build.Services;
 using Horde.Build.Utilities;
 using MongoDB.Bson;
 using MongoDB.Bson.Serialization.Attributes;
 using MongoDB.Driver;
-using System;
-using System.Collections.Generic;
-using System.Linq;
-using System.Threading.Tasks;
 
 namespace Horde.Build.Collections.Impl
 {
@@ -35,18 +34,18 @@ namespace Horde.Build.Collections.Impl
 				User = null!;
 			}
 
-			public SubscriptionDocument(UserId UserId, bool Email, bool Slack)
+			public SubscriptionDocument(UserId userId, bool email, bool slack)
 			{
-				this.UserId = UserId;
-				this.Email = Email;
-				this.Slack = Slack;
+				UserId = userId;
+				Email = email;
+				Slack = slack;
 			}
 
-			public SubscriptionDocument(INotificationSubscription Subscription)
+			public SubscriptionDocument(INotificationSubscription subscription)
 			{
-				this.UserId = Subscription.UserId;
-				this.Email = Subscription.Email;
-				this.Slack = Subscription.Slack;
+				UserId = subscription.UserId;
+				Email = subscription.Email;
+				Slack = subscription.Slack;
 			}
 		}
 
@@ -60,71 +59,71 @@ namespace Horde.Build.Collections.Impl
 			IReadOnlyList<INotificationSubscription> INotificationTrigger.Subscriptions => Subscriptions;
 		}
 
-		IMongoCollection<TriggerDocument> Triggers;
-		IUserCollection UserCollection;
+		readonly IMongoCollection<TriggerDocument> _triggers;
+		readonly IUserCollection _userCollection;
 
 		/// <summary>
 		/// Constructor
 		/// </summary>
-		/// <param name="DatabaseService">The database singleton</param>
-		/// <param name="UserCollection"></param>
-		public NotificationTriggerCollection(DatabaseService DatabaseService, IUserCollection UserCollection)
+		/// <param name="databaseService">The database singleton</param>
+		/// <param name="userCollection"></param>
+		public NotificationTriggerCollection(DatabaseService databaseService, IUserCollection userCollection)
 		{
-			this.Triggers = DatabaseService.GetCollection<TriggerDocument>("NotificationTriggers");
-			this.UserCollection = UserCollection;
+			_triggers = databaseService.GetCollection<TriggerDocument>("NotificationTriggers");
+			_userCollection = userCollection;
 		}
 
 		/// <inheritdoc/>
-		public async Task<INotificationTrigger?> GetAsync(ObjectId TriggerId)
+		public async Task<INotificationTrigger?> GetAsync(ObjectId triggerId)
 		{
-			TriggerDocument? Trigger = await Triggers.Find(x => x.Id == TriggerId).FirstOrDefaultAsync();
-			if (Trigger != null)
+			TriggerDocument? trigger = await _triggers.Find(x => x.Id == triggerId).FirstOrDefaultAsync();
+			if (trigger != null)
 			{
-				for (int Idx = 0; Idx < Trigger.Subscriptions.Count; Idx++)
+				for (int idx = 0; idx < trigger.Subscriptions.Count; idx++)
 				{
-					SubscriptionDocument Subscription = Trigger.Subscriptions[Idx];
-					if (Subscription.User != null)
+					SubscriptionDocument subscription = trigger.Subscriptions[idx];
+					if (subscription.User != null)
 					{
-						IUser? User = await UserCollection.FindUserByLoginAsync(Subscription.User);
-						if (User == null)
+						IUser? user = await _userCollection.FindUserByLoginAsync(subscription.User);
+						if (user == null)
 						{
-							Trigger.Subscriptions.RemoveAt(Idx);
-							Idx--;
+							trigger.Subscriptions.RemoveAt(idx);
+							idx--;
 						}
 						else
 						{
-							Subscription.UserId = User.Id;
-							Subscription.User = null;
+							subscription.UserId = user.Id;
+							subscription.User = null;
 						}
 					}
 				}
 			}
-			return Trigger;
+			return trigger;
 		}
 
 		/// <inheritdoc/>
-		public async Task<INotificationTrigger> FindOrAddAsync(ObjectId TriggerId)
+		public async Task<INotificationTrigger> FindOrAddAsync(ObjectId triggerId)
 		{
 			for (; ; )
 			{
 				// Find an existing trigger
-				INotificationTrigger? Existing = await GetAsync(TriggerId);
-				if (Existing != null)
+				INotificationTrigger? existing = await GetAsync(triggerId);
+				if (existing != null)
 				{
-					return Existing;
+					return existing;
 				}
 
 				// Try to insert a new document
 				try
 				{
-					TriggerDocument NewDocument = new TriggerDocument();
-					NewDocument.Id = TriggerId;
-					await Triggers.InsertOneAsync(NewDocument);
-					return NewDocument;
+					TriggerDocument newDocument = new TriggerDocument();
+					newDocument.Id = triggerId;
+					await _triggers.InsertOneAsync(newDocument);
+					return newDocument;
 				}
-				catch (MongoWriteException Ex)
+				catch (MongoWriteException ex)
 				{
-					if (Ex.WriteError.Category != ServerErrorCategory.DuplicateKey)
+					if (ex.WriteError.Category != ServerErrorCategory.DuplicateKey)
 					{
 						throw;
 					}
@@ -135,62 +134,62 @@ namespace Horde.Build.Collections.Impl
 		/// <summary>
 		/// Updates an existing document
 		/// </summary>
-		/// <param name="Trigger">The trigger to update</param>
-		/// <param name="Transaction">The update definition</param>
+		/// <param name="trigger">The trigger to update</param>
+		/// <param name="transaction">The update definition</param>
 		/// <returns>The updated document</returns>
-		async Task<INotificationTrigger?> TryUpdateAsync(INotificationTrigger Trigger, TransactionBuilder<TriggerDocument> Transaction)
+		async Task<INotificationTrigger?> TryUpdateAsync(INotificationTrigger trigger, TransactionBuilder<TriggerDocument> transaction)
 		{
-			TriggerDocument Document = (TriggerDocument)Trigger;
-			int NextUpdateIndex = Document.UpdateIndex + 1;
+			TriggerDocument document = (TriggerDocument)trigger;
+			int nextUpdateIndex = document.UpdateIndex + 1;
 
-			FilterDefinition<TriggerDocument> Filter = Builders<TriggerDocument>.Filter.Expr(x => x.Id == Trigger.Id && x.UpdateIndex == Document.UpdateIndex);
-			UpdateDefinition<TriggerDocument> Update = Transaction.ToUpdateDefinition().Set(x => x.UpdateIndex, NextUpdateIndex);
+			FilterDefinition<TriggerDocument> filter = Builders<TriggerDocument>.Filter.Expr(x => x.Id == trigger.Id && x.UpdateIndex == document.UpdateIndex);
+			UpdateDefinition<TriggerDocument> update = transaction.ToUpdateDefinition().Set(x => x.UpdateIndex, nextUpdateIndex);
 
-			UpdateResult Result = await Triggers.UpdateOneAsync(Filter, Update);
-			if (Result.ModifiedCount > 0)
+			UpdateResult result = await _triggers.UpdateOneAsync(filter, update);
+			if (result.ModifiedCount > 0)
 			{
-				Transaction.ApplyTo(Document);
-				Document.UpdateIndex = NextUpdateIndex;
-				return Document;
+				transaction.ApplyTo(document);
+				document.UpdateIndex = nextUpdateIndex;
+				return document;
 			}
 
 			return null;
 		}
 
 		/// <inheritdoc/>
-		public async Task DeleteAsync(ObjectId TriggerId)
+		public async Task DeleteAsync(ObjectId triggerId)
 		{
-			await Triggers.DeleteOneAsync(x => x.Id == TriggerId);
+			await _triggers.DeleteOneAsync(x => x.Id == triggerId);
 		}
 
 		/// <inheritdoc/>
-		public async Task DeleteAsync(List<ObjectId> TriggerIds)
+		public async Task DeleteAsync(List<ObjectId> triggerIds)
 		{
-			FilterDefinition<TriggerDocument> Filter = Builders<TriggerDocument>.Filter.In(x => x.Id, TriggerIds);
-			await Triggers.DeleteManyAsync(Filter);
+			FilterDefinition<TriggerDocument> filter = Builders<TriggerDocument>.Filter.In(x => x.Id, triggerIds);
+			await _triggers.DeleteManyAsync(filter);
 		}
 
 		/// <inheritdoc/>
-		public async Task<INotificationTrigger?> FireAsync(INotificationTrigger Trigger)
+		public async Task<INotificationTrigger?> FireAsync(INotificationTrigger trigger)
 		{
-			if (Trigger.Fired)
+			if (trigger.Fired)
 			{
 				return null;
 			}
 
 			for (; ; )
 			{
-				TransactionBuilder<TriggerDocument> Transaction = new TransactionBuilder<TriggerDocument>();
-				Transaction.Set(x => x.Fired, true);
+				TransactionBuilder<TriggerDocument> transaction = new TransactionBuilder<TriggerDocument>();
+				transaction.Set(x => x.Fired, true);
 
-				INotificationTrigger? NewTrigger = await TryUpdateAsync(Trigger, Transaction);
-				if (NewTrigger != null)
+				INotificationTrigger? newTrigger = await TryUpdateAsync(trigger, transaction);
+				if (newTrigger != null)
 				{
-					return NewTrigger;
+					return newTrigger;
 				}
 
-				NewTrigger = await FindOrAddAsync(Trigger.Id); // Need to add to prevent race condition on triggering vs adding
-				if (NewTrigger == null || NewTrigger.Fired)
+				newTrigger = await FindOrAddAsync(trigger.Id); // Need to add to prevent race condition on triggering vs adding
+				if (newTrigger == null || newTrigger.Fired)
 				{
 					return null;
 				}
@@ -198,43 +197,43 @@ namespace Horde.Build.Collections.Impl
 		}
 
 		/// <inheritdoc/>
-		public async Task<INotificationTrigger?> UpdateSubscriptionsAsync(INotificationTrigger Trigger, UserId UserId, bool? Email, bool? Slack)
+		public async Task<INotificationTrigger?> UpdateSubscriptionsAsync(INotificationTrigger trigger, UserId userId, bool? email, bool? slack)
 		{
 			for (; ; )
 			{
 				// If the trigger has already fired, don't add a new subscription to it
-				if(Trigger.Fired)
+				if(trigger.Fired)
 				{
-					return Trigger;
+					return trigger;
 				}
 
 				// Try to update the trigger
-				List<SubscriptionDocument> NewSubscriptions = new List<SubscriptionDocument>();
-				NewSubscriptions.AddRange(Trigger.Subscriptions.Select(x => new SubscriptionDocument(x)));
+				List<SubscriptionDocument> newSubscriptions = new List<SubscriptionDocument>();
+				newSubscriptions.AddRange(trigger.Subscriptions.Select(x => new SubscriptionDocument(x)));
 
-				SubscriptionDocument? NewSubscription = NewSubscriptions.FirstOrDefault(x => x.UserId == UserId);
-				if (NewSubscription == null)
+				SubscriptionDocument? newSubscription = newSubscriptions.FirstOrDefault(x => x.UserId == userId);
+				if (newSubscription == null)
 				{
-					NewSubscription = new SubscriptionDocument(UserId, Email ?? false, Slack ?? false);
-					NewSubscriptions.Add(NewSubscription);
+					newSubscription = new SubscriptionDocument(userId, email ?? false, slack ?? false);
+					newSubscriptions.Add(newSubscription);
 				}
 				else
 				{
-					NewSubscription.Email = Email ?? NewSubscription.Email;
-					NewSubscription.Slack = Slack ?? NewSubscription.Slack;
+					newSubscription.Email = email ?? newSubscription.Email;
+					newSubscription.Slack = slack ?? newSubscription.Slack;
 				}
 
-				TransactionBuilder<TriggerDocument> Transaction = new TransactionBuilder<TriggerDocument>();
-				Transaction.Set(x => x.Subscriptions, NewSubscriptions);
+				TransactionBuilder<TriggerDocument> transaction = new TransactionBuilder<TriggerDocument>();
+				transaction.Set(x => x.Subscriptions, newSubscriptions);
 
-				INotificationTrigger? NewTrigger = await TryUpdateAsync(Trigger, Transaction);
-				if (NewTrigger != null)
+				INotificationTrigger? newTrigger = await TryUpdateAsync(trigger, transaction);
+				if (newTrigger != null)
 				{
-					return NewTrigger;
+					return newTrigger;
 				}
 
-				NewTrigger = await GetAsync(Trigger.Id);
-				if (NewTrigger == null)
+				newTrigger = await GetAsync(trigger.Id);
+				if (newTrigger == null)
 				{
 					return null;
 				}
