Horde: Fix coding conventions within Horde.Build.

diff --git a/Engine/Source/Programs/Horde/Horde.Build/Notifications/Impl/SlackNotificationSink.cs b/Engine/Source/Programs/Horde/Horde.Build/Notifications/Impl/SlackNotificationSink.cs
--- a/Engine/Source/Programs/Horde/Horde.Build/Notifications/Impl/SlackNotificationSink.cs
+++ b/Engine/Source/Programs/Horde/Horde.Build/Notifications/Impl/SlackNotificationSink.cs
@@ -1,14 +1,26 @@
 // Copyright Epic Games, Inc. All Rights Reserved.
 
+using System;
+using System.Collections.Generic;
+using System.Drawing;
+using System.Globalization;
+using System.Linq;
+using System.Net.Http;
+using System.Net.WebSockets;
+using System.Text;
+using System.Text.Json;
+using System.Text.Json.Serialization;
+using System.Text.RegularExpressions;
+using System.Threading;
+using System.Threading.Tasks;
 using EpicGames.Core;
-using HordeCommon;
 using Horde.Build.Api;
 using Horde.Build.Collections;
 using Horde.Build.Models;
 using Horde.Build.Services;
 using Horde.Build.Utilities;
-using Horde.Build.Utilities.BlockKit;
 using Horde.Build.Utilities.Slack.BlockKit;
+using HordeCommon;
 using Microsoft.Extensions.Caching.Memory;
 using Microsoft.Extensions.Hosting;
 using Microsoft.Extensions.Logging;
@@ -16,23 +28,6 @@ using Microsoft.Extensions.Options;
 using MongoDB.Bson;
 using MongoDB.Bson.Serialization.Attributes;
 using MongoDB.Driver;
-using Newtonsoft.Json.Linq;
-using System;
-using System.Collections.Concurrent;
-using System.Collections.Generic;
-using System.Diagnostics;
-using System.Drawing;
-using System.Globalization;
-using System.Linq;
-using System.Net.Http;
-using System.Net.WebSockets;
-using System.Security.Claims;
-using System.Text;
-using System.Text.Json;
-using System.Text.Json.Serialization;
-using System.Text.RegularExpressions;
-using System.Threading;
-using System.Threading.Tasks;
 
 namespace Horde.Build.Notifications.Impl
 {
@@ -179,61 +174,61 @@ namespace Horde.Build.Notifications.Impl
 			{
 			}
 
-			public SlackUser(UserId Id, UserInfo? Info)
+			public SlackUser(UserId id, UserInfo? info)
 			{
-				this.Id = Id;
-				this.SlackUserId = Info?.UserName;
-				if (Info != null && Info.Profile != null && Info.Profile.IsCustomImage)
+				Id = id;
+				SlackUserId = info?.UserName;
+				if (info != null && info.Profile != null && info.Profile.IsCustomImage)
 				{
-					this.Image24 = Info.Profile.Image24;
-					this.Image32 = Info.Profile.Image32;
-					this.Image48 = Info.Profile.Image48;
-					this.Image72 = Info.Profile.Image72;
+					Image24 = info.Profile.Image24;
+					Image32 = info.Profile.Image32;
+					Image48 = info.Profile.Image48;
+					Image72 = info.Profile.Image72;
 				}
-				this.Time = DateTime.UtcNow;
-				this.Version = CurrentVersion;
+				Time = DateTime.UtcNow;
+				Version = CurrentVersion;
 			}
 		}
 
-		IIssueService IssueService;
-		IUserCollection UserCollection;
-		ILogFileService LogFileService;
-		StreamService StreamService;
-		ServerSettings Settings;
-		IMongoCollection<MessageStateDocument> MessageStates;
-		IMongoCollection<SlackUser> SlackUsers;
-		HashSet<string>? AllowUsers;
-		ILogger Logger;
+		readonly IIssueService _issueService;
+		readonly IUserCollection _userCollection;
+		readonly ILogFileService _logFileService;
+		readonly StreamService _streamService;
+		readonly ServerSettings _settings;
+		readonly IMongoCollection<MessageStateDocument> _messageStates;
+		readonly IMongoCollection<SlackUser> _slackUsers;
+		readonly HashSet<string>? _allowUsers;
+		readonly ILogger _logger;
 
 		/// <summary>
 		/// Map of email address to Slack user ID.
 		/// </summary>
-		private MemoryCache UserCache = new MemoryCache(new MemoryCacheOptions());
+		private readonly MemoryCache _userCache = new MemoryCache(new MemoryCacheOptions());
 
 		/// <summary>
 		/// Constructor
 		/// </summary>
-		/// <param name="DatabaseService"></param>
-		/// <param name="IssueService"></param>
-		/// <param name="UserCollection">The user collection</param>
-		/// <param name="LogFileService"></param>
-		/// <param name="StreamService"></param>
-		/// <param name="Settings">The current configuration settings</param>
-		/// <param name="Logger">Logging device</param>
-		public SlackNotificationSink(DatabaseService DatabaseService, IIssueService IssueService, IUserCollection UserCollection, ILogFileService LogFileService, StreamService StreamService, IOptions<ServerSettings> Settings, ILogger<SlackNotificationSink> Logger)
+		/// <param name="databaseService"></param>
+		/// <param name="issueService"></param>
+		/// <param name="userCollection">The user collection</param>
+		/// <param name="logFileService"></param>
+		/// <param name="streamService"></param>
+		/// <param name="settings">The current configuration settings</param>
+		/// <param name="logger">Logging device</param>
+		public SlackNotificationSink(DatabaseService databaseService, IIssueService issueService, IUserCollection userCollection, ILogFileService logFileService, StreamService streamService, IOptions<ServerSettings> settings, ILogger<SlackNotificationSink> logger)
 		{
-			this.IssueService = IssueService;
-			this.UserCollection = UserCollection;
-			this.LogFileService = LogFileService;
-			this.StreamService = StreamService;
-			this.Settings = Settings.Value;
-			this.MessageStates = DatabaseService.Database.GetCollection<MessageStateDocument>("Slack");
-			this.SlackUsers = DatabaseService.Database.GetCollection<SlackUser>("Slack.UsersV2");
-			this.Logger = Logger;
+			_issueService = issueService;
+			_userCollection = userCollection;
+			_logFileService = logFileService;
+			_streamService = streamService;
+			_settings = settings.Value;
+			_messageStates = databaseService.Database.GetCollection<MessageStateDocument>("Slack");
+			_slackUsers = databaseService.Database.GetCollection<SlackUser>("Slack.UsersV2");
+			_logger = logger;
 
-			if (!String.IsNullOrEmpty(Settings.Value.SlackUsers))
+			if (!String.IsNullOrEmpty(settings.Value.SlackUsers))
 			{
-				AllowUsers = new HashSet<string>(Settings.Value.SlackUsers.Split(','), StringComparer.OrdinalIgnoreCase);
+				_allowUsers = new HashSet<string>(settings.Value.SlackUsers.Split(','), StringComparer.OrdinalIgnoreCase);
 			}
 		}
 
@@ -242,7 +237,7 @@ namespace Horde.Build.Notifications.Impl
 		{
 			base.Dispose();
 
-			UserCache.Dispose();
+			_userCache.Dispose();
 
 			GC.SuppressFinalize(this);
 		}
@@ -250,190 +245,190 @@ namespace Horde.Build.Notifications.Impl
 		#region Avatars
 
 		/// <inheritdoc/>
-		public async Task<IAvatar?> GetAvatarAsync(IUser User)
+		public async Task<IAvatar?> GetAvatarAsync(IUser user)
 		{
-			return await GetSlackUser(User);
+			return await GetSlackUser(user);
 		}
 
 		#endregion
 
 		#region Message state 
 
-		async Task<(MessageStateDocument, bool)> AddOrUpdateMessageStateAsync(string Recipient, string EventId, UserId? UserId, string Digest)
+		async Task<(MessageStateDocument, bool)> AddOrUpdateMessageStateAsync(string recipient, string eventId, UserId? userId, string digest)
 		{
-			ObjectId NewId = ObjectId.GenerateNewId();
+			ObjectId newId = ObjectId.GenerateNewId();
 
-			FilterDefinition<MessageStateDocument> Filter = Builders<MessageStateDocument>.Filter.Eq(x => x.Recipient, Recipient) & Builders<MessageStateDocument>.Filter.Eq(x => x.EventId, EventId);
-			UpdateDefinition<MessageStateDocument> Update = Builders<MessageStateDocument>.Update.SetOnInsert(x => x.Id, NewId).Set(x => x.UserId, UserId).Set(x => x.Digest, Digest);
+			FilterDefinition<MessageStateDocument> filter = Builders<MessageStateDocument>.Filter.Eq(x => x.Recipient, recipient) & Builders<MessageStateDocument>.Filter.Eq(x => x.EventId, eventId);
+			UpdateDefinition<MessageStateDocument> update = Builders<MessageStateDocument>.Update.SetOnInsert(x => x.Id, newId).Set(x => x.UserId, userId).Set(x => x.Digest, digest);
 
-			MessageStateDocument State = await MessageStates.FindOneAndUpdateAsync(Filter, Update, new FindOneAndUpdateOptions<MessageStateDocument> { IsUpsert = true, ReturnDocument = ReturnDocument.After });
-			return (State, State.Id == NewId);
+			MessageStateDocument state = await _messageStates.FindOneAndUpdateAsync(filter, update, new FindOneAndUpdateOptions<MessageStateDocument> { IsUpsert = true, ReturnDocument = ReturnDocument.After });
+			return (state, state.Id == newId);
 		}
 
-		async Task SetMessageTimestampAsync(ObjectId MessageId, string Channel, string Ts)
+		async Task SetMessageTimestampAsync(ObjectId messageId, string channel, string ts)
 		{
-			FilterDefinition<MessageStateDocument> Filter = Builders<MessageStateDocument>.Filter.Eq(x => x.Id, MessageId);
-			UpdateDefinition<MessageStateDocument> Update = Builders<MessageStateDocument>.Update.Set(x => x.Channel, Channel).Set(x => x.Ts, Ts);
-			await MessageStates.FindOneAndUpdateAsync(Filter, Update);
+			FilterDefinition<MessageStateDocument> filter = Builders<MessageStateDocument>.Filter.Eq(x => x.Id, messageId);
+			UpdateDefinition<MessageStateDocument> update = Builders<MessageStateDocument>.Update.Set(x => x.Channel, channel).Set(x => x.Ts, ts);
+			await _messageStates.FindOneAndUpdateAsync(filter, update);
 		}
 
 		#endregion
 		
 		/// <inheritdoc/>
-		public async Task NotifyJobScheduledAsync(List<JobScheduledNotification> Notifications)
+		public async Task NotifyJobScheduledAsync(List<JobScheduledNotification> notifications)
 		{
-			if (Settings.JobNotificationChannel != null)
+			if (_settings.JobNotificationChannel != null)
 			{
-				string JobIds = string.Join(", ", Notifications.Select(x => x.JobId));
-				Logger.LogInformation("Sending Slack notification for scheduled job IDs {JobIds} to channel {SlackChannel}", JobIds, Settings.JobNotificationChannel);
-				await SendJobScheduledOnEmptyAutoScaledPoolMessageAsync($"#{Settings.JobNotificationChannel}", Notifications);
+				string jobIds = String.Join(", ", notifications.Select(x => x.JobId));
+				_logger.LogInformation("Sending Slack notification for scheduled job IDs {JobIds} to channel {SlackChannel}", jobIds, _settings.JobNotificationChannel);
+				await SendJobScheduledOnEmptyAutoScaledPoolMessageAsync($"#{_settings.JobNotificationChannel}", notifications);
 			}
 		}
 		
-		private async Task SendJobScheduledOnEmptyAutoScaledPoolMessageAsync(string Recipient, List<JobScheduledNotification> Notifications)
+		private async Task SendJobScheduledOnEmptyAutoScaledPoolMessageAsync(string recipient, List<JobScheduledNotification> notifications)
 		{
-			string JobIds = string.Join(", ", Notifications.Select(x => x.JobId));
+			string jobIds = String.Join(", ", notifications.Select(x => x.JobId));
 				
-			StringBuilder Sb = new();
-			foreach (JobScheduledNotification Notification in Notifications)
+			StringBuilder sb = new();
+			foreach (JobScheduledNotification notification in notifications)
 			{
-				string JobUrl = Settings.DashboardUrl + "/job/" + Notification.JobId;
-				Sb.AppendLine($"Job `{Notification.JobName}` with ID <{JobUrl}|{Notification.JobId}> in pool `{Notification.PoolName}`");
+				string jobUrl = _settings.DashboardUrl + "/job/" + notification.JobId;
+				sb.AppendLine($"Job `{notification.JobName}` with ID <{jobUrl}|{notification.JobId}> in pool `{notification.PoolName}`");
 			}
 			
-			Color OutcomeColor = BlockKitAttachmentColors.Warning;
-			BlockKitAttachment Attachment = new BlockKitAttachment();
-			Attachment.Color = OutcomeColor;
-			Attachment.FallbackText = $"Job(s) scheduled in an auto-scaled pool with no agents online. Job IDs {JobIds}";
-
-			Attachment.Blocks.Add(new HeaderBlock($"Jobs scheduled in empty pool", false, true));
-			Attachment.Blocks.Add(new SectionBlock($"One or more jobs were scheduled in an auto-scaled pool but with no current agents online."));
-			Attachment.Blocks.Add(new SectionBlock(Sb.ToString()));
+			Color outcomeColor = BlockKitAttachmentColors.Warning;
+			BlockKitAttachment attachment = new BlockKitAttachment();
+			attachment.Color = outcomeColor;
+			attachment.FallbackText = $"Job(s) scheduled in an auto-scaled pool with no agents online. Job IDs {jobIds}";
+
+			attachment.Blocks.Add(new HeaderBlock($"Jobs scheduled in empty pool", false, true));
+			attachment.Blocks.Add(new SectionBlock($"One or more jobs were scheduled in an auto-scaled pool but with no current agents online."));
+			attachment.Blocks.Add(new SectionBlock(sb.ToString()));
 			
-			await SendMessageAsync(Recipient, Attachments: new[] { Attachment });
+			await SendMessageAsync(recipient, attachments: new[] { attachment });
 		}
 
 		#region Job Complete
 
 		/// <inheritdoc/>
-		public async Task NotifyJobCompleteAsync(IStream JobStream, IJob Job, IGraph Graph, LabelOutcome Outcome)
+		public async Task NotifyJobCompleteAsync(IStream jobStream, IJob job, IGraph graph, LabelOutcome outcome)
 		{
-			if (Job.NotificationChannel != null)
+			if (job.NotificationChannel != null)
 			{
-				await SendJobCompleteNotificationToChannelAsync(Job.NotificationChannel, Job.NotificationChannelFilter, JobStream, Job, Graph, Outcome);
+				await SendJobCompleteNotificationToChannelAsync(job.NotificationChannel, job.NotificationChannelFilter, jobStream, job, graph, outcome);
 			}
-			if (JobStream.NotificationChannel != null)
+			if (jobStream.NotificationChannel != null)
 			{
-				await SendJobCompleteNotificationToChannelAsync(JobStream.NotificationChannel, JobStream.NotificationChannelFilter, JobStream, Job, Graph, Outcome);
+				await SendJobCompleteNotificationToChannelAsync(jobStream.NotificationChannel, jobStream.NotificationChannelFilter, jobStream, job, graph, outcome);
 			}
 		}
 
-		Task SendJobCompleteNotificationToChannelAsync(string NotificationChannel, string? NotificationFilter, IStream JobStream, IJob Job, IGraph Graph, LabelOutcome Outcome)
+		Task SendJobCompleteNotificationToChannelAsync(string notificationChannel, string? notificationFilter, IStream jobStream, IJob job, IGraph graph, LabelOutcome outcome)
 		{
-			if (NotificationFilter != null)
+			if (notificationFilter != null)
 			{
-				List<LabelOutcome> Outcomes = new List<LabelOutcome>();
-				foreach (string FilterOption in NotificationFilter.Split('|'))
+				List<LabelOutcome> outcomes = new List<LabelOutcome>();
+				foreach (string filterOption in notificationFilter.Split('|'))
 				{
-					LabelOutcome Result;
-					if (Enum.TryParse(FilterOption, out Result))
+					LabelOutcome result;
+					if (Enum.TryParse(filterOption, out result))
 					{
-						Outcomes.Add(Result);
+						outcomes.Add(result);
 					}
 					else
 					{
-						Logger.LogWarning("Invalid filter option {Option} specified in job filter {NotificationChannelFilter} in job {JobId} or stream {StreamId}", FilterOption, NotificationFilter, Job.Id, Job.StreamId);
+						_logger.LogWarning("Invalid filter option {Option} specified in job filter {NotificationChannelFilter} in job {JobId} or stream {StreamId}", filterOption, notificationFilter, job.Id, job.StreamId);
 					}
 				}
-				if (!Outcomes.Contains(Outcome))
+				if (!outcomes.Contains(outcome))
 				{
 					return Task.CompletedTask;
 				}
 			}
-			return SendJobCompleteMessageAsync(NotificationChannel, JobStream, Job, Graph);
+			return SendJobCompleteMessageAsync(notificationChannel, jobStream, job, graph);
 		}
 
 		/// <inheritdoc/>
-		public async Task NotifyJobCompleteAsync(IUser SlackUser, IStream JobStream, IJob Job, IGraph Graph, LabelOutcome Outcome)
+		public async Task NotifyJobCompleteAsync(IUser slackUser, IStream jobStream, IJob job, IGraph graph, LabelOutcome outcome)
 		{
-			string? SlackUserId = await GetSlackUserId(SlackUser);
-			if (SlackUserId != null)
+			string? slackUserId = await GetSlackUserId(slackUser);
+			if (slackUserId != null)
 			{
-				await SendJobCompleteMessageAsync(SlackUserId, JobStream, Job, Graph);
+				await SendJobCompleteMessageAsync(slackUserId, jobStream, job, graph);
 			}
 		}
 
-		private Task SendJobCompleteMessageAsync(string Recipient, IStream Stream, IJob Job, IGraph Graph)
+		private Task SendJobCompleteMessageAsync(string recipient, IStream stream, IJob job, IGraph graph)
 		{
-			JobStepOutcome JobOutcome = Job.Batches.SelectMany(x => x.Steps).Min(x => x.Outcome);
-			Logger.LogInformation("Sending Slack notification for job {JobId} outcome {Outcome} to {SlackUser}", Job.Id, JobOutcome, Recipient);
+			JobStepOutcome jobOutcome = job.Batches.SelectMany(x => x.Steps).Min(x => x.Outcome);
+			_logger.LogInformation("Sending Slack notification for job {JobId} outcome {Outcome} to {SlackUser}", job.Id, jobOutcome, recipient);
 
-			Uri JobLink = new Uri($"{Settings.DashboardUrl}job/{Job.Id}");
+			Uri jobLink = new Uri($"{_settings.DashboardUrl}job/{job.Id}");
 
-			Color OutcomeColor = JobOutcome == JobStepOutcome.Failure ? BlockKitAttachmentColors.Error : JobOutcome == JobStepOutcome.Warnings ? BlockKitAttachmentColors.Warning : BlockKitAttachmentColors.Success;
+			Color outcomeColor = jobOutcome == JobStepOutcome.Failure ? BlockKitAttachmentColors.Error : jobOutcome == JobStepOutcome.Warnings ? BlockKitAttachmentColors.Warning : BlockKitAttachmentColors.Success;
 
-			BlockKitAttachment Attachment = new BlockKitAttachment();
-			Attachment.FallbackText = $"{Stream.Name} - {GetJobChangeText(Job)} - {Job.Name} - {JobOutcome}";
-			Attachment.Color = OutcomeColor;
-			Attachment.Blocks.Add(new SectionBlock($"*<{JobLink}|{Stream.Name} - {GetJobChangeText(Job)} - {Job.Name}>*"));
-			if (JobOutcome == JobStepOutcome.Success)
+			BlockKitAttachment attachment = new BlockKitAttachment();
+			attachment.FallbackText = $"{stream.Name} - {GetJobChangeText(job)} - {job.Name} - {jobOutcome}";
+			attachment.Color = outcomeColor;
+			attachment.Blocks.Add(new SectionBlock($"*<{jobLink}|{stream.Name} - {GetJobChangeText(job)} - {job.Name}>*"));
+			if (jobOutcome == JobStepOutcome.Success)
 			{
-				Attachment.Blocks.Add(new SectionBlock($"*Job Succeeded*"));
+				attachment.Blocks.Add(new SectionBlock($"*Job Succeeded*"));
 			}
 			else
 			{
-				List<string> FailedStepStrings = new List<string>();
-				List<string> WarningStepStrings = new List<string>();
+				List<string> failedStepStrings = new List<string>();
+				List<string> warningStepStrings = new List<string>();
 
-				IReadOnlyDictionary<NodeRef, IJobStep> NodeToStep = Job.GetStepForNodeMap();
-				foreach ((NodeRef NodeRef, IJobStep Step) in NodeToStep)
+				IReadOnlyDictionary<NodeRef, IJobStep> nodeToStep = job.GetStepForNodeMap();
+				foreach ((NodeRef nodeRef, IJobStep step) in nodeToStep)
 				{
-					INode StepNode = Graph.GetNode(NodeRef);
-					string StepName = $"<{JobLink}?step={Step.Id}|{StepNode.Name}>";
-					if (Step.Outcome == JobStepOutcome.Failure)
+					INode stepNode = graph.GetNode(nodeRef);
+					string stepName = $"<{jobLink}?step={step.Id}|{stepNode.Name}>";
+					if (step.Outcome == JobStepOutcome.Failure)
 					{
-						FailedStepStrings.Add(StepName);
+						failedStepStrings.Add(stepName);
 					}
-					else if (Step.Outcome == JobStepOutcome.Warnings)
+					else if (step.Outcome == JobStepOutcome.Warnings)
 					{
-						WarningStepStrings.Add(StepName);
+						warningStepStrings.Add(stepName);
 					}
 				}
 
-				if (FailedStepStrings.Any())
+				if (failedStepStrings.Any())
 				{
-					string Msg = $"*Errors*\n{string.Join(", ", FailedStepStrings)}";
-					Attachment.Blocks.Add(new SectionBlock(Msg.Substring(0, Math.Min(Msg.Length, 3000))));
+					string msg = $"*Errors*\n{String.Join(", ", failedStepStrings)}";
+					attachment.Blocks.Add(new SectionBlock(msg.Substring(0, Math.Min(msg.Length, 3000))));
 				}
-				else if (WarningStepStrings.Any())
+				else if (warningStepStrings.Any())
 				{
-					string Msg = $"*Warnings*\n{string.Join(", ", WarningStepStrings)}";
-					Attachment.Blocks.Add(new SectionBlock(Msg.Substring(0, Math.Min(Msg.Length, 3000))));
+					string msg = $"*Warnings*\n{String.Join(", ", warningStepStrings)}";
+					attachment.Blocks.Add(new SectionBlock(msg.Substring(0, Math.Min(msg.Length, 3000))));
 				}
 			}
 
-			if (Job.AutoSubmit)
+			if (job.AutoSubmit)
 			{
-				Attachment.Blocks.Add(new DividerBlock());
-				if (Job.AutoSubmitChange != null)
+				attachment.Blocks.Add(new DividerBlock());
+				if (job.AutoSubmitChange != null)
 				{
-					Attachment.Blocks.Add(new SectionBlock($"Shelved files were submitted in CL {Job.AutoSubmitChange}."));
+					attachment.Blocks.Add(new SectionBlock($"Shelved files were submitted in CL {job.AutoSubmitChange}."));
 				}
 				else
 				{
-					Attachment.Color = BlockKitAttachmentColors.Warning;
+					attachment.Color = BlockKitAttachmentColors.Warning;
 
-					string AutoSubmitMessage = String.Empty;
-					if (!String.IsNullOrEmpty(Job.AutoSubmitMessage))
+					string autoSubmitMessage = String.Empty;
+					if (!String.IsNullOrEmpty(job.AutoSubmitMessage))
 					{
-						AutoSubmitMessage = $"\n\n```{Job.AutoSubmitMessage}```";
+						autoSubmitMessage = $"\n\n```{job.AutoSubmitMessage}```";
 					}
 
-					Attachment.Blocks.Add(new SectionBlock($"Files in CL {Job.PreflightChange} were *not submitted*. Please resolve the following issues and submit manually.{Job.AutoSubmitMessage}"));
+					attachment.Blocks.Add(new SectionBlock($"Files in CL {job.PreflightChange} were *not submitted*. Please resolve the following issues and submit manually.{job.AutoSubmitMessage}"));
 				}
 			}
 
-			return SendMessageAsync(Recipient, Attachments: new[] { Attachment });
+			return SendMessageAsync(recipient, attachments: new[] { attachment });
 		}
 
 		#endregion
@@ -441,70 +436,70 @@ namespace Horde.Build.Notifications.Impl
 		#region Job step complete
 
 		/// <inheritdoc/>
-		public async Task NotifyJobStepCompleteAsync(IUser SlackUser, IStream JobStream, IJob Job, IJobStepBatch Batch, IJobStep Step, INode Node, List<ILogEventData> JobStepEventData)
+		public async Task NotifyJobStepCompleteAsync(IUser slackUser, IStream jobStream, IJob job, IJobStepBatch batch, IJobStep step, INode node, List<ILogEventData> jobStepEventData)
 		{
-			Logger.LogInformation("Sending Slack notification for job {JobId}, batch {BatchId}, step {StepId}, outcome {Outcome} to {SlackUser} ({UserId})", Job.Id, Batch.Id, Step.Id, Step.Outcome, SlackUser.Name, SlackUser.Id);
+			_logger.LogInformation("Sending Slack notification for job {JobId}, batch {BatchId}, step {StepId}, outcome {Outcome} to {SlackUser} ({UserId})", job.Id, batch.Id, step.Id, step.Outcome, slackUser.Name, slackUser.Id);
 
-			string? SlackUserId = await GetSlackUserId(SlackUser);
-			if (SlackUserId != null)
+			string? slackUserId = await GetSlackUserId(slackUser);
+			if (slackUserId != null)
 			{
-				await SendJobStepCompleteMessageAsync(SlackUserId, JobStream, Job, Step, Node, JobStepEventData);
+				await SendJobStepCompleteMessageAsync(slackUserId, jobStream, job, step, node, jobStepEventData);
 			}
 		}
 
 		/// <summary>
 		/// Creates a Slack message about a completed step job.
 		/// </summary>
-		/// <param name="Recipient"></param>
-		/// <param name="Stream"></param>
-		/// <param name="Job">The job that contains the step that completed.</param>
-		/// <param name="Step">The job step that completed.</param>
-		/// <param name="Node">The node for the job step.</param>
-		/// <param name="Events">Any events that occurred during the job step.</param>
-		private Task SendJobStepCompleteMessageAsync(string Recipient, IStream Stream, IJob Job, IJobStep Step, INode Node, List<ILogEventData> Events)
+		/// <param name="recipient"></param>
+		/// <param name="stream"></param>
+		/// <param name="job">The job that contains the step that completed.</param>
+		/// <param name="step">The job step that completed.</param>
+		/// <param name="node">The node for the job step.</param>
+		/// <param name="events">Any events that occurred during the job step.</param>
+		private Task SendJobStepCompleteMessageAsync(string recipient, IStream stream, IJob job, IJobStep step, INode node, List<ILogEventData> events)
 		{
-			Uri JobStepLink = new Uri($"{Settings.DashboardUrl}job/{Job.Id}?step={Step.Id}");
-			Uri JobStepLogLink = new Uri($"{Settings.DashboardUrl}log/{Step.LogId}");
+			Uri jobStepLink = new Uri($"{_settings.DashboardUrl}job/{job.Id}?step={step.Id}");
+			Uri jobStepLogLink = new Uri($"{_settings.DashboardUrl}log/{step.LogId}");
 
-			Color OutcomeColor = Step.Outcome == JobStepOutcome.Failure ? BlockKitAttachmentColors.Error : Step.Outcome == JobStepOutcome.Warnings ? BlockKitAttachmentColors.Warning : BlockKitAttachmentColors.Success;
+			Color outcomeColor = step.Outcome == JobStepOutcome.Failure ? BlockKitAttachmentColors.Error : step.Outcome == JobStepOutcome.Warnings ? BlockKitAttachmentColors.Warning : BlockKitAttachmentColors.Success;
 
-			BlockKitAttachment Attachment = new BlockKitAttachment();
-			Attachment.FallbackText = $"{Stream.Name} - {GetJobChangeText(Job)} - {Job.Name} - {Node.Name} - {Step.Outcome}";
-			Attachment.Color = OutcomeColor;
-			Attachment.Blocks.Add(new SectionBlock($"*<{JobStepLink}|{Stream.Name} - {GetJobChangeText(Job)} - {Job.Name} - {Node.Name}>*"));
-			if (Step.Outcome == JobStepOutcome.Success)
+			BlockKitAttachment attachment = new BlockKitAttachment();
+			attachment.FallbackText = $"{stream.Name} - {GetJobChangeText(job)} - {job.Name} - {node.Name} - {step.Outcome}";
+			attachment.Color = outcomeColor;
+			attachment.Blocks.Add(new SectionBlock($"*<{jobStepLink}|{stream.Name} - {GetJobChangeText(job)} - {job.Name} - {node.Name}>*"));
+			if (step.Outcome == JobStepOutcome.Success)
 			{
-				Attachment.Blocks.Add(new SectionBlock($"*Job Step Succeeded*"));
+				attachment.Blocks.Add(new SectionBlock($"*Job Step Succeeded*"));
 			}
 			else
 			{
-				List<ILogEventData> Errors = Events.Where(x => x.Severity == EventSeverity.Error).ToList();
-				List<ILogEventData> Warnings = Events.Where(x => x.Severity == EventSeverity.Warning).ToList();
-				List<string> EventStrings = new List<string>();
-				if (Errors.Any())
+				List<ILogEventData> errors = events.Where(x => x.Severity == EventSeverity.Error).ToList();
+				List<ILogEventData> warnings = events.Where(x => x.Severity == EventSeverity.Warning).ToList();
+				List<string> eventStrings = new List<string>();
+				if (errors.Any())
 				{
-					string ErrorSummary = Errors.Count > MaxJobStepEvents ? $"*Errors (First {MaxJobStepEvents} shown)*" : $"*Errors*";
-					EventStrings.Add(ErrorSummary);
-					foreach (ILogEventData Error in Errors.Take(MaxJobStepEvents))
+					string errorSummary = errors.Count > MaxJobStepEvents ? $"*Errors (First {MaxJobStepEvents} shown)*" : $"*Errors*";
+					eventStrings.Add(errorSummary);
+					foreach (ILogEventData error in errors.Take(MaxJobStepEvents))
 					{
-						EventStrings.Add($"```{Error.Message}```");
+						eventStrings.Add($"```{error.Message}```");
 					}
 				}
-				else if (Warnings.Any())
+				else if (warnings.Any())
 				{
-					string WarningSummary = Warnings.Count > MaxJobStepEvents ? $"*Warnings (First {MaxJobStepEvents} shown)*" : $"*Warnings*";
-					EventStrings.Add(WarningSummary);
-					foreach (ILogEventData Warning in Warnings.Take(MaxJobStepEvents))
+					string warningSummary = warnings.Count > MaxJobStepEvents ? $"*Warnings (First {MaxJobStepEvents} shown)*" : $"*Warnings*";
+					eventStrings.Add(warningSummary);
+					foreach (ILogEventData warning in warnings.Take(MaxJobStepEvents))
 					{
-						EventStrings.Add($"```{Warning.Message}```");
+						eventStrings.Add($"```{warning.Message}```");
 					}
 				}
 
-				Attachment.Blocks.Add(new SectionBlock(string.Join("\n", EventStrings)));
-				Attachment.Blocks.Add(new SectionBlock($"<{JobStepLogLink}|View Job Step Log>"));
+				attachment.Blocks.Add(new SectionBlock(String.Join("\n", eventStrings)));
+				attachment.Blocks.Add(new SectionBlock($"<{jobStepLogLink}|View Job Step Log>"));
 			}
 
-			return SendMessageAsync(Recipient, Attachments: new[] { Attachment });
+			return SendMessageAsync(recipient, attachments: new[] { attachment });
 		}
 
 		#endregion
@@ -512,59 +507,59 @@ namespace Horde.Build.Notifications.Impl
 		#region Label complete
 
 		/// <inheritdoc/>
-		public async Task NotifyLabelCompleteAsync(IUser User, IJob Job, IStream Stream, ILabel Label, int LabelIdx, LabelOutcome Outcome, List<(string, JobStepOutcome, Uri)> StepData)
+		public async Task NotifyLabelCompleteAsync(IUser user, IJob job, IStream stream, ILabel label, int labelIdx, LabelOutcome outcome, List<(string, JobStepOutcome, Uri)> stepData)
 		{
-			Logger.LogInformation("Sending Slack notification for job {JobId} outcome {Outcome} to {Name} ({UserId})", Job.Id, Outcome, User.Name, User.Id);
+			_logger.LogInformation("Sending Slack notification for job {JobId} outcome {Outcome} to {Name} ({UserId})", job.Id, outcome, user.Name, user.Id);
 
-			string? SlackUserId = await GetSlackUserId(User);
-			if (SlackUserId != null)
+			string? slackUserId = await GetSlackUserId(user);
+			if (slackUserId != null)
 			{
-				await SendLabelUpdateMessageAsync(SlackUserId, Stream, Job, Label, LabelIdx, Outcome, StepData);
+				await SendLabelUpdateMessageAsync(slackUserId, stream, job, label, labelIdx, outcome, stepData);
 			}
 		}
 
-		Task SendLabelUpdateMessageAsync(string Recipient, IStream Stream, IJob Job, ILabel Label, int LabelIdx, LabelOutcome Outcome, List<(string, JobStepOutcome, Uri)> JobStepData)
+		Task SendLabelUpdateMessageAsync(string recipient, IStream stream, IJob job, ILabel label, int labelIdx, LabelOutcome outcome, List<(string, JobStepOutcome, Uri)> jobStepData)
 		{
-			Uri LabelLink = new Uri($"{Settings.DashboardUrl}job/{Job.Id}?label={LabelIdx}");
+			Uri labelLink = new Uri($"{_settings.DashboardUrl}job/{job.Id}?label={labelIdx}");
 
-			Color OutcomeColor = Outcome == LabelOutcome.Failure ? BlockKitAttachmentColors.Error : Outcome == LabelOutcome.Warnings ? BlockKitAttachmentColors.Warning : BlockKitAttachmentColors.Success;
+			Color outcomeColor = outcome == LabelOutcome.Failure ? BlockKitAttachmentColors.Error : outcome == LabelOutcome.Warnings ? BlockKitAttachmentColors.Warning : BlockKitAttachmentColors.Success;
 
-			BlockKitAttachment Attachment = new BlockKitAttachment();
-			Attachment.FallbackText = $"{Stream.Name} - {GetJobChangeText(Job)} - {Job.Name} - Label {Label.DashboardName} - {Outcome}";
-			Attachment.Color = OutcomeColor;
-			Attachment.Blocks.Add(new SectionBlock($"*<{LabelLink}|{Stream.Name} - {GetJobChangeText(Job)} - {Job.Name} - Label {Label.DashboardName}>*"));
-			if (Outcome == LabelOutcome.Success)
+			BlockKitAttachment attachment = new BlockKitAttachment();
+			attachment.FallbackText = $"{stream.Name} - {GetJobChangeText(job)} - {job.Name} - Label {label.DashboardName} - {outcome}";
+			attachment.Color = outcomeColor;
+			attachment.Blocks.Add(new SectionBlock($"*<{labelLink}|{stream.Name} - {GetJobChangeText(job)} - {job.Name} - Label {label.DashboardName}>*"));
+			if (outcome == LabelOutcome.Success)
 			{
-				Attachment.Blocks.Add(new SectionBlock($"*Label Succeeded*"));
+				attachment.Blocks.Add(new SectionBlock($"*Label Succeeded*"));
 			}
 			else
 			{
-				List<string> FailedStepStrings = new List<string>();
-				List<string> WarningStepStrings = new List<string>();
-				foreach ((string Name, JobStepOutcome StepOutcome, Uri Link) JobStep in JobStepData)
+				List<string> failedStepStrings = new List<string>();
+				List<string> warningStepStrings = new List<string>();
+				foreach ((string Name, JobStepOutcome StepOutcome, Uri Link) jobStep in jobStepData)
 				{
-					string StepString = $"<{JobStep.Link}|{JobStep.Name}>";
-					if (JobStep.StepOutcome == JobStepOutcome.Failure)
+					string stepString = $"<{jobStep.Link}|{jobStep.Name}>";
+					if (jobStep.StepOutcome == JobStepOutcome.Failure)
 					{
-						FailedStepStrings.Add(StepString);
+						failedStepStrings.Add(stepString);
 					}
-					else if (JobStep.StepOutcome == JobStepOutcome.Warnings)
+					else if (jobStep.StepOutcome == JobStepOutcome.Warnings)
 					{
-						WarningStepStrings.Add(StepString);
+						warningStepStrings.Add(stepString);
 					}
 				}
 
-				if (FailedStepStrings.Any())
+				if (failedStepStrings.Any())
 				{
-					Attachment.Blocks.Add(new SectionBlock($"*Errors*\n{string.Join(", ", FailedStepStrings)}"));
+					attachment.Blocks.Add(new SectionBlock($"*Errors*\n{String.Join(", ", failedStepStrings)}"));
 				}
-				else if (WarningStepStrings.Any())
+				else if (warningStepStrings.Any())
 				{
-					Attachment.Blocks.Add(new SectionBlock($"*Warnings*\n{string.Join(", ", WarningStepStrings)}"));
+					attachment.Blocks.Add(new SectionBlock($"*Warnings*\n{String.Join(", ", warningStepStrings)}"));
 				}
 			}
 
-			return SendMessageAsync(Recipient, Attachments: new[] { Attachment });
+			return SendMessageAsync(recipient, attachments: new[] { attachment });
 		}
 
 		#endregion
@@ -572,342 +567,342 @@ namespace Horde.Build.Notifications.Impl
 		#region Issues
 
 		/// <inheritdoc/>
-		public async Task NotifyIssueUpdatedAsync(IIssue Issue)
+		public async Task NotifyIssueUpdatedAsync(IIssue issue)
 		{
-			IIssueDetails Details = await IssueService.GetIssueDetailsAsync(Issue);
+			IIssueDetails details = await _issueService.GetIssueDetailsAsync(issue);
 
-			HashSet<UserId> UserIds = new HashSet<UserId>();
-			if (Issue.Promoted)
+			HashSet<UserId> userIds = new HashSet<UserId>();
+			if (issue.Promoted)
 			{
-				UserIds.UnionWith(Details.Suspects.Select(x => x.AuthorId));
+				userIds.UnionWith(details.Suspects.Select(x => x.AuthorId));
 			}
-			if (Issue.OwnerId.HasValue)
+			if (issue.OwnerId.HasValue)
 			{
-				UserIds.Add(Issue.OwnerId.Value);
+				userIds.Add(issue.OwnerId.Value);
 			}
 
-			HashSet<string> Channels = new HashSet<string>(StringComparer.OrdinalIgnoreCase);
+			HashSet<string> channels = new HashSet<string>(StringComparer.OrdinalIgnoreCase);
 
-			List<MessageStateDocument> ExistingMessages = await MessageStates.Find(x => x.EventId == GetIssueEventId(Issue)).ToListAsync();
-			foreach (MessageStateDocument ExistingMessage in ExistingMessages)
+			List<MessageStateDocument> existingMessages = await _messageStates.Find(x => x.EventId == GetIssueEventId(issue)).ToListAsync();
+			foreach (MessageStateDocument existingMessage in existingMessages)
 			{
-				if (ExistingMessage.UserId != null)
+				if (existingMessage.UserId != null)
 				{
-					UserIds.Add(ExistingMessage.UserId.Value);
+					userIds.Add(existingMessage.UserId.Value);
 				}
 				else
 				{
-					Channels.Add(ExistingMessage.Recipient);
+					channels.Add(existingMessage.Recipient);
 				}
 			}
 
-			if (Issue.OwnerId == null && (Details.Suspects.Count == 0 || Details.Suspects.All(x => x.DeclinedAt != null) || Issue.CreatedAt < DateTime.UtcNow - TimeSpan.FromHours(1.0)))
+			if (issue.OwnerId == null && (details.Suspects.Count == 0 || details.Suspects.All(x => x.DeclinedAt != null) || issue.CreatedAt < DateTime.UtcNow - TimeSpan.FromHours(1.0)))
 			{
-				foreach (IIssueSpan Span in Details.Spans)
+				foreach (IIssueSpan span in details.Spans)
 				{
-					IStream? Stream = await StreamService.GetCachedStream(Span.StreamId);
-					if (Stream != null)
+					IStream? stream = await _streamService.GetCachedStream(span.StreamId);
+					if (stream != null)
 					{
-						TemplateRef? TemplateRef;
-						if (Stream.Templates.TryGetValue(Span.TemplateRefId, out TemplateRef) && TemplateRef.TriageChannel != null)
+						TemplateRef? templateRef;
+						if (stream.Templates.TryGetValue(span.TemplateRefId, out templateRef) && templateRef.TriageChannel != null)
 						{
-							Channels.Add(TemplateRef.TriageChannel);
+							channels.Add(templateRef.TriageChannel);
 						}
-						else if (Stream.TriageChannel != null)
+						else if (stream.TriageChannel != null)
 						{
-							Channels.Add(Stream.TriageChannel);
+							channels.Add(stream.TriageChannel);
 						}
 					}
 				}
 			}
 
-			if (UserIds.Count > 0)
+			if (userIds.Count > 0)
 			{
-				foreach (UserId UserId in UserIds)
+				foreach (UserId userId in userIds)
 				{
-					IUser? User = await UserCollection.GetUserAsync(UserId);
-					if (User == null)
+					IUser? user = await _userCollection.GetUserAsync(userId);
+					if (user == null)
 					{
-						Logger.LogWarning("Unable to find user {UserId}", UserId);
+						_logger.LogWarning("Unable to find user {UserId}", userId);
 					}
 					else
 					{
-						await NotifyIssueUpdatedAsync(User, Issue, Details);
+						await NotifyIssueUpdatedAsync(user, issue, details);
 					}
 				}
 			}
 
-			if (Channels.Count > 0)
+			if (channels.Count > 0)
 			{
-				foreach (string Channel in Channels)
+				foreach (string channel in channels)
 				{
-					await SendIssueMessageAsync(Channel, Issue, Details, null);
+					await SendIssueMessageAsync(channel, issue, details, null);
 				}
 			}
 		}
 
-		async Task NotifyIssueUpdatedAsync(IUser User, IIssue Issue, IIssueDetails Details)
+		async Task NotifyIssueUpdatedAsync(IUser user, IIssue issue, IIssueDetails details)
 		{
 
-			string? SlackUserId = await GetSlackUserId(User);
-			if (SlackUserId == null)
+			string? slackUserId = await GetSlackUserId(user);
+			if (slackUserId == null)
 			{
 				return;
 			}
 
-			await SendIssueMessageAsync(SlackUserId, Issue, Details, User.Id);
+			await SendIssueMessageAsync(slackUserId, issue, details, user.Id);
 		}
 
-		async Task<string> GetSampleText(IIssueDetails Details)
+		async Task<string> GetSampleText(IIssueDetails details)
 		{
-			StringBuilder SampleText = new StringBuilder();
-			if (Details.Spans.Count > 0)
+			StringBuilder sampleText = new StringBuilder();
+			if (details.Spans.Count > 0)
 			{
-				IIssueSpan Span = Details.Spans[0];
-				if (Span.FirstFailure.LogId != null)
+				IIssueSpan span = details.Spans[0];
+				if (span.FirstFailure.LogId != null)
 				{
-					ILogFile? LogFile = await LogFileService.GetCachedLogFileAsync(Span.FirstFailure.LogId.Value);
-					if (LogFile != null)
+					ILogFile? logFile = await _logFileService.GetCachedLogFileAsync(span.FirstFailure.LogId.Value);
+					if (logFile != null)
 					{
-						List<ILogEvent> Events = await LogFileService.FindEventsForSpansAsync(Details.Spans.Select(x => x.Id).Take(1), null, 0, 1);
-						foreach (ILogEvent Event in Events)
+						List<ILogEvent> logEvents = await _logFileService.FindEventsForSpansAsync(details.Spans.Select(x => x.Id).Take(1), null, 0, 1);
+						foreach (ILogEvent logEvent in logEvents)
 						{
-							ILogEventData Data = await LogFileService.GetEventDataAsync(LogFile, Event.LineIndex, Event.LineCount);
-							if (SampleText.Length > 0)
+							ILogEventData data = await _logFileService.GetEventDataAsync(logFile, logEvent.LineIndex, logEvent.LineCount);
+							if (sampleText.Length > 0)
 							{
-								SampleText.Append("\n\n");
+								sampleText.Append("\n\n");
 							}
-							SampleText.Append(CultureInfo.InvariantCulture, $"```{Data.Message}```");
+							sampleText.Append(CultureInfo.InvariantCulture, $"```{data.Message}```");
 						}
 					}
 				}
 			}
-			return SampleText.ToString();
+			return sampleText.ToString();
 		}
 
-		async Task SendIssueMessageAsync(string Recipient, IIssue Issue, IIssueDetails Details, UserId? UserId)
+		async Task SendIssueMessageAsync(string recipient, IIssue issue, IIssueDetails details, UserId? userId)
 		{
-			using IDisposable Scope = Logger.BeginScope("SendIssueMessageAsync (User: {SlackUser}, Issue: {IssueId})", Recipient, Issue.Id);
+			using IDisposable scope = _logger.BeginScope("SendIssueMessageAsync (User: {SlackUser}, Issue: {IssueId})", recipient, issue.Id);
 
-			BlockKitAttachment Attachment = new BlockKitAttachment();
-			Attachment.Color = Color.Red;
+			BlockKitAttachment attachment = new BlockKitAttachment();
+			attachment.Color = Color.Red;
 
-			Uri IssueUrl = Settings.DashboardUrl;
-			if (Details.Steps.Count > 0)
+			Uri issueUrl = _settings.DashboardUrl;
+			if (details.Steps.Count > 0)
 			{
-				IIssueStep FirstStep = Details.Steps[0];
-				IssueUrl = new Uri(IssueUrl, $"job/{FirstStep.JobId}?step={FirstStep.StepId}&issue={Issue.Id}");
+				IIssueStep firstStep = details.Steps[0];
+				issueUrl = new Uri(issueUrl, $"job/{firstStep.JobId}?step={firstStep.StepId}&issue={issue.Id}");
 			}
-			Attachment.Blocks.Add(new SectionBlock($"*<{IssueUrl}|Issue #{Issue.Id}: {Issue.Summary}>*"));
+			attachment.Blocks.Add(new SectionBlock($"*<{issueUrl}|Issue #{issue.Id}: {issue.Summary}>*"));
 
-			string StreamList = StringUtils.FormatList(Details.Spans.Select(x => $"*{x.StreamName}*").Distinct(StringComparer.OrdinalIgnoreCase).OrderBy(x => x, StringComparer.OrdinalIgnoreCase));
-			Attachment.Blocks.Add(new SectionBlock($"Occurring in {StreamList}"));
+			string streamList = StringUtils.FormatList(details.Spans.Select(x => $"*{x.StreamName}*").Distinct(StringComparer.OrdinalIgnoreCase).OrderBy(x => x, StringComparer.OrdinalIgnoreCase));
+			attachment.Blocks.Add(new SectionBlock($"Occurring in {streamList}"));
 
-			if (Issue.FixChange != null)
+			if (issue.FixChange != null)
 			{
-				IIssueStep? FixFailedStep = Issue.FindFixFailedStep(Details.Spans);
+				IIssueStep? fixFailedStep = issue.FindFixFailedStep(details.Spans);
 
-				string Text;
-				if (Issue.FixChange.Value < 0)
+				string text;
+				if (issue.FixChange.Value < 0)
 				{
-					Text = ":tick: Marked as a systemic issue.";
+					text = ":tick: Marked as a systemic issue.";
 				}
-				else if (FixFailedStep != null)
+				else if (fixFailedStep != null)
 				{
-					Uri FixFailedUrl = new Uri(Settings.DashboardUrl, $"job/{FixFailedStep.JobId}?step={FixFailedStep.StepId}&issue={Issue.Id}");
-					Text = $":cross: Marked fixed in *CL {Issue.FixChange.Value}*, but seen again at *<{FixFailedUrl}|CL {FixFailedStep.Change}>*";
+					Uri fixFailedUrl = new Uri(_settings.DashboardUrl, $"job/{fixFailedStep.JobId}?step={fixFailedStep.StepId}&issue={issue.Id}");
+					text = $":cross: Marked fixed in *CL {issue.FixChange.Value}*, but seen again at *<{fixFailedUrl}|CL {fixFailedStep.Change}>*";
 				}
 				else
 				{
-					Text = $":tick: Marked fixed in *CL {Issue.FixChange.Value}*.";
+					text = $":tick: Marked fixed in *CL {issue.FixChange.Value}*.";
 				}
-				Attachment.Blocks.Add(new SectionBlock(Text));
+				attachment.Blocks.Add(new SectionBlock(text));
 			}
-			else if (UserId != null && Issue.OwnerId == UserId)
+			else if (userId != null && issue.OwnerId == userId)
 			{
-				if (Issue.AcknowledgedAt.HasValue)
+				if (issue.AcknowledgedAt.HasValue)
 				{
-					Attachment.Blocks.Add(new SectionBlock($":+1: Acknowledged at {FormatSlackTime(Issue.AcknowledgedAt.Value)}"));
+					attachment.Blocks.Add(new SectionBlock($":+1: Acknowledged at {FormatSlackTime(issue.AcknowledgedAt.Value)}"));
 				}
 				else
 				{
-					if (Issue.NominatedById != null)
+					if (issue.NominatedById != null)
 					{
-						IUser? NominatedByUser = await UserCollection.GetUserAsync(Issue.NominatedById.Value);
-						if (NominatedByUser != null)
+						IUser? nominatedByUser = await _userCollection.GetUserAsync(issue.NominatedById.Value);
+						if (nominatedByUser != null)
 						{
-							string? NominatedBySlackUserId = await GetSlackUserId(NominatedByUser);
-							string Mention = (NominatedBySlackUserId != null) ? $"<@{NominatedBySlackUserId}>" : NominatedByUser.Login ?? $"User {NominatedByUser.Id}";
-							string Text = $"You were nominated to fix this issue by {Mention} at {FormatSlackTime(Issue.NominatedAt ?? DateTime.UtcNow)}";
-							Attachment.Blocks.Add(new SectionBlock(Text));
+							string? nominatedBySlackUserId = await GetSlackUserId(nominatedByUser);
+							string mention = (nominatedBySlackUserId != null) ? $"<@{nominatedBySlackUserId}>" : nominatedByUser.Login ?? $"User {nominatedByUser.Id}";
+							string text = $"You were nominated to fix this issue by {mention} at {FormatSlackTime(issue.NominatedAt ?? DateTime.UtcNow)}";
+							attachment.Blocks.Add(new SectionBlock(text));
 						}
 					}
 					else
 					{
-						List<int> Changes = Details.Suspects.Where(x => x.AuthorId == UserId).Select(x => x.Change).OrderBy(x => x).ToList();
-						if (Changes.Count > 0)
+						List<int> changes = details.Suspects.Where(x => x.AuthorId == userId).Select(x => x.Change).OrderBy(x => x).ToList();
+						if (changes.Count > 0)
 						{
-							string Text = $"Horde has determined that {StringUtils.FormatList(Changes.Select(x => $"CL {x}"), "or")} is the most likely cause for this issue.";
-							Attachment.Blocks.Add(new SectionBlock(Text));
+							string text = $"Horde has determined that {StringUtils.FormatList(changes.Select(x => $"CL {x}"), "or")} is the most likely cause for this issue.";
+							attachment.Blocks.Add(new SectionBlock(text));
 						}
 					}
 
-					ActionsBlock Actions = new ActionsBlock();
-					Actions.AddButton("Acknowledge", Value: $"issue_{Issue.Id}_ack_{UserId}", Style: ActionButton.ButtonStyle.Primary);
-					Actions.AddButton("Not Me", Value: $"issue_{Issue.Id}_decline_{UserId}", Style: ActionButton.ButtonStyle.Danger);
-					Attachment.Blocks.Add(Actions);
+					ActionsBlock actions = new ActionsBlock();
+					actions.AddButton("Acknowledge", value: $"issue_{issue.Id}_ack_{userId}", style: ActionButton.ButtonStyle.Primary);
+					actions.AddButton("Not Me", value: $"issue_{issue.Id}_decline_{userId}", style: ActionButton.ButtonStyle.Danger);
+					attachment.Blocks.Add(actions);
 				}
 			}
-			else if (Issue.OwnerId != null)
+			else if (issue.OwnerId != null)
 			{
-				string OwnerMention = await FormatMentionAsync(Issue.OwnerId.Value);
-				if (Issue.AcknowledgedAt.HasValue)
+				string ownerMention = await FormatMentionAsync(issue.OwnerId.Value);
+				if (issue.AcknowledgedAt.HasValue)
 				{
-					Attachment.Blocks.Add(new SectionBlock($":+1: Acknowledged by {OwnerMention} at {FormatSlackTime(Issue.AcknowledgedAt.Value)}"));
+					attachment.Blocks.Add(new SectionBlock($":+1: Acknowledged by {ownerMention} at {FormatSlackTime(issue.AcknowledgedAt.Value)}"));
 				}
-				else if (Issue.NominatedById == null)
+				else if (issue.NominatedById == null)
 				{
-					Attachment.Blocks.Add(new SectionBlock($"Assigned to {OwnerMention}"));
+					attachment.Blocks.Add(new SectionBlock($"Assigned to {ownerMention}"));
 				}
-				else if (Issue.NominatedById == UserId)
+				else if (issue.NominatedById == userId)
 				{
-					Attachment.Blocks.Add(new SectionBlock($"You nominated {OwnerMention} to fix this issue."));
+					attachment.Blocks.Add(new SectionBlock($"You nominated {ownerMention} to fix this issue."));
 				}
 				else
 				{
-					Attachment.Blocks.Add(new SectionBlock($"{OwnerMention} was nominated to fix this issue by {await FormatMentionAsync(Issue.NominatedById.Value)}"));
+					attachment.Blocks.Add(new SectionBlock($"{ownerMention} was nominated to fix this issue by {await FormatMentionAsync(issue.NominatedById.Value)}"));
 				}
 			}
-			else if (UserId != null)
+			else if (userId != null)
 			{
-				IIssueSuspect? Suspect = Details.Suspects.FirstOrDefault(x => x.AuthorId == UserId);
-				if (Suspect != null)
+				IIssueSuspect? suspect = details.Suspects.FirstOrDefault(x => x.AuthorId == userId);
+				if (suspect != null)
 				{
-					if (Suspect.DeclinedAt != null)
+					if (suspect.DeclinedAt != null)
 					{
-						Attachment.Blocks.Add(new SectionBlock($":downvote: Declined at {FormatSlackTime(Suspect.DeclinedAt.Value)}"));
+						attachment.Blocks.Add(new SectionBlock($":downvote: Declined at {FormatSlackTime(suspect.DeclinedAt.Value)}"));
 					}
 					else
 					{
-						Attachment.Blocks.Add(new SectionBlock("Please check if any of your recently submitted changes have caused this issue."));
+						attachment.Blocks.Add(new SectionBlock("Please check if any of your recently submitted changes have caused this issue."));
 
-						ActionsBlock Actions = new ActionsBlock();
-						Actions.AddButton("Will Fix", Value: $"issue_{Issue.Id}_accept_{UserId}", Style: ActionButton.ButtonStyle.Primary);
-						Actions.AddButton("Not Me", Value: $"issue_{Issue.Id}_decline_{UserId}", Style: ActionButton.ButtonStyle.Danger);
-						Attachment.Blocks.Add(Actions);
+						ActionsBlock actions = new ActionsBlock();
+						actions.AddButton("Will Fix", value: $"issue_{issue.Id}_accept_{userId}", style: ActionButton.ButtonStyle.Primary);
+						actions.AddButton("Not Me", value: $"issue_{issue.Id}_decline_{userId}", style: ActionButton.ButtonStyle.Danger);
+						attachment.Blocks.Add(actions);
 					}
 				}
 			}
-			else if (Details.Suspects.Count > 0)
+			else if (details.Suspects.Count > 0)
 			{
-				List<string> DeclinedLines = new List<string>();
-				foreach (IIssueSuspect Suspect in Details.Suspects)
+				List<string> declinedLines = new List<string>();
+				foreach (IIssueSuspect suspect in details.Suspects)
 				{
-					if (Suspect.DeclinedAt == null)
+					if (suspect.DeclinedAt == null)
 					{
-						DeclinedLines.Add($":heavy_minus_sign: Ignored by {await FormatNameAsync(Suspect.AuthorId)} (CL {Suspect.Change})");
+						declinedLines.Add($":heavy_minus_sign: Ignored by {await FormatNameAsync(suspect.AuthorId)} (CL {suspect.Change})");
 					}
 					else
 					{
-						DeclinedLines.Add($":downvote: Declined by {await FormatNameAsync(Suspect.AuthorId)} at {FormatSlackTime(Suspect.DeclinedAt.Value)} (CL {Suspect.Change})");
+						declinedLines.Add($":downvote: Declined by {await FormatNameAsync(suspect.AuthorId)} at {FormatSlackTime(suspect.DeclinedAt.Value)} (CL {suspect.Change})");
 					}
 				}
-				Attachment.Blocks.Add(new SectionBlock(String.Join("\n", DeclinedLines)));
+				attachment.Blocks.Add(new SectionBlock(String.Join("\n", declinedLines)));
 			}
 
-			await SendOrUpdateMessageAsync(Recipient, GetIssueEventId(Issue), UserId, Attachments: new[] { Attachment });
+			await SendOrUpdateMessageAsync(recipient, GetIssueEventId(issue), userId, attachments: new[] { attachment });
 		}
 
-		static string GetIssueEventId(IIssue Issue)
+		static string GetIssueEventId(IIssue issue)
 		{
-			return $"issue_{Issue.Id}";
+			return $"issue_{issue.Id}";
 		}
 
-		async Task<string> FormatNameAsync(UserId UserId)
+		async Task<string> FormatNameAsync(UserId userId)
 		{
-			IUser? User = await UserCollection.GetUserAsync(UserId);
-			if (User == null)
+			IUser? user = await _userCollection.GetUserAsync(userId);
+			if (user == null)
 			{
-				return $"User {UserId}";
+				return $"User {userId}";
 			}
-			return User.Name;
+			return user.Name;
 		}
 
-		async Task<string> FormatMentionAsync(UserId UserId)
+		async Task<string> FormatMentionAsync(UserId userId)
 		{
-			IUser? User = await UserCollection.GetUserAsync(UserId);
-			if (User == null)
+			IUser? user = await _userCollection.GetUserAsync(userId);
+			if (user == null)
 			{
-				return $"User {UserId}";
+				return $"User {userId}";
 			}
 
-			string? SlackUserId = await GetSlackUserId(User);
-			if (SlackUserId == null)
+			string? slackUserId = await GetSlackUserId(user);
+			if (slackUserId == null)
 			{
-				return User.Login;
+				return user.Login;
 			}
 
-			return $"<@{SlackUserId}>";
+			return $"<@{slackUserId}>";
 		}
 
-		async Task HandleIssueResponseAsync(EventPayload Payload, ActionInfo Action)
+		async Task HandleIssueResponseAsync(EventPayload payload, ActionInfo action)
 		{
-			string? UserName = Payload.User?.UserName;
-			if(UserName == null)
+			string? userName = payload.User?.UserName;
+			if(userName == null)
 			{
-				Logger.LogWarning("No user for message payload: {Payload}", Payload);
+				_logger.LogWarning("No user for message payload: {Payload}", payload);
 				return;
 			}
-			if (Payload.ResponseUrl == null)
+			if (payload.ResponseUrl == null)
 			{
-				Logger.LogWarning("No response url for payload: {Payload}", Payload);
+				_logger.LogWarning("No response url for payload: {Payload}", payload);
 				return;
 			}
 
-			Match Match = Regex.Match(Action.Value ?? String.Empty, @"^issue_(\d+)_([a-zA-Z]+)_([a-fA-F0-9]{24})$");
-			if (!Match.Success)
+			Match match = Regex.Match(action.Value ?? String.Empty, @"^issue_(\d+)_([a-zA-Z]+)_([a-fA-F0-9]{24})$");
+			if (!match.Success)
 			{
-				Logger.LogWarning("Could not match format of button action: {Action}", Action.Value);
+				_logger.LogWarning("Could not match format of button action: {Action}", action.Value);
 				return;
 			}
 
-			int IssueId = int.Parse(Match.Groups[1].Value, NumberStyles.Integer, CultureInfo.InvariantCulture);
-			string Verb = Match.Groups[2].Value;
-			UserId UserId = Match.Groups[3].Value.ToObjectId<IUser>();
-			Logger.LogInformation("Issue {IssueId}: {Action} from {SlackUser} ({UserId})", IssueId, Verb, UserName, UserId);
+			int issueId = Int32.Parse(match.Groups[1].Value, NumberStyles.Integer, CultureInfo.InvariantCulture);
+			string verb = match.Groups[2].Value;
+			UserId userId = match.Groups[3].Value.ToObjectId<IUser>();
+			_logger.LogInformation("Issue {IssueId}: {Action} from {SlackUser} ({UserId})", issueId, verb, userName, userId);
 
-			if (String.Equals(Verb, "ack", StringComparison.Ordinal))
+			if (String.Equals(verb, "ack", StringComparison.Ordinal))
 			{
-				await IssueService.UpdateIssueAsync(IssueId, Acknowledged: true);
+				await _issueService.UpdateIssueAsync(issueId, acknowledged: true);
 			}
-			else if (String.Equals(Verb, "accept", StringComparison.Ordinal))
+			else if (String.Equals(verb, "accept", StringComparison.Ordinal))
 			{
-				await IssueService.UpdateIssueAsync(IssueId, OwnerId: UserId, Acknowledged: true);
+				await _issueService.UpdateIssueAsync(issueId, ownerId: userId, acknowledged: true);
 			}
-			else if (String.Equals(Verb, "decline", StringComparison.Ordinal))
+			else if (String.Equals(verb, "decline", StringComparison.Ordinal))
 			{
-				await IssueService.UpdateIssueAsync(IssueId, DeclinedById: UserId);
+				await _issueService.UpdateIssueAsync(issueId, declinedById: userId);
 			}
 
-			IIssue? NewIssue = await IssueService.GetIssueAsync(IssueId);
-			if (NewIssue != null)
+			IIssue? newIssue = await _issueService.GetIssueAsync(issueId);
+			if (newIssue != null)
 			{
-				IUser? User = await UserCollection.GetUserAsync(UserId);
-				if (User != null)
+				IUser? user = await _userCollection.GetUserAsync(userId);
+				if (user != null)
 				{
-					string? Recipient = await GetSlackUserId(User);
-					if (Recipient != null)
+					string? recipient = await GetSlackUserId(user);
+					if (recipient != null)
 					{
-						IIssueDetails Details = await IssueService.GetIssueDetailsAsync(NewIssue);
-						await SendIssueMessageAsync(Recipient, NewIssue, Details, UserId);
+						IIssueDetails details = await _issueService.GetIssueDetailsAsync(newIssue);
+						await SendIssueMessageAsync(recipient, newIssue, details, userId);
 					}
 				}
 			}
 		}
 
-		static string FormatSlackTime(DateTimeOffset Time)
+		static string FormatSlackTime(DateTimeOffset time)
 		{
-			return $"<!date^{Time.ToUnixTimeSeconds()}^{{time}}|{Time}>";
+			return $"<!date^{time.ToUnixTimeSeconds()}^{{time}}|{time}>";
 		}
 
 		#endregion
@@ -915,62 +910,62 @@ namespace Horde.Build.Notifications.Impl
 		#region Stream updates
 
 		/// <inheritdoc/>
-		public async Task NotifyConfigUpdateFailureAsync(string ErrorMessage, string FileName, int? Change = null, IUser? Author = null, string? Description = null)
+		public async Task NotifyConfigUpdateFailureAsync(string errorMessage, string fileName, int? change = null, IUser? author = null, string? description = null)
 		{
-			Logger.LogInformation("Sending config update failure notification for {FileName} (change: {Change}, author: {UserId})", FileName, Change ?? -1, Author?.Id ?? UserId.Empty);
+			_logger.LogInformation("Sending config update failure notification for {FileName} (change: {Change}, author: {UserId})", fileName, change ?? -1, author?.Id ?? UserId.Empty);
 
-			string? SlackUserId = null;
-			if (Author != null)
+			string? slackUserId = null;
+			if (author != null)
 			{
-				SlackUserId = await GetSlackUserId(Author);
-				if (SlackUserId == null)
+				slackUserId = await GetSlackUserId(author);
+				if (slackUserId == null)
 				{
-					Logger.LogWarning("Unable to identify Slack user id for {UserId}", Author.Id);
+					_logger.LogWarning("Unable to identify Slack user id for {UserId}", author.Id);
 				}
 				else
 				{
-					Logger.LogInformation("Mappsed user {UserId} to Slack user {SlackUserId}", Author.Id, SlackUserId);
+					_logger.LogInformation("Mappsed user {UserId} to Slack user {SlackUserId}", author.Id, slackUserId);
 				}
 			}
 
-			if (SlackUserId != null)
+			if (slackUserId != null)
 			{
-				await SendConfigUpdateFailureMessageAsync(SlackUserId, ErrorMessage, FileName, Change, SlackUserId, Description);
+				await SendConfigUpdateFailureMessageAsync(slackUserId, errorMessage, fileName, change, slackUserId, description);
 			}
-			if (Settings.UpdateStreamsNotificationChannel != null)
+			if (_settings.UpdateStreamsNotificationChannel != null)
 			{
-				await SendConfigUpdateFailureMessageAsync($"#{Settings.UpdateStreamsNotificationChannel}", ErrorMessage, FileName, Change, SlackUserId, Description);
+				await SendConfigUpdateFailureMessageAsync($"#{_settings.UpdateStreamsNotificationChannel}", errorMessage, fileName, change, slackUserId, description);
 			}
 		}
 
-		private Task SendConfigUpdateFailureMessageAsync(string Recipient, string ErrorMessage, string FileName, int? Change = null, string? Author = null, string? Description = null)
+		private Task SendConfigUpdateFailureMessageAsync(string recipient, string errorMessage, string fileName, int? change = null, string? author = null, string? description = null)
 		{
-			Color OutcomeColor = BlockKitAttachmentColors.Error;
-			BlockKitAttachment Attachment = new BlockKitAttachment();
-			Attachment.Color = OutcomeColor;
-			Attachment.FallbackText = $"Update Failure: {FileName}";
+			Color outcomeColor = BlockKitAttachmentColors.Error;
+			BlockKitAttachment attachment = new BlockKitAttachment();
+			attachment.Color = outcomeColor;
+			attachment.FallbackText = $"Update Failure: {fileName}";
 
-			Attachment.Blocks.Add(new HeaderBlock($"Config Update Failure :rip:", false, true));
+			attachment.Blocks.Add(new HeaderBlock($"Config Update Failure :rip:", false, true));
 
-			Attachment.Blocks.Add(new SectionBlock($"Horde was unable to update {FileName}"));
-			Attachment.Blocks.Add(new SectionBlock($"```{ErrorMessage}```"));
-			if (Change != null)
+			attachment.Blocks.Add(new SectionBlock($"Horde was unable to update {fileName}"));
+			attachment.Blocks.Add(new SectionBlock($"```{errorMessage}```"));
+			if (change != null)
 			{
-				if (Author != null)
+				if (author != null)
 				{
-					Attachment.Blocks.Add(new SectionBlock($"Possibly due to CL: {Change.Value} by <@{Author}>"));
+					attachment.Blocks.Add(new SectionBlock($"Possibly due to CL: {change.Value} by <@{author}>"));
 				}
 				else
 				{
-					Attachment.Blocks.Add(new SectionBlock($"Possibly due to CL: {Change.Value} - (Could not determine author from P4 user)"));
+					attachment.Blocks.Add(new SectionBlock($"Possibly due to CL: {change.Value} - (Could not determine author from P4 user)"));
 				}
-				if (Description != null)
+				if (description != null)
 				{
-					Attachment.Blocks.Add(new SectionBlock($"```{Description}```"));
+					attachment.Blocks.Add(new SectionBlock($"```{description}```"));
 				}
 			}
 
-			return SendMessageAsync(Recipient, Attachments: new[] { Attachment });
+			return SendMessageAsync(recipient, attachments: new[] { attachment });
 		}
 
 		#endregion
@@ -978,34 +973,34 @@ namespace Horde.Build.Notifications.Impl
 		#region Stream update (file)
 
 		/// <inheritdoc/>
-		public async Task NotifyStreamUpdateFailedAsync(FileSummary File)
+		public async Task NotifyStreamUpdateFailedAsync(FileSummary file)
 		{
-			Logger.LogDebug("Sending stream update failure notification for {File}", File.DepotPath);
-			if (Settings.UpdateStreamsNotificationChannel != null)
+			_logger.LogDebug("Sending stream update failure notification for {File}", file.DepotPath);
+			if (_settings.UpdateStreamsNotificationChannel != null)
 			{
-				await SendStreamUpdateFailureMessage($"#{Settings.UpdateStreamsNotificationChannel}", File);
+				await SendStreamUpdateFailureMessage($"#{_settings.UpdateStreamsNotificationChannel}", file);
 			}
 		}
 
 		/// <summary>
 		/// Creates a stream update failure message in relation to a file
 		/// </summary>
-		/// <param name="Recipient"></param>
-		/// <param name="File">The file</param>
+		/// <param name="recipient"></param>
+		/// <param name="file">The file</param>
 		/// <returns></returns>
-		Task SendStreamUpdateFailureMessage(string Recipient, FileSummary File)
+		Task SendStreamUpdateFailureMessage(string recipient, FileSummary file)
 		{
-			Color OutcomeColor = BlockKitAttachmentColors.Error;
-			BlockKitAttachment Attachment = new BlockKitAttachment();
-			Attachment.Color = OutcomeColor;
-			Attachment.FallbackText = $"{File.DepotPath} - Update Failure";
+			Color outcomeColor = BlockKitAttachmentColors.Error;
+			BlockKitAttachment attachment = new BlockKitAttachment();
+			attachment.Color = outcomeColor;
+			attachment.FallbackText = $"{file.DepotPath} - Update Failure";
 
-			Attachment.Blocks.Add(new HeaderBlock($"Stream Update Failure :rip:", false, true));
+			attachment.Blocks.Add(new HeaderBlock($"Stream Update Failure :rip:", false, true));
 
-			Attachment.Blocks.Add(new SectionBlock($"<!here> Horde was unable to update {File.DepotPath}"));
-			Attachment.Blocks.Add(new SectionBlock($"```{File.Error}```"));
+			attachment.Blocks.Add(new SectionBlock($"<!here> Horde was unable to update {file.DepotPath}"));
+			attachment.Blocks.Add(new SectionBlock($"```{file.Error}```"));
 
-			return SendMessageAsync(Recipient, Attachments: new[] { Attachment });
+			return SendMessageAsync(recipient, attachments: new[] { attachment });
 		}
 
 		#endregion
@@ -1013,178 +1008,177 @@ namespace Horde.Build.Notifications.Impl
 		#region Device notifications
 
 		/// <inheritdoc/>
-		public async Task NotifyDeviceServiceAsync(string Message, IDevice? Device = null, IDevicePool? Pool = null, IStream? Stream = null, IJob? Job = null, IJobStep? Step = null, INode? Node = null, IUser? User = null)
+		public async Task NotifyDeviceServiceAsync(string message, IDevice? device = null, IDevicePool? pool = null, IStream? stream = null, IJob? job = null, IJobStep? step = null, INode? node = null, IUser? user = null)
 		{
-			string Recipient = $"#{Settings.DeviceServiceNotificationChannel}";
+			string recipient = $"#{_settings.DeviceServiceNotificationChannel}";
 
-			if (User != null)
+			if (user != null)
 			{
-				string? SlackRecipient = await GetSlackUserId(User);
+				string? slackRecipient = await GetSlackUserId(user);
 
-				if (SlackRecipient == null)
+				if (slackRecipient == null)
 				{
-					Logger.LogError("NotifyDeviceServiceAsync - Unable to send user slack notification, user {UserId} slack user id not found", User.Id);
+					_logger.LogError("NotifyDeviceServiceAsync - Unable to send user slack notification, user {UserId} slack user id not found", user.Id);
 					return;
 				}
 
-				Recipient = SlackRecipient;
+				recipient = slackRecipient;
 			}
 
-			Logger.LogDebug("Sending device service notification to {Recipient}", Recipient);
+			_logger.LogDebug("Sending device service notification to {Recipient}", recipient);
 
-			if (Settings.DeviceServiceNotificationChannel != null)
+			if (_settings.DeviceServiceNotificationChannel != null)
 			{
-				await SendDeviceServiceMessage(Recipient, Message, Device, Pool, Stream, Job, Step, Node, User);
+				await SendDeviceServiceMessage(recipient, message, device, pool, stream, job, step, node, user);
 			}
 		}
 
 		/// <summary>
 		/// Creates a Slack message for a device service notification
 		/// </summary>
-		/// <param name="Recipient"></param>
-		/// <param name="Message"></param>
-		/// <param name="Device"></param>
-		/// <param name="Pool"></param>
-		/// <param name="Stream"></param>
-		/// <param name="Job">The job that contains the step that completed.</param>
-		/// <param name="Step">The job step that completed.</param>
-		/// <param name="Node">The node for the job step.</param>
-		/// <param name="User">The user to notify.</param>
-		private Task SendDeviceServiceMessage(string Recipient, string Message, IDevice? Device = null, IDevicePool? Pool = null, IStream? Stream = null, IJob? Job = null, IJobStep? Step = null, INode? Node = null, IUser? User = null)
+		/// <param name="recipient"></param>
+		/// <param name="message"></param>
+		/// <param name="device"></param>
+		/// <param name="pool"></param>
+		/// <param name="stream"></param>
+		/// <param name="job">The job that contains the step that completed.</param>
+		/// <param name="step">The job step that completed.</param>
+		/// <param name="node">The node for the job step.</param>
+		/// <param name="user">The user to notify.</param>
+		private Task SendDeviceServiceMessage(string recipient, string message, IDevice? device = null, IDevicePool? pool = null, IStream? stream = null, IJob? job = null, IJobStep? step = null, INode? node = null, IUser? user = null)
 		{
 
-			if (User != null)
+			if (user != null)
 			{
-				return SendMessageAsync(Recipient, Message);
+				return SendMessageAsync(recipient, message);
 			}
 
 			// truncate message to avoid slack error on message length
-			if (Message.Length > 150)
+			if (message.Length > 150)
 			{
-				Message = Message.Substring(0, 146) + "...";
+				message = message.Substring(0, 146) + "...";
 			}
 
-			BlockKitAttachment Attachment = new BlockKitAttachment();
+			BlockKitAttachment attachment = new BlockKitAttachment();
 							
-			Attachment.FallbackText = $"{Message}";
+			attachment.FallbackText = $"{message}";
 
-			if (Device != null && Pool != null)
+			if (device != null && pool != null)
 			{
-				Attachment.FallbackText += $" - Device: {Device.Name} Pool: {Pool.Name}";
+				attachment.FallbackText += $" - Device: {device.Name} Pool: {pool.Name}";
 			}
 				
-			Attachment.Blocks.Add(new HeaderBlock($"{Message}", false, false));
+			attachment.Blocks.Add(new HeaderBlock($"{message}", false, false));
 
-			if (Stream != null && Job != null && Step != null && Node != null)
+			if (stream != null && job != null && step != null && node != null)
 			{
-				Uri JobStepLink = new Uri($"{Settings.DashboardUrl}job/{Job.Id}?step={Step.Id}");
-				Uri JobStepLogLink = new Uri($"{Settings.DashboardUrl}log/{Step.LogId}");
+				Uri jobStepLink = new Uri($"{_settings.DashboardUrl}job/{job.Id}?step={step.Id}");
+				Uri jobStepLogLink = new Uri($"{_settings.DashboardUrl}log/{step.LogId}");
 
-				Attachment.FallbackText += $" - {Stream.Name} - {GetJobChangeText(Job)} - {Job.Name} - {Node.Name}";
-				Attachment.Blocks.Add(new SectionBlock($"*<{JobStepLink}|{Stream.Name} - {GetJobChangeText(Job)} - {Job.Name} - {Node.Name}>*"));
-				Attachment.Blocks.Add(new SectionBlock($"<{JobStepLogLink}|View Job Step Log>"));
+				attachment.FallbackText += $" - {stream.Name} - {GetJobChangeText(job)} - {job.Name} - {node.Name}";
+				attachment.Blocks.Add(new SectionBlock($"*<{jobStepLink}|{stream.Name} - {GetJobChangeText(job)} - {job.Name} - {node.Name}>*"));
+				attachment.Blocks.Add(new SectionBlock($"<{jobStepLogLink}|View Job Step Log>"));
 			}
 			else
 			{
-				Attachment.FallbackText += " - No job information (Gauntlet might need to be updated in stream)";
-				Attachment.Blocks.Add(new SectionBlock("*No job information (Gauntlet might need to be updated in stream)*"));
+				attachment.FallbackText += " - No job information (Gauntlet might need to be updated in stream)";
+				attachment.Blocks.Add(new SectionBlock("*No job information (Gauntlet might need to be updated in stream)*"));
 			}
 
-			return SendMessageAsync(Recipient, Attachments: new[] { Attachment });
+			return SendMessageAsync(recipient, attachments: new[] { attachment });
 
 		}
 		
 		#endregion
 
-
 		const int MaxJobStepEvents = 5;
 
-		static string GetJobChangeText(IJob Job)
+		static string GetJobChangeText(IJob job)
 		{
-			if (Job.PreflightChange == 0)
+			if (job.PreflightChange == 0)
 			{
-				return $"CL {Job.Change}";
+				return $"CL {job.Change}";
 			}
 			else
 			{
-				return $"Preflight CL {Job.PreflightChange} against CL {Job.Change}";
+				return $"Preflight CL {job.PreflightChange} against CL {job.Change}";
 			}
 		}
 
-		static bool ShouldUpdateUser(SlackUser? Document)
+		static bool ShouldUpdateUser(SlackUser? document)
 		{
-			if(Document == null || Document.Version < SlackUser.CurrentVersion)
+			if(document == null || document.Version < SlackUser.CurrentVersion)
 			{
 				return true;
 			}
 
-			TimeSpan ExpiryTime;
-			if (Document.SlackUserId == null)
+			TimeSpan expiryTime;
+			if (document.SlackUserId == null)
 			{
-				ExpiryTime = TimeSpan.FromMinutes(10.0);
+				expiryTime = TimeSpan.FromMinutes(10.0);
 			}
 			else
 			{
-				ExpiryTime = TimeSpan.FromDays(1.0);
+				expiryTime = TimeSpan.FromDays(1.0);
 			}
-			return Document.Time + ExpiryTime < DateTime.UtcNow;
+			return document.Time + expiryTime < DateTime.UtcNow;
 		}
 
-		private async Task<string?> GetSlackUserId(IUser User)
+		private async Task<string?> GetSlackUserId(IUser user)
 		{
-			return (await GetSlackUser(User))?.SlackUserId;
+			return (await GetSlackUser(user))?.SlackUserId;
 		}
 
-		private async Task<SlackUser?> GetSlackUser(IUser User)
+		private async Task<SlackUser?> GetSlackUser(IUser user)
 		{
-			string? Email = User.Email;
-			if (Email == null)
+			string? email = user.Email;
+			if (email == null)
 			{
-				Logger.LogWarning("Unable to find Slack user id for {UserId} ({Name}): No email address in user profile", User.Id, User.Name);
+				_logger.LogWarning("Unable to find Slack user id for {UserId} ({Name}): No email address in user profile", user.Id, user.Name);
 				return null;
 			}
 
-			SlackUser? UserDocument;
-			if (!UserCache.TryGetValue(Email, out UserDocument))
+			SlackUser? userDocument;
+			if (!_userCache.TryGetValue(email, out userDocument))
 			{
-				UserDocument = await SlackUsers.Find(x => x.Id == User.Id).FirstOrDefaultAsync();
-				if (UserDocument == null || ShouldUpdateUser(UserDocument))
+				userDocument = await _slackUsers.Find(x => x.Id == user.Id).FirstOrDefaultAsync();
+				if (userDocument == null || ShouldUpdateUser(userDocument))
 				{
-					UserInfo? UserInfo = await GetSlackUserInfoByEmail(Email);
-					if (UserDocument == null || UserInfo != null)
+					UserInfo? userInfo = await GetSlackUserInfoByEmail(email);
+					if (userDocument == null || userInfo != null)
 					{
-						UserDocument = new SlackUser(User.Id, UserInfo);
-						await SlackUsers.ReplaceOneAsync(x => x.Id == User.Id, UserDocument, new ReplaceOptions { IsUpsert = true });
+						userDocument = new SlackUser(user.Id, userInfo);
+						await _slackUsers.ReplaceOneAsync(x => x.Id == user.Id, userDocument, new ReplaceOptions { IsUpsert = true });
 					}
 				}
-				using (ICacheEntry Entry = UserCache.CreateEntry(Email))
+				using (ICacheEntry entry = _userCache.CreateEntry(email))
 				{
-					Entry.SlidingExpiration = TimeSpan.FromMinutes(10.0);
-					Entry.Value = UserDocument;
+					entry.SlidingExpiration = TimeSpan.FromMinutes(10.0);
+					entry.Value = userDocument;
 				}
 			}
 
-			return UserDocument;
+			return userDocument;
 		}
 
-		private async Task<UserInfo?> GetSlackUserInfoByEmail(string Email)
+		private async Task<UserInfo?> GetSlackUserInfoByEmail(string email)
 		{
-			using HttpClient Client = new HttpClient();
+			using HttpClient client = new HttpClient();
 
-			using HttpRequestMessage GetUserIdRequest = new HttpRequestMessage(HttpMethod.Post, $"https://slack.com/api/users.lookupByEmail?email={Email}");
-			GetUserIdRequest.Headers.Add("Authorization", $"Bearer {Settings.SlackToken ?? ""}");
+			using HttpRequestMessage getUserIdRequest = new HttpRequestMessage(HttpMethod.Post, $"https://slack.com/api/users.lookupByEmail?email={email}");
+			getUserIdRequest.Headers.Add("Authorization", $"Bearer {_settings.SlackToken ?? ""}");
 
-			HttpResponseMessage ResponseMessage = await Client.SendAsync(GetUserIdRequest);
-			byte[] ResponseData = await ResponseMessage.Content.ReadAsByteArrayAsync();
+			HttpResponseMessage responseMessage = await client.SendAsync(getUserIdRequest);
+			byte[] responseData = await responseMessage.Content.ReadAsByteArrayAsync();
 
-			UserResponse UserResponse = JsonSerializer.Deserialize<UserResponse>(ResponseData)!;
-			if(!UserResponse.Ok || UserResponse.User == null)
+			UserResponse userResponse = JsonSerializer.Deserialize<UserResponse>(responseData)!;
+			if(!userResponse.Ok || userResponse.User == null)
 			{
-				Logger.LogWarning("Unable to find Slack user id for {Email}: {Response}", Email, Encoding.UTF8.GetString(ResponseData));
+				_logger.LogWarning("Unable to find Slack user id for {Email}: {Response}", email, Encoding.UTF8.GetString(responseData));
 				return null;
 			}
 
-			return UserResponse.User;
+			return userResponse.User;
 		}
 
 		class SlackResponse
@@ -1205,125 +1199,125 @@ namespace Horde.Build.Notifications.Impl
 			public string? Ts { get; set; }
 		}
 
-		private async Task SendMessageAsync(string Recipient, string? Text = null, BlockBase[]? Blocks = null, BlockKitAttachment[]? Attachments = null)
+		private async Task SendMessageAsync(string recipient, string? text = null, BlockBase[]? blocks = null, BlockKitAttachment[]? attachments = null)
 		{
-			if (AllowUsers != null && !AllowUsers.Contains(Recipient))
+			if (_allowUsers != null && !_allowUsers.Contains(recipient))
 			{
-				Logger.LogDebug("Suppressing message to {Recipient}: {Text}", Recipient, Text);
+				_logger.LogDebug("Suppressing message to {Recipient}: {Text}", recipient, text);
 				return;
 			}
 
-			BlockKitMessage Message = new BlockKitMessage();
-			Message.Channel = Recipient;
-			Message.Text = Text;
-			if (Blocks != null)
+			BlockKitMessage message = new BlockKitMessage();
+			message.Channel = recipient;
+			message.Text = text;
+			if (blocks != null)
 			{
-				Message.Blocks.AddRange(Blocks);
+				message.Blocks.AddRange(blocks);
 			}
-			if (Attachments != null)
+			if (attachments != null)
 			{
-				Message.Attachments.AddRange(Attachments);
+				message.Attachments.AddRange(attachments);
 			}
 
-			await SendRequestAsync<PostMessageResponse>(PostMessageUrl, Message);
+			await SendRequestAsync<PostMessageResponse>(PostMessageUrl, message);
 		}
 
-		private async Task SendOrUpdateMessageAsync(string Recipient, string EventId, UserId? UserId, string? Text = null, BlockBase[]? Blocks = null, BlockKitAttachment[]? Attachments = null)
+		private async Task SendOrUpdateMessageAsync(string recipient, string eventId, UserId? userId, string? text = null, BlockBase[]? blocks = null, BlockKitAttachment[]? attachments = null)
 		{
-			if (AllowUsers != null && !AllowUsers.Contains(Recipient))
+			if (_allowUsers != null && !_allowUsers.Contains(recipient))
 			{
-				Logger.LogDebug("Suppressing message to {Recipient}: {Text}", Recipient, Text);
+				_logger.LogDebug("Suppressing message to {Recipient}: {Text}", recipient, text);
 				return;
 			}
 
-			BlockKitMessage Message = new BlockKitMessage();
-			Message.Text = Text;
-			if (Blocks != null)
+			BlockKitMessage message = new BlockKitMessage();
+			message.Text = text;
+			if (blocks != null)
 			{
-				Message.Blocks.AddRange(Blocks);
+				message.Blocks.AddRange(blocks);
 			}
-			if (Attachments != null)
+			if (attachments != null)
 			{
-				Message.Attachments.AddRange(Attachments);
+				message.Attachments.AddRange(attachments);
 			}
 
-			string RequestDigest = ContentHash.MD5(JsonSerializer.Serialize(Message, new JsonSerializerOptions { DefaultIgnoreCondition = JsonIgnoreCondition.WhenWritingNull })).ToString();
+			string requestDigest = ContentHash.MD5(JsonSerializer.Serialize(message, new JsonSerializerOptions { DefaultIgnoreCondition = JsonIgnoreCondition.WhenWritingNull })).ToString();
 
-			(MessageStateDocument State, bool IsNew) = await AddOrUpdateMessageStateAsync(Recipient, EventId, UserId, RequestDigest);
-			if (IsNew)
+			(MessageStateDocument state, bool isNew) = await AddOrUpdateMessageStateAsync(recipient, eventId, userId, requestDigest);
+			if (isNew)
 			{
-				Logger.LogInformation("Sending new slack message to {SlackUser} (msg: {MessageId})", Recipient, State.Id);
-				Message.Channel = Recipient;
-				Message.Ts = null;
+				_logger.LogInformation("Sending new slack message to {SlackUser} (msg: {MessageId})", recipient, state.Id);
+				message.Channel = recipient;
+				message.Ts = null;
 
-				PostMessageResponse Response = await SendRequestAsync<PostMessageResponse>(PostMessageUrl, Message);
-				if (String.IsNullOrEmpty(Response.Ts) || String.IsNullOrEmpty(Response.Channel))
+				PostMessageResponse response = await SendRequestAsync<PostMessageResponse>(PostMessageUrl, message);
+				if (String.IsNullOrEmpty(response.Ts) || String.IsNullOrEmpty(response.Channel))
 				{
-					Logger.LogWarning("Missing 'ts' or 'channel' field on slack response");
+					_logger.LogWarning("Missing 'ts' or 'channel' field on slack response");
 				}
-				await SetMessageTimestampAsync(State.Id, Response.Channel ?? String.Empty, Response.Ts ?? String.Empty);
+				await SetMessageTimestampAsync(state.Id, response.Channel ?? String.Empty, response.Ts ?? String.Empty);
 			}
-			else if (!String.IsNullOrEmpty(State.Ts))
+			else if (!String.IsNullOrEmpty(state.Ts))
 			{
-				Logger.LogInformation("Updating existing slack message {MessageId} for user {SlackUser} ({Channel}, {MessageTs})", State.Id, Recipient, State.Channel, State.Ts);
-				Message.Channel = State.Channel;
-				Message.Ts = State.Ts;
-				await SendRequestAsync<PostMessageResponse>(UpdateMessageUrl, Message);
+				_logger.LogInformation("Updating existing slack message {MessageId} for user {SlackUser} ({Channel}, {MessageTs})", state.Id, recipient, state.Channel, state.Ts);
+				message.Channel = state.Channel;
+				message.Ts = state.Ts;
+				await SendRequestAsync<PostMessageResponse>(UpdateMessageUrl, message);
 			}
 		}
 
-		private async Task<TResponse> SendRequestAsync<TResponse>(string RequestUrl, object Request) where TResponse : SlackResponse
+		private async Task<TResponse> SendRequestAsync<TResponse>(string requestUrl, object request) where TResponse : SlackResponse
 		{
-			using (HttpClient Client = new HttpClient())
+			using (HttpClient client = new HttpClient())
 			{
-				using (HttpRequestMessage SendMessageRequest = new HttpRequestMessage(HttpMethod.Post, RequestUrl))
+				using (HttpRequestMessage sendMessageRequest = new HttpRequestMessage(HttpMethod.Post, requestUrl))
 				{
-					string RequestJson = JsonSerializer.Serialize(Request, new JsonSerializerOptions { DefaultIgnoreCondition = JsonIgnoreCondition.WhenWritingNull });
-					using (StringContent MessageContent = new StringContent(RequestJson, Encoding.UTF8, "application/json"))
+					string requestJson = JsonSerializer.Serialize(request, new JsonSerializerOptions { DefaultIgnoreCondition = JsonIgnoreCondition.WhenWritingNull });
+					using (StringContent messageContent = new StringContent(requestJson, Encoding.UTF8, "application/json"))
 					{
-						SendMessageRequest.Content = MessageContent;
-						SendMessageRequest.Headers.Add("Authorization", $"Bearer {Settings.SlackToken ?? ""}");
+						sendMessageRequest.Content = messageContent;
+						sendMessageRequest.Headers.Add("Authorization", $"Bearer {_settings.SlackToken ?? ""}");
 
-						HttpResponseMessage Response = await Client.SendAsync(SendMessageRequest);
-						byte[] ResponseBytes = await Response.Content.ReadAsByteArrayAsync();
+						HttpResponseMessage response = await client.SendAsync(sendMessageRequest);
+						byte[] responseBytes = await response.Content.ReadAsByteArrayAsync();
 
-						TResponse ResponseObject = JsonSerializer.Deserialize<TResponse>(ResponseBytes)!;
-						if(!ResponseObject.Ok)
+						TResponse responseObject = JsonSerializer.Deserialize<TResponse>(responseBytes)!;
+						if(!responseObject.Ok)
 						{
-							Logger.LogError("Failed to send Slack message ({Error}). Request: {Request}. Response: {Response}", ResponseObject.Error, RequestJson, Encoding.UTF8.GetString(ResponseBytes));
+							_logger.LogError("Failed to send Slack message ({Error}). Request: {Request}. Response: {Response}", responseObject.Error, requestJson, Encoding.UTF8.GetString(responseBytes));
 						}
 
-						return ResponseObject;
+						return responseObject;
 					}
 				}
 			}
 		}
 
 		/// <inheritdoc/>
-		protected async override Task ExecuteAsync(CancellationToken StoppingToken)
+		protected override async Task ExecuteAsync(CancellationToken stoppingToken)
 		{
-			if (!String.IsNullOrEmpty(Settings.SlackSocketToken))
+			if (!String.IsNullOrEmpty(_settings.SlackSocketToken))
 			{
-				while (!StoppingToken.IsCancellationRequested)
+				while (!stoppingToken.IsCancellationRequested)
 				{
 					try
 					{
-						Uri? WebSocketUrl = await GetWebSocketUrlAsync(StoppingToken);
-						if (WebSocketUrl == null)
+						Uri? webSocketUrl = await GetWebSocketUrlAsync(stoppingToken);
+						if (webSocketUrl == null)
 						{
-							await Task.Delay(TimeSpan.FromSeconds(5.0), StoppingToken);
+							await Task.Delay(TimeSpan.FromSeconds(5.0), stoppingToken);
 							continue;
 						}
-						await HandleSocketAsync(WebSocketUrl, StoppingToken);
+						await HandleSocketAsync(webSocketUrl, stoppingToken);
 					}
 					catch (OperationCanceledException)
 					{
 						break;
 					}
-					catch (Exception Ex)
+					catch (Exception ex)
 					{
-						Logger.LogError(Ex, "Exception while updating Slack socket");
-						await Task.Delay(TimeSpan.FromSeconds(5.0), StoppingToken);
+						_logger.LogError(ex, "Exception while updating Slack socket");
+						await Task.Delay(TimeSpan.FromSeconds(5.0), stoppingToken);
 					}
 				}
 			}
@@ -1332,90 +1326,90 @@ namespace Horde.Build.Notifications.Impl
 		/// <summary>
 		/// Get the url for opening a websocket to Slack
 		/// </summary>
-		/// <param name="StoppingToken"></param>
+		/// <param name="stoppingToken"></param>
 		/// <returns></returns>
-		private async Task<Uri?> GetWebSocketUrlAsync(CancellationToken StoppingToken)
+		private async Task<Uri?> GetWebSocketUrlAsync(CancellationToken stoppingToken)
 		{
-			using HttpClient Client = new HttpClient();
-			Client.DefaultRequestHeaders.Add("Authorization", $"Bearer {Settings.SlackSocketToken}");
+			using HttpClient client = new HttpClient();
+			client.DefaultRequestHeaders.Add("Authorization", $"Bearer {_settings.SlackSocketToken}");
 
-			using FormUrlEncodedContent Content = new FormUrlEncodedContent(Array.Empty<KeyValuePair<string?, string?>>());
-			HttpResponseMessage Response = await Client.PostAsync(new Uri("https://slack.com/api/apps.connections.open"), Content, StoppingToken);
+			using FormUrlEncodedContent content = new FormUrlEncodedContent(Array.Empty<KeyValuePair<string?, string?>>());
+			HttpResponseMessage response = await client.PostAsync(new Uri("https://slack.com/api/apps.connections.open"), content, stoppingToken);
 
-			byte[] ResponseData = await Response.Content.ReadAsByteArrayAsync(StoppingToken);
+			byte[] responseData = await response.Content.ReadAsByteArrayAsync(stoppingToken);
 
-			SocketResponse SocketResponse = JsonSerializer.Deserialize<SocketResponse>(ResponseData)!;
-			if (!SocketResponse.Ok)
+			SocketResponse socketResponse = JsonSerializer.Deserialize<SocketResponse>(responseData)!;
+			if (!socketResponse.Ok)
 			{
-				Logger.LogWarning("Unable to get websocket url: {Response}", Encoding.UTF8.GetString(ResponseData));
+				_logger.LogWarning("Unable to get websocket url: {Response}", Encoding.UTF8.GetString(responseData));
 				return null;
 			}
 
-			return SocketResponse.Url;
+			return socketResponse.Url;
 		}
 
 		/// <summary>
 		/// Handle the lifetime of a websocket connection
 		/// </summary>
-		/// <param name="SocketUrl"></param>
-		/// <param name="StoppingToken"></param>
+		/// <param name="socketUrl"></param>
+		/// <param name="stoppingToken"></param>
 		/// <returns></returns>
-		private async Task HandleSocketAsync(Uri SocketUrl, CancellationToken StoppingToken)
+		private async Task HandleSocketAsync(Uri socketUrl, CancellationToken stoppingToken)
 		{
-			using ClientWebSocket Socket = new ClientWebSocket();
-			await Socket.ConnectAsync(SocketUrl, StoppingToken);
+			using ClientWebSocket socket = new ClientWebSocket();
+			await socket.ConnectAsync(socketUrl, stoppingToken);
 
-			byte[] Buffer = new byte[2048];
-			while (!StoppingToken.IsCancellationRequested)
+			byte[] buffer = new byte[2048];
+			while (!stoppingToken.IsCancellationRequested)
 			{
 				// Read the next message
-				int Length = 0;
+				int length = 0;
 				for (; ; )
 				{
-					if (Length == Buffer.Length)
+					if (length == buffer.Length)
 					{
-						Array.Resize(ref Buffer, Buffer.Length + 2048);
+						Array.Resize(ref buffer, buffer.Length + 2048);
 					}
 
-					WebSocketReceiveResult Result = await Socket.ReceiveAsync(new ArraySegment<byte>(Buffer, Length, Buffer.Length - Length), StoppingToken);
-					if (Result.MessageType == WebSocketMessageType.Close)
+					WebSocketReceiveResult result = await socket.ReceiveAsync(new ArraySegment<byte>(buffer, length, buffer.Length - length), stoppingToken);
+					if (result.MessageType == WebSocketMessageType.Close)
 					{
 						return;
 					}
-					Length += Result.Count;
+					length += result.Count;
 
-					if (Result.EndOfMessage)
+					if (result.EndOfMessage)
 					{
 						break;
 					}
 				}
 
 				// Get the message data
-				Logger.LogInformation("Slack event: {Message}", Encoding.UTF8.GetString(Buffer, 0, Length));
-				EventMessage EventMessage = JsonSerializer.Deserialize<EventMessage>(Buffer.AsSpan(0, Length))!;
+				_logger.LogInformation("Slack event: {Message}", Encoding.UTF8.GetString(buffer, 0, length));
+				EventMessage eventMessage = JsonSerializer.Deserialize<EventMessage>(buffer.AsSpan(0, length))!;
 
 				// Acknowledge the message
-				if (EventMessage.EnvelopeId != null)
+				if (eventMessage.EnvelopeId != null)
 				{
-					object Response = new { EventMessage.EnvelopeId };
-					await Socket.SendAsync(JsonSerializer.SerializeToUtf8Bytes(Response), WebSocketMessageType.Text, true, StoppingToken);
+					object response = new { eventMessage.EnvelopeId };
+					await socket.SendAsync(JsonSerializer.SerializeToUtf8Bytes(response), WebSocketMessageType.Text, true, stoppingToken);
 				}
 
 				// Handle the message type
-				if (EventMessage.Type != null)
+				if (eventMessage.Type != null)
 				{
-					string Type = EventMessage.Type;
-					if (Type.Equals("disconnect", StringComparison.Ordinal))
+					string type = eventMessage.Type;
+					if (type.Equals("disconnect", StringComparison.Ordinal))
 					{
 						break;
 					}
-					else if (Type.Equals("interactive", StringComparison.Ordinal))
+					else if (type.Equals("interactive", StringComparison.Ordinal))
 					{
-						await HandleInteractionMessage(EventMessage);
+						await HandleInteractionMessage(eventMessage);
 					}
 					else
 					{
-						Logger.LogDebug("Unhandled event type ({Type})", Type);
+						_logger.LogDebug("Unhandled event type ({Type})", type);
 					}
 				}
 			}
@@ -1424,19 +1418,19 @@ namespace Horde.Build.Notifications.Impl
 		/// <summary>
 		/// Handle a button being clicked
 		/// </summary>
-		/// <param name="Message">The event message</param>
+		/// <param name="message">The event message</param>
 		/// <returns></returns>
-		private async Task HandleInteractionMessage(EventMessage Message)
+		private async Task HandleInteractionMessage(EventMessage message)
 		{
-			if (Message.Payload != null && Message.Payload.User != null && Message.Payload.User.UserName != null && String.Equals(Message.Payload.Type, "block_actions", StringComparison.Ordinal))
+			if (message.Payload != null && message.Payload.User != null && message.Payload.User.UserName != null && String.Equals(message.Payload.Type, "block_actions", StringComparison.Ordinal))
 			{
-				foreach (ActionInfo Action in Message.Payload.Actions)
+				foreach (ActionInfo action in message.Payload.Actions)
 				{
-					if (Action.Value != null)
+					if (action.Value != null)
 					{
-						if (Action.Value.StartsWith("issue_", StringComparison.Ordinal))
+						if (action.Value.StartsWith("issue_", StringComparison.Ordinal))
 						{
-							await HandleIssueResponseAsync(Message.Payload, Action);
+							await HandleIssueResponseAsync(message.Payload, action);
 						}
 					}
 				}
diff --git a/Engine/Source/Programs/Horde/Horde.Build/Notifications/Impl/SubscriptionCollection.cs b/Engine/Source/Programs/Horde/Horde.Build/Notifications/Impl/SubscriptionCollection.cs
--- a/Engine/Source/Programs/Horde/Horde.Build/Notifications/Impl/SubscriptionCollection.cs
+++ b/Engine/Source/Programs/Horde/Horde.Build/Notifications/Impl/SubscriptionCollection.cs
@@ -1,22 +1,19 @@
 // Copyright Epic Games, Inc. All Rights Reserved.
 
+using System;
+using System.Collections.Generic;
+using System.Globalization;
+using System.Linq;
+using System.Text;
+using System.Threading.Tasks;
 using EpicGames.Core;
-using HordeCommon;
 using Horde.Build.Api;
 using Horde.Build.Models;
 using Horde.Build.Services;
 using Horde.Build.Utilities;
-using MongoDB.Bson;
+using HordeCommon;
 using MongoDB.Bson.Serialization.Attributes;
 using MongoDB.Driver;
-using Polly;
-using System;
-using System.Collections.Generic;
-using System.Globalization;
-using System.Linq;
-using System.Reflection.Metadata.Ecma335;
-using System.Text;
-using System.Threading.Tasks;
 
 namespace Horde.Build.Collections.Impl
 {
@@ -32,24 +29,24 @@ namespace Horde.Build.Collections.Impl
 		[BsonKnownTypes(typeof(JobCompleteEvent), typeof(LabelCompleteEvent), typeof(StepCompleteEvent))]
 		abstract class Event : IEvent
 		{
-			public static Event FromRecord(EventRecord Record)
+			public static Event FromRecord(EventRecord record)
 			{
-				JobCompleteEventRecord? JobCompleteRecord = Record as JobCompleteEventRecord;
-				if (JobCompleteRecord != null)
+				JobCompleteEventRecord? jobCompleteRecord = record as JobCompleteEventRecord;
+				if (jobCompleteRecord != null)
 				{
-					return new JobCompleteEvent(JobCompleteRecord);
+					return new JobCompleteEvent(jobCompleteRecord);
 				}
 
-				LabelCompleteEventRecord? LabelCompleteRecord = Record as LabelCompleteEventRecord;
-				if (LabelCompleteRecord != null)
+				LabelCompleteEventRecord? labelCompleteRecord = record as LabelCompleteEventRecord;
+				if (labelCompleteRecord != null)
 				{
-					return new LabelCompleteEvent(LabelCompleteRecord);
+					return new LabelCompleteEvent(labelCompleteRecord);
 				}
 
-				StepCompleteEventRecord? StepCompleteRecord = Record as StepCompleteEventRecord;
-				if (StepCompleteRecord != null)
+				StepCompleteEventRecord? stepCompleteRecord = record as StepCompleteEventRecord;
+				if (stepCompleteRecord != null)
 				{
-					return new StepCompleteEvent(StepCompleteRecord);
+					return new StepCompleteEvent(stepCompleteRecord);
 				}
 
 				throw new ArgumentException("Invalid record type");
@@ -68,11 +65,11 @@ namespace Horde.Build.Collections.Impl
 			{
 			}
 
-			public JobCompleteEvent(JobCompleteEventRecord Record)
+			public JobCompleteEvent(JobCompleteEventRecord record)
 			{
-				StreamId = new StreamId(Record.StreamId);
-				TemplateId = new TemplateRefId(Record.TemplateId);
-				Outcome = Record.Outcome;
+				StreamId = new StreamId(record.StreamId);
+				TemplateId = new TemplateRefId(record.TemplateId);
+				Outcome = record.Outcome;
 			}
 
 			public override EventRecord ToRecord()
@@ -99,13 +96,13 @@ namespace Horde.Build.Collections.Impl
 				LabelName = String.Empty;
 			}
 
-			public LabelCompleteEvent(LabelCompleteEventRecord Record)
+			public LabelCompleteEvent(LabelCompleteEventRecord record)
 			{
-				StreamId = new StreamId(Record.StreamId);
-				TemplateId = new TemplateRefId(Record.TemplateId);
-				CategoryName = Record.CategoryName;
-				LabelName = Record.LabelName;
-				Outcome = Record.Outcome;
+				StreamId = new StreamId(record.StreamId);
+				TemplateId = new TemplateRefId(record.TemplateId);
+				CategoryName = record.CategoryName;
+				LabelName = record.LabelName;
+				Outcome = record.Outcome;
 			}
 
 			public override EventRecord ToRecord()
@@ -115,14 +112,14 @@ namespace Horde.Build.Collections.Impl
 
 			public override string ToString()
 			{
-				StringBuilder Result = new StringBuilder();
-				Result.Append(CultureInfo.InvariantCulture, $"stream={StreamId}, template={TemplateId}");
+				StringBuilder result = new StringBuilder();
+				result.Append(CultureInfo.InvariantCulture, $"stream={StreamId}, template={TemplateId}");
 				if(CategoryName != null)
 				{
-					Result.Append(CultureInfo.InvariantCulture, $", category={CategoryName}");
+					result.Append(CultureInfo.InvariantCulture, $", category={CategoryName}");
 				}
-				Result.Append(CultureInfo.InvariantCulture, $", label={LabelName}, outcome={Outcome}");
-				return Result.ToString();
+				result.Append(CultureInfo.InvariantCulture, $", label={LabelName}, outcome={Outcome}");
+				return result.ToString();
 			}
 		}
 
@@ -138,12 +135,12 @@ namespace Horde.Build.Collections.Impl
 				StepName = String.Empty;
 			}
 
-			public StepCompleteEvent(StepCompleteEventRecord Record)
+			public StepCompleteEvent(StepCompleteEventRecord record)
 			{
-				StreamId = new StreamId(Record.StreamId);
-				TemplateId = new TemplateRefId(Record.TemplateId);
-				StepName = Record.StepName;
-				Outcome = Record.Outcome;
+				StreamId = new StreamId(record.StreamId);
+				TemplateId = new TemplateRefId(record.TemplateId);
+				StepName = record.StepName;
+				Outcome = record.Outcome;
 			}
 
 			public override EventRecord ToRecord()
@@ -169,83 +166,83 @@ namespace Horde.Build.Collections.Impl
 			[BsonConstructor]
 			private Subscription()
 			{
-				this.Id = String.Empty;
-				this.Event = null!;
+				Id = String.Empty;
+				Event = null!;
 			}
 
-			public Subscription(NewSubscription Subscription)
+			public Subscription(NewSubscription subscription)
 			{
-				this.Event = Event.FromRecord(Subscription.Event);
-				this.UserId = Subscription.UserId;
-				this.NotificationType = Subscription.NotificationType;
-				this.Id = ContentHash.SHA1($"{Event}, user={UserId}, type={NotificationType.ToString()}").ToString();
+				Event = Event.FromRecord(subscription.Event);
+				UserId = subscription.UserId;
+				NotificationType = subscription.NotificationType;
+				Id = ContentHash.SHA1($"{Event}, user={UserId}, type={NotificationType}").ToString();
 			}
 		}
 
-		IMongoCollection<Subscription> Collection;
+		readonly IMongoCollection<Subscription> _collection;
 
 		/// <summary>
 		/// Constructor
 		/// </summary>
-		/// <param name="DatabaseService">The database service</param>
-		public SubscriptionCollection(DatabaseService DatabaseService)
+		/// <param name="databaseService">The database service</param>
+		public SubscriptionCollection(DatabaseService databaseService)
 		{
-			Collection = DatabaseService.GetCollection<Subscription>("SubscriptionsV2");
-			if (!DatabaseService.ReadOnlyMode)
+			_collection = databaseService.GetCollection<Subscription>("SubscriptionsV2");
+			if (!databaseService.ReadOnlyMode)
 			{
-				Collection.Indexes.CreateOne(new CreateIndexModel<Subscription>(Builders<Subscription>.IndexKeys.Ascending(x => x.Event)));
-				Collection.Indexes.CreateOne(new CreateIndexModel<Subscription>(Builders<Subscription>.IndexKeys.Ascending(x => x.UserId)));
+				_collection.Indexes.CreateOne(new CreateIndexModel<Subscription>(Builders<Subscription>.IndexKeys.Ascending(x => x.Event)));
+				_collection.Indexes.CreateOne(new CreateIndexModel<Subscription>(Builders<Subscription>.IndexKeys.Ascending(x => x.UserId)));
 			}
 		}
 
 		/// <inheritdoc/>
-		public async Task<List<ISubscription>> AddAsync(IEnumerable<NewSubscription> NewSubscriptions)
+		public async Task<List<ISubscription>> AddAsync(IEnumerable<NewSubscription> newSubscriptions)
 		{
-			List<Subscription> NewDocuments = NewSubscriptions.Select(x => new Subscription(x)).ToList();
+			List<Subscription> newDocuments = newSubscriptions.Select(x => new Subscription(x)).ToList();
 			try
 			{
-				await Collection.InsertManyAsync(NewDocuments, new InsertManyOptions { IsOrdered = false });
+				await _collection.InsertManyAsync(newDocuments, new InsertManyOptions { IsOrdered = false });
 			}
-			catch (MongoBulkWriteException Ex)
+			catch (MongoBulkWriteException ex)
 			{
-				foreach (WriteError WriteError in Ex.WriteErrors)
+				foreach (WriteError writeError in ex.WriteErrors)
 				{
-					if (WriteError.Category != ServerErrorCategory.DuplicateKey)
+					if (writeError.Category != ServerErrorCategory.DuplicateKey)
 					{
 						throw;
 					}
 				}
 			}
-			return NewDocuments.ConvertAll<ISubscription>(x => x);
+			return newDocuments.ConvertAll<ISubscription>(x => x);
 		}
 
 		/// <inheritdoc/>
-		public async Task RemoveAsync(IEnumerable<ISubscription> Subscriptions)
+		public async Task RemoveAsync(IEnumerable<ISubscription> subscriptions)
 		{
-			FilterDefinition<Subscription> Filter = Builders<Subscription>.Filter.In(x => x.Id, Subscriptions.Select(x => ((Subscription)x).Id));
-			await Collection.DeleteManyAsync(Filter);
+			FilterDefinition<Subscription> filter = Builders<Subscription>.Filter.In(x => x.Id, subscriptions.Select(x => ((Subscription)x).Id));
+			await _collection.DeleteManyAsync(filter);
 		}
 
 		/// <inheritdoc/>
-		public async Task<ISubscription?> GetAsync(string SubscriptionId)
+		public async Task<ISubscription?> GetAsync(string subscriptionId)
 		{
-			FilterDefinition<Subscription> Filter = Builders<Subscription>.Filter.Eq(x => x.Id, SubscriptionId);
-			return await Collection.Find(Filter).FirstOrDefaultAsync();
+			FilterDefinition<Subscription> filter = Builders<Subscription>.Filter.Eq(x => x.Id, subscriptionId);
+			return await _collection.Find(filter).FirstOrDefaultAsync();
 		}
 
 		/// <inheritdoc/>
-		public async Task<List<ISubscription>> FindSubscribersAsync(EventRecord EventRecord)
+		public async Task<List<ISubscription>> FindSubscribersAsync(EventRecord eventRecord)
 		{
-			FilterDefinition<Subscription> Filter = Builders<Subscription>.Filter.Eq(x => x.Event, Event.FromRecord(EventRecord));
-			List<Subscription> Results = await Collection.Find(Filter).ToListAsync();
-			return Results.ConvertAll<ISubscription>(x => x);
+			FilterDefinition<Subscription> filter = Builders<Subscription>.Filter.Eq(x => x.Event, Event.FromRecord(eventRecord));
+			List<Subscription> results = await _collection.Find(filter).ToListAsync();
+			return results.ConvertAll<ISubscription>(x => x);
 		}
 
 		/// <inheritdoc/>
-		public async Task<List<ISubscription>> FindSubscriptionsAsync(UserId UserId)
+		public async Task<List<ISubscription>> FindSubscriptionsAsync(UserId userId)
 		{
-			List<Subscription> Results = await Collection.Find(x => x.UserId == UserId).ToListAsync();
-			return Results.ConvertAll<ISubscription>(x => x);
+			List<Subscription> results = await _collection.Find(x => x.UserId == userId).ToListAsync();
+			return results.ConvertAll<ISubscription>(x => x);
 		}
 	}
 }
