Horde: Fix coding conventions within Horde.Build.

diff --git a/Engine/Source/Programs/Horde/Horde.Build/Perforce/Change.cs b/Engine/Source/Programs/Horde/Horde.Build/Perforce/Change.cs
--- a/Engine/Source/Programs/Horde/Horde.Build/Perforce/Change.cs
+++ b/Engine/Source/Programs/Horde/Horde.Build/Perforce/Change.cs
@@ -1,10 +1,5 @@
 // Copyright Epic Games, Inc. All Rights Reserved.
 
-using EpicGames.Core;
-using Horde.Build.Api;
-using Horde.Build.Utilities;
-using MongoDB.Bson;
-using MongoDB.Bson.Serialization.Attributes;
 using System;
 using System.Collections.Generic;
 using System.Diagnostics;
@@ -12,12 +7,10 @@ using System.Diagnostics.CodeAnalysis;
 using System.Globalization;
 using System.Linq;
 using System.Text.RegularExpressions;
-using System.Threading.Tasks;
+using EpicGames.Core;
 
 namespace Horde.Build.Models
 {
-	using UserId = ObjectId<IUser>;
-
 	/// <summary>
 	/// Summary information for a change
 	/// </summary>
@@ -47,16 +40,16 @@ namespace Horde.Build.Models
 		/// <summary>
 		/// Constructor
 		/// </summary>
-		/// <param name="Number">Changelist number</param>
-		/// <param name="Author">Author of the change</param>
-		/// <param name="Path">Base path for modified files</param>
-		/// <param name="Description">Changelist description</param>
-		public ChangeSummary(int Number, IUser Author, string Path, string Description)
+		/// <param name="number">Changelist number</param>
+		/// <param name="author">Author of the change</param>
+		/// <param name="path">Base path for modified files</param>
+		/// <param name="description">Changelist description</param>
+		public ChangeSummary(int number, IUser author, string path, string description)
 		{
-			this.Number = Number;
-			this.Author = Author;
-			this.Path = Path;
-			this.Description = Description;
+			Number = number;
+			Author = author;
+			Path = path;
+			Description = description;
 		}
 	}
 
@@ -115,20 +108,20 @@ namespace Horde.Build.Models
 		/// <summary>
 		/// Constructor
 		/// </summary>
-		/// <param name="Path"></param>
-		/// <param name="DepotPath"></param>
-		/// <param name="Revision"></param>
-		/// <param name="Length"></param>
-		/// <param name="Digest"></param>
-		/// <param name="Type"></param>
-		public ChangeFile(string Path, string DepotPath, int Revision, long Length, Md5Hash? Digest, string Type)
+		/// <param name="path"></param>
+		/// <param name="depotPath"></param>
+		/// <param name="revision"></param>
+		/// <param name="length"></param>
+		/// <param name="digest"></param>
+		/// <param name="type"></param>
+		public ChangeFile(string path, string depotPath, int revision, long length, Md5Hash? digest, string type)
 		{
-			this.Path = Path;
-			this.DepotPath = DepotPath;
-			this.Revision = Revision;
-			this.Length = Length;
-			this.Digest = Digest;
-			this.Type = Type;
+			Path = path;
+			DepotPath = depotPath;
+			Revision = revision;
+			Length = length;
+			Digest = digest;
+			Type = type;
 		}
 	}
 
@@ -141,7 +134,7 @@ namespace Horde.Build.Models
 		/// <summary>
 		/// Set of extensions to treat as code
 		/// </summary>
-		static readonly HashSet<string> CodeExtensions = new HashSet<string>
+		static readonly HashSet<string> s_codeExtensions = new HashSet<string>
 		{
 			".c",
 			".cc",
@@ -194,20 +187,20 @@ namespace Horde.Build.Models
 		/// <summary>
 		/// Constructor
 		/// </summary>
-		/// <param name="Number">Changelist number</param>
-		/// <param name="Author">Author of the change</param>
-		/// <param name="Path">Base path for modified files</param>
-		/// <param name="Description">Changelist description</param>
-		/// <param name="Files">List of files modified, relative to the stream base</param>
-		/// <param name="Date">Date that the change was submitted</param>
-		public ChangeDetails(int Number, IUser Author, string Path, string Description, List<ChangeFile> Files, DateTime Date)
+		/// <param name="number">Changelist number</param>
+		/// <param name="author">Author of the change</param>
+		/// <param name="path">Base path for modified files</param>
+		/// <param name="description">Changelist description</param>
+		/// <param name="files">List of files modified, relative to the stream base</param>
+		/// <param name="date">Date that the change was submitted</param>
+		public ChangeDetails(int number, IUser author, string path, string description, List<ChangeFile> files, DateTime date)
 		{
-			this.Number = Number;
-			this.Author = Author;
-			this.Path = Path;
-			this.Description = Description;
-			this.Files = Files;
-			this.Date = Date;
+			Number = number;
+			Author = author;
+			Path = path;
+			Description = description;
+			Files = files;
+			Date = date;
 		}
 
 		/// <summary>
@@ -216,44 +209,44 @@ namespace Horde.Build.Models
 		/// <returns>True if this change is a code change</returns>
 		public ChangeContentFlags GetContentFlags()
 		{
-			ChangeContentFlags Scope = 0;
+			ChangeContentFlags scope = 0;
 
 			// Check whether the files are code or content
-			foreach (ChangeFile File in Files)
+			foreach (ChangeFile file in Files)
 			{
-				if (CodeExtensions.Any(Extension => File.Path.EndsWith(Extension, StringComparison.OrdinalIgnoreCase)))
+				if (s_codeExtensions.Any(extension => file.Path.EndsWith(extension, StringComparison.OrdinalIgnoreCase)))
 				{
-					Scope |= ChangeContentFlags.ContainsCode;
+					scope |= ChangeContentFlags.ContainsCode;
 				}
 				else
 				{
-					Scope |= ChangeContentFlags.ContainsContent;
+					scope |= ChangeContentFlags.ContainsContent;
 				}
 
-				if (Scope == (ChangeContentFlags.ContainsCode | ChangeContentFlags.ContainsContent))
+				if (scope == (ChangeContentFlags.ContainsCode | ChangeContentFlags.ContainsContent))
 				{
 					break;
 				}
 			}
-			return Scope;
+			return scope;
 		}
 
 		/// <summary>
 		/// Attempts to parse the Robomerge source from this commit information
 		/// </summary>
-		/// <param name="Source">On success, receives the source information</param>
+		/// <param name="source">On success, receives the source information</param>
 		/// <returns>True if the commit was merged from another stream</returns>
-		public bool TryParseRobomergeSource([NotNullWhen(true)] out (string, int)? Source)
+		public bool TryParseRobomergeSource([NotNullWhen(true)] out (string, int)? source)
 		{
-			Match Match = Regex.Match(Description, @"#ROBOMERGE-SOURCE: CL (\d+) in (//[^ ]*)/...", RegexOptions.Multiline);
-			if (Match.Success)
+			Match match = Regex.Match(Description, @"#ROBOMERGE-SOURCE: CL (\d+) in (//[^ ]*)/...", RegexOptions.Multiline);
+			if (match.Success)
 			{
-				Source = (Match.Groups[2].Value, int.Parse(Match.Groups[1].Value, CultureInfo.InvariantCulture));
+				source = (match.Groups[2].Value, Int32.Parse(match.Groups[1].Value, CultureInfo.InvariantCulture));
 				return true;
 			}
 			else
 			{
-				Source = null;
+				source = null;
 				return false;
 			}
 		}
@@ -287,16 +280,16 @@ namespace Horde.Build.Models
 		/// <summary>
 		/// Information about a file
 		/// </summary>
-		/// <param name="DepotPath">Depot </param>
-		/// <param name="Exists"></param>
-		/// <param name="Change"></param>
-		/// <param name="Error"></param>
-		public FileSummary(string DepotPath, bool Exists, int Change, string? Error = null)
+		/// <param name="depotPath">Depot </param>
+		/// <param name="exists"></param>
+		/// <param name="change"></param>
+		/// <param name="error"></param>
+		public FileSummary(string depotPath, bool exists, int change, string? error = null)
 		{
-			this.DepotPath = DepotPath;
-			this.Exists = Exists;
-			this.Change = Change;
-			this.Error = Error;
+			DepotPath = depotPath;
+			Exists = exists;
+			Change = change;
+			Error = error;
 		}
 	}
 }
diff --git a/Engine/Source/Programs/Horde/Horde.Build/Perforce/ICommit.cs b/Engine/Source/Programs/Horde/Horde.Build/Perforce/ICommit.cs
--- a/Engine/Source/Programs/Horde/Horde.Build/Perforce/ICommit.cs
+++ b/Engine/Source/Programs/Horde/Horde.Build/Perforce/ICommit.cs
@@ -1,16 +1,8 @@
 ﻿// Copyright Epic Games, Inc. All Rights Reserved.
 
-using EpicGames.Core;
-using EpicGames.Horde.Storage;
-using EpicGames.Perforce;
+using System;
 using Horde.Build.Models;
 using Horde.Build.Utilities;
-using MongoDB.Bson;
-using System;
-using System.Collections.Generic;
-using System.Linq;
-using System.Text;
-using System.Threading.Tasks;
 
 namespace Horde.Build.Commits
 {
@@ -101,24 +93,24 @@ namespace Horde.Build.Commits
 		/// <summary>
 		/// Constructor
 		/// </summary>
-		public NewCommit(ICommit Commit)
-			: this(Commit.StreamId, Commit.Change, Commit.OriginalChange, Commit.AuthorId, Commit.OwnerId, Commit.Description, Commit.BasePath, Commit.DateUtc)
+		public NewCommit(ICommit commit)
+			: this(commit.StreamId, commit.Change, commit.OriginalChange, commit.AuthorId, commit.OwnerId, commit.Description, commit.BasePath, commit.DateUtc)
 		{
 		}
 
 		/// <summary>
 		/// Constructor
 		/// </summary>
-		public NewCommit(StreamId StreamId, int Change, int OriginalChange, UserId AuthorId, UserId OwnerId, string Description, string BasePath, DateTime DateUtc)
+		public NewCommit(StreamId streamId, int change, int originalChange, UserId authorId, UserId ownerId, string description, string basePath, DateTime dateUtc)
 		{
-			this.StreamId = StreamId;
-			this.Change = Change;
-			this.OriginalChange = OriginalChange;
-			this.AuthorId = AuthorId;
-			this.OwnerId = OwnerId;
-			this.Description = Description;
-			this.BasePath = BasePath;
-			this.DateUtc = DateUtc;
+			StreamId = streamId;
+			Change = change;
+			OriginalChange = originalChange;
+			AuthorId = authorId;
+			OwnerId = ownerId;
+			Description = description;
+			BasePath = basePath;
+			DateUtc = dateUtc;
 		}
 	}
 }
diff --git a/Engine/Source/Programs/Horde/Horde.Build/Perforce/ICommitCollection.cs b/Engine/Source/Programs/Horde/Horde.Build/Perforce/ICommitCollection.cs
--- a/Engine/Source/Programs/Horde/Horde.Build/Perforce/ICommitCollection.cs
+++ b/Engine/Source/Programs/Horde/Horde.Build/Perforce/ICommitCollection.cs
@@ -1,18 +1,14 @@
 ﻿// Copyright Epic Games, Inc. All Rights Reserved.
 
-using Horde.Build.Models;
-using System;
 using System.Collections.Generic;
-using System.Linq;
 using System.Threading.Tasks;
-using EpicGames.Core;
-using MongoDB.Bson;
+using Horde.Build.Models;
 using Horde.Build.Utilities;
 
 namespace Horde.Build.Commits
 {
-	using StreamId = StringId<IStream>;
 	using CommitId = ObjectId<ICommit>;
+	using StreamId = StringId<IStream>;
 
 	/// <summary>
 	/// Stores a collection of commits
@@ -22,26 +18,26 @@ namespace Horde.Build.Commits
 		/// <summary>
 		/// Adds or replaces an existing commit
 		/// </summary>
-		/// <param name="NewCommit">The new commit to add</param>
+		/// <param name="newCommit">The new commit to add</param>
 		/// <returns>The commit that was created</returns>
-		Task<ICommit> AddOrReplaceAsync(NewCommit NewCommit);
+		Task<ICommit> AddOrReplaceAsync(NewCommit newCommit);
 
 		/// <summary>
 		/// Gets a single commit
 		/// </summary>
-		/// <param name="Id">Identifier for the commit</param>
-		Task<ICommit?> GetCommitAsync(CommitId Id);
+		/// <param name="id">Identifier for the commit</param>
+		Task<ICommit?> GetCommitAsync(CommitId id);
 
 		/// <summary>
 		/// Finds commits matching certain criteria
 		/// </summary>
-		/// <param name="StreamId"></param>
-		/// <param name="MinChange"></param>
-		/// <param name="MaxChange"></param>
-		/// <param name="Index"></param>
-		/// <param name="Count"></param>
+		/// <param name="streamId"></param>
+		/// <param name="minChange"></param>
+		/// <param name="maxChange"></param>
+		/// <param name="index"></param>
+		/// <param name="count"></param>
 		/// <returns></returns>
-		Task<List<ICommit>> FindCommitsAsync(StreamId StreamId, int? MinChange = null, int? MaxChange = null, int? Index = null, int? Count = null);
+		Task<List<ICommit>> FindCommitsAsync(StreamId streamId, int? minChange = null, int? maxChange = null, int? index = null, int? count = null);
 	}
 
 	/// <summary>
@@ -52,18 +48,18 @@ namespace Horde.Build.Commits
 		/// <summary>
 		/// Gets a commit from a stream by changelist number
 		/// </summary>
-		/// <param name="CommitCollection">The commit collection</param>
-		/// <param name="StreamId"></param>
-		/// <param name="Change"></param>
+		/// <param name="commitCollection">The commit collection</param>
+		/// <param name="streamId"></param>
+		/// <param name="change"></param>
 		/// <returns></returns>
-		public static async Task<ICommit?> GetCommitAsync(this ICommitCollection CommitCollection, StreamId StreamId, int Change)
+		public static async Task<ICommit?> GetCommitAsync(this ICommitCollection commitCollection, StreamId streamId, int change)
 		{
-			List<ICommit> Commits = await CommitCollection.FindCommitsAsync(StreamId, Change, Change);
-			if (Commits.Count == 0)
+			List<ICommit> commits = await commitCollection.FindCommitsAsync(streamId, change, change);
+			if (commits.Count == 0)
 			{
 				return null;
 			}
-			return Commits[0];
+			return commits[0];
 		}
 	}
 }
diff --git a/Engine/Source/Programs/Horde/Horde.Build/Perforce/ICommitService.cs b/Engine/Source/Programs/Horde/Horde.Build/Perforce/ICommitService.cs
--- a/Engine/Source/Programs/Horde/Horde.Build/Perforce/ICommitService.cs
+++ b/Engine/Source/Programs/Horde/Horde.Build/Perforce/ICommitService.cs
@@ -1,13 +1,6 @@
 // Copyright Epic Games, Inc. All Rights Reserved.
 
-using Horde.Build.Models;
 using System;
-using System.Collections.Generic;
-using System.Linq;
-using System.Threading.Tasks;
-using EpicGames.Core;
-using MongoDB.Bson;
-using Horde.Build.Utilities;
 
 namespace Horde.Build.Commits
 {
@@ -19,8 +12,8 @@ namespace Horde.Build.Commits
 		/// <summary>
 		/// Registers a delegate to be called when a new commit is added
 		/// </summary>
-		/// <param name="OnAddCommit">Callback for a new commit being added</param>
+		/// <param name="onAddCommit">Callback for a new commit being added</param>
 		/// <returns>Disposable handler.</returns>
-		IDisposable AddListener(Action<ICommit> OnAddCommit);
+		IDisposable AddListener(Action<ICommit> onAddCommit);
 	}
 }
diff --git a/Engine/Source/Programs/Horde/Horde.Build/Perforce/IPerforceService.cs b/Engine/Source/Programs/Horde/Horde.Build/Perforce/IPerforceService.cs
--- a/Engine/Source/Programs/Horde/Horde.Build/Perforce/IPerforceService.cs
+++ b/Engine/Source/Programs/Horde/Horde.Build/Perforce/IPerforceService.cs
@@ -1,25 +1,12 @@
 // Copyright Epic Games, Inc. All Rights Reserved.
 
-using EpicGames.Core;
-using EpicGames.Perforce;
-using Horde.Build.Models;
-using Microsoft.Extensions.Logging;
-using Polly;
-using Polly.Extensions.Http;
 using System;
 using System.Collections.Generic;
 using System.Diagnostics.CodeAnalysis;
-using System.Globalization;
-using System.IO;
-using System.Linq;
-using System.Net.Http;
-using System.Text;
-using System.Text.Json;
 using System.Text.RegularExpressions;
-using System.Threading;
 using System.Threading.Tasks;
-
-using P4 = Perforce.P4;
+using EpicGames.Perforce;
+using Horde.Build.Models;
 
 namespace Horde.Build.Services
 {
@@ -83,10 +70,10 @@ namespace Horde.Build.Services
 		/// <summary>
 		/// Maps a depot path into the stream
 		/// </summary>
-		/// <param name="DepotPath">The depot path</param>
-		/// <param name="StreamPath">The resulting stream path</param>
+		/// <param name="depotPath">The depot path</param>
+		/// <param name="streamPath">The resulting stream path</param>
 		/// <returns></returns>
-		bool TryGetStreamPath(string DepotPath, [NotNullWhen(true)] out string? StreamPath);
+		bool TryGetStreamPath(string depotPath, [NotNullWhen(true)] out string? streamPath);
 	}
 
 	/// <summary>
@@ -97,153 +84,153 @@ namespace Horde.Build.Services
 		/// <summary>
 		/// Finds or adds a user from the given Perforce server, adding the user (and populating their profile with Perforce data) if they do not currently exist
 		/// </summary>
-		/// <param name="ClusterName"></param>
-		/// <param name="UserName"></param>
+		/// <param name="clusterName"></param>
+		/// <param name="userName"></param>
 		/// <returns></returns>
-		public ValueTask<IUser> FindOrAddUserAsync(string ClusterName, string UserName);
+		public ValueTask<IUser> FindOrAddUserAsync(string clusterName, string userName);
 
 		/// <summary>
 		/// 
 		/// </summary>
-		/// <param name="ClusterName"></param>
+		/// <param name="clusterName"></param>
 		/// <returns></returns>
-		public Task<IPerforceConnection?> GetServiceUserConnection(string? ClusterName);
+		public Task<IPerforceConnection?> GetServiceUserConnection(string? clusterName);
 
 		/// <summary>
 		/// Gets the definition of a stream
 		/// </summary>
-		/// <param name="ClusterName">Name of the Perforce cluster</param>
-		/// <param name="StreamName">The stream name</param>
+		/// <param name="clusterName">Name of the Perforce cluster</param>
+		/// <param name="streamName">The stream name</param>
 		/// <returns></returns>
-		public Task<IStreamView> GetStreamViewAsync(string ClusterName, string StreamName);
+		public Task<IStreamView> GetStreamViewAsync(string clusterName, string streamName);
 
 		/// <summary>
 		/// Create a new changelist by submitting the given file
 		/// </summary>
-		/// <param name="ClusterName">Name of the Perforce cluster</param>
-		/// <param name="StreamName">The stream to query</param>
-		/// <param name="Path">Path for the file to submit</param>
-		/// <param name="Description">Description for the changelist</param>
+		/// <param name="clusterName">Name of the Perforce cluster</param>
+		/// <param name="streamName">The stream to query</param>
+		/// <param name="path">Path for the file to submit</param>
+		/// <param name="description">Description for the changelist</param>
 		/// <returns>New changelist number</returns>
-		public Task<int> CreateNewChangeAsync(string ClusterName, string StreamName, string Path, string Description);
+		public Task<int> CreateNewChangeAsync(string clusterName, string streamName, string path, string description);
 
 		/// <summary>
 		/// Gets the code change corresponding to an actual change submitted to a stream
 		/// </summary>
-		/// <param name="ClusterName">Name of the Perforce cluster</param>
-		/// <param name="StreamName"></param>
-		/// <param name="Change">The changelist number to query</param>
+		/// <param name="clusterName">Name of the Perforce cluster</param>
+		/// <param name="streamName"></param>
+		/// <param name="change">The changelist number to query</param>
 		/// <returns>Code change for the latest change</returns>
-		public Task<int> GetCodeChangeAsync(string ClusterName, string StreamName, int Change);
+		public Task<int> GetCodeChangeAsync(string clusterName, string streamName, int change);
 
 		/// <summary>
 		/// Gets information about the given user
 		/// </summary>
-		/// <param name="ClusterName">Name of the Perforce cluster</param>
-		/// <param name="UserName">The user name</param>
+		/// <param name="clusterName">Name of the Perforce cluster</param>
+		/// <param name="userName">The user name</param>
 		/// <returns>User information</returns>
-		public Task<PerforceUserInfo?> GetUserInfoAsync(string ClusterName, string UserName);
+		public Task<PerforceUserInfo?> GetUserInfoAsync(string clusterName, string userName);
 
 		/// <summary>
 		/// Finds changes submitted to a depot Gets the latest change for a particular stream
 		/// </summary>
-		/// <param name="ClusterName">Name of the Perforce cluster</param>
-		/// <param name="MinChange">The minimum changelist number</param>
-		/// <param name="MaxChange">The maximum changelist number</param>
-		/// <param name="MaxResults"></param>
+		/// <param name="clusterName">Name of the Perforce cluster</param>
+		/// <param name="minChange">The minimum changelist number</param>
+		/// <param name="maxChange">The maximum changelist number</param>
+		/// <param name="maxResults"></param>
 		/// <returns>Changelist information</returns>
-		public Task<List<ChangeSummary>> GetChangesAsync(string ClusterName, int? MinChange, int? MaxChange, int MaxResults);
+		public Task<List<ChangeSummary>> GetChangesAsync(string clusterName, int? minChange, int? maxChange, int maxResults);
 
 		/// <summary>
 		/// Gets the latest change for a particular stream
 		/// </summary>
-		/// <param name="ClusterName">Name of the Perforce cluster</param>
-		/// <param name="StreamName">The stream to query</param>
-		/// <param name="MinChange">The minimum changelist number</param>
-		/// <param name="MaxChange">The maximum changelist number</param>
-		/// <param name="Results">Number of results to return</param>
-		/// <param name="ImpersonateUser">Name of the user to impersonate</param>
+		/// <param name="clusterName">Name of the Perforce cluster</param>
+		/// <param name="streamName">The stream to query</param>
+		/// <param name="minChange">The minimum changelist number</param>
+		/// <param name="maxChange">The maximum changelist number</param>
+		/// <param name="results">Number of results to return</param>
+		/// <param name="impersonateUser">Name of the user to impersonate</param>
 		/// <returns>Latest changelist number</returns>
-		public Task<List<ChangeSummary>> GetChangesAsync(string ClusterName, string StreamName, int? MinChange, int? MaxChange, int Results, string? ImpersonateUser);
+		public Task<List<ChangeSummary>> GetChangesAsync(string clusterName, string streamName, int? minChange, int? maxChange, int results, string? impersonateUser);
 
 		/// <summary>
 		/// Checks a shelf is valid for the given stream
 		/// </summary>
-		/// <param name="ClusterName">Name of the Perforce cluster</param>
-		/// <param name="StreamName">The stream to query</param>
-		/// <param name="ChangeNumber">Shelved changelist number</param>
-		/// <param name="ImpersonateUser">Name of the user to impersonate</param>
+		/// <param name="clusterName">Name of the Perforce cluster</param>
+		/// <param name="streamName">The stream to query</param>
+		/// <param name="changeNumber">Shelved changelist number</param>
+		/// <param name="impersonateUser">Name of the user to impersonate</param>
 		/// <returns></returns>
-		public Task<CheckShelfResult> CheckShelfAsync(string ClusterName, string StreamName, int ChangeNumber, string? ImpersonateUser);
+		public Task<CheckShelfResult> CheckShelfAsync(string clusterName, string streamName, int changeNumber, string? impersonateUser);
 
 		/// <summary>
 		/// Gets the latest change for a particular stream
 		/// </summary>
-		/// <param name="ClusterName">Name of the Perforce cluster</param>
-		/// <param name="StreamName">The stream to query</param>
-		/// <param name="ChangeNumber">Change numbers to query</param>
+		/// <param name="clusterName">Name of the Perforce cluster</param>
+		/// <param name="streamName">The stream to query</param>
+		/// <param name="changeNumber">Change numbers to query</param>
 		/// <returns>Commit details</returns>
-		public Task<ChangeDetails> GetChangeDetailsAsync(string ClusterName, string StreamName, int ChangeNumber);
+		public Task<ChangeDetails> GetChangeDetailsAsync(string clusterName, string streamName, int changeNumber);
 
 		/// <summary>
 		/// Gets the latest change for a particular stream
 		/// </summary>
-		/// <param name="ClusterName">Name of the Perforce cluster</param>
-		/// <param name="StreamName">The stream to query</param>
-		/// <param name="ChangeNumbers">Change numbers to query</param>
-		/// <param name="ImpersonateUser">Name of the user to impersonate</param>
+		/// <param name="clusterName">Name of the Perforce cluster</param>
+		/// <param name="streamName">The stream to query</param>
+		/// <param name="changeNumbers">Change numbers to query</param>
+		/// <param name="impersonateUser">Name of the user to impersonate</param>
 		/// <returns>Commit details</returns>
-		public Task<List<ChangeDetails>> GetChangeDetailsAsync(string ClusterName, string StreamName, IReadOnlyList<int> ChangeNumbers, string? ImpersonateUser);
+		public Task<List<ChangeDetails>> GetChangeDetailsAsync(string clusterName, string streamName, IReadOnlyList<int> changeNumbers, string? impersonateUser);
 
 		/// <summary>
 		/// Gets the latest changes for a set of depot paths
 		/// </summary>
-		/// <param name="ClusterName">Name of the Perforce cluster</param>
-		/// <param name="Paths"></param>
+		/// <param name="clusterName">Name of the Perforce cluster</param>
+		/// <param name="paths"></param>
 		/// <returns></returns>
-		public Task<List<FileSummary>> FindFilesAsync(string ClusterName, IEnumerable<string> Paths);
+		public Task<List<FileSummary>> FindFilesAsync(string clusterName, IEnumerable<string> paths);
 
 		/// <summary>
 		/// Gets the contents of a file in the depot
 		/// </summary>
-		/// <param name="ClusterName">Name of the Perforce cluster</param>
-		/// <param name="Path">Path to read</param>
+		/// <param name="clusterName">Name of the Perforce cluster</param>
+		/// <param name="path">Path to read</param>
 		/// <returns>Data for the file</returns>
-		public Task<byte[]> PrintAsync(string ClusterName, string Path);
+		public Task<byte[]> PrintAsync(string clusterName, string path);
 
 		/// <summary>
 		/// Duplicates a shelved changelist
 		/// </summary>
-		/// <param name="ClusterName">Name of the Perforce cluster</param>
-		/// <param name="ShelvedChange">The shelved changelist</param>
+		/// <param name="clusterName">Name of the Perforce cluster</param>
+		/// <param name="shelvedChange">The shelved changelist</param>
 		/// <returns>The duplicated changelist</returns>
-		public Task<int> DuplicateShelvedChangeAsync(string ClusterName, int ShelvedChange);
+		public Task<int> DuplicateShelvedChangeAsync(string clusterName, int shelvedChange);
 
 		/// <summary>
 		/// Submit a shelved changelist
 		/// </summary>
-		/// <param name="ClusterName">Name of the Perforce cluster</param>
-		/// <param name="ShelvedChange">The shelved changelist number, created by <see cref="DuplicateShelvedChangeAsync(string,int)"/></param>
-		/// <param name="OriginalChange">The original changelist number</param>
+		/// <param name="clusterName">Name of the Perforce cluster</param>
+		/// <param name="shelvedChange">The shelved changelist number, created by <see cref="DuplicateShelvedChangeAsync(String,Int32)"/></param>
+		/// <param name="originalChange">The original changelist number</param>
 		/// <returns>Tuple consisting of the submitted changelist number and message</returns>
-		public Task<(int? Change, string Message)> SubmitShelvedChangeAsync(string ClusterName, int ShelvedChange, int OriginalChange);
+		public Task<(int? Change, string Message)> SubmitShelvedChangeAsync(string clusterName, int shelvedChange, int originalChange);
 
 		/// <summary>
 		/// Deletes a changelist containing shelved files.
 		/// </summary>
-		/// <param name="ClusterName">Name of the Perforce cluster</param>
-		/// <param name="ShelvedChange">The changelist containing shelved files</param>
+		/// <param name="clusterName">Name of the Perforce cluster</param>
+		/// <param name="shelvedChange">The changelist containing shelved files</param>
 		/// <returns>Async tasy</returns>
-		public Task DeleteShelvedChangeAsync(string ClusterName, int ShelvedChange);
+		public Task DeleteShelvedChangeAsync(string clusterName, int shelvedChange);
 
 		/// <summary>
 		/// Updates a changelist description
 		/// </summary>
-		/// <param name="ClusterName">Name of the Perforce cluster</param>
-		/// <param name="Change">The change to update</param>
-		/// <param name="Description">The new description</param>
+		/// <param name="clusterName">Name of the Perforce cluster</param>
+		/// <param name="change">The change to update</param>
+		/// <param name="description">The new description</param>
 		/// <returns>Async task</returns>
-		public Task UpdateChangelistDescription(string ClusterName, int Change, string Description);
+		public Task UpdateChangelistDescription(string clusterName, int change, string description);
 	}
 
 	/// <summary>
@@ -254,73 +241,73 @@ namespace Horde.Build.Services
 		/// <summary>
 		/// Creates a new change for a template
 		/// </summary>
-		/// <param name="Perforce">The Perforce service instance</param>
-		/// <param name="Stream">Stream containing the template</param>
-		/// <param name="Template">The template being built</param>
+		/// <param name="perforce">The Perforce service instance</param>
+		/// <param name="stream">Stream containing the template</param>
+		/// <param name="template">The template being built</param>
 		/// <returns>New changelist number</returns>
-		public static Task<int> CreateNewChangeForTemplateAsync(this IPerforceService Perforce, IStream Stream, ITemplate Template)
+		public static Task<int> CreateNewChangeForTemplateAsync(this IPerforceService perforce, IStream stream, ITemplate template)
 		{
-			string Description = (Template.SubmitDescription ?? "[Horde] New change for $(TemplateName)").Replace("$(TemplateName)", Template.Name, StringComparison.OrdinalIgnoreCase);
+			string description = (template.SubmitDescription ?? "[Horde] New change for $(TemplateName)").Replace("$(TemplateName)", template.Name, StringComparison.OrdinalIgnoreCase);
 
-			Match Match = Regex.Match(Template.SubmitNewChange!, @"^(//[^/]+/[^/]+)/(.+)$");
-			if (Match.Success)
+			Match match = Regex.Match(template.SubmitNewChange!, @"^(//[^/]+/[^/]+)/(.+)$");
+			if (match.Success)
 			{
-				return Perforce.CreateNewChangeAsync(Stream.ClusterName, Match.Groups[1].Value, Match.Groups[2].Value, Description);
+				return perforce.CreateNewChangeAsync(stream.ClusterName, match.Groups[1].Value, match.Groups[2].Value, description);
 			}
 			else
 			{
-				return Perforce.CreateNewChangeAsync(Stream.ClusterName, Stream.Name, Template.SubmitNewChange!, Description);
+				return perforce.CreateNewChangeAsync(stream.ClusterName, stream.Name, template.SubmitNewChange!, description);
 			}
 		}
 
 		/// <summary>
 		/// Gets the latest change for a particular stream
 		/// </summary>
-		/// <param name="Perforce">The perforce implementation</param>
-		/// <param name="ClusterName">Name of the Perforce cluster</param>
-		/// <param name="StreamName">The stream to query</param>
-		/// <param name="ChangeNumber">Change number to query</param>
-		/// <param name="ImpersonateUser">Name of the user to impersonate</param>
+		/// <param name="perforce">The perforce implementation</param>
+		/// <param name="clusterName">Name of the Perforce cluster</param>
+		/// <param name="streamName">The stream to query</param>
+		/// <param name="changeNumber">Change number to query</param>
+		/// <param name="impersonateUser">Name of the user to impersonate</param>
 		/// <returns>Commit details</returns>
-		public static async Task<ChangeDetails> GetChangeDetailsAsync(this IPerforceService Perforce, string ClusterName, string StreamName, int ChangeNumber, string? ImpersonateUser)
+		public static async Task<ChangeDetails> GetChangeDetailsAsync(this IPerforceService perforce, string clusterName, string streamName, int changeNumber, string? impersonateUser)
 		{
-			List<ChangeDetails> Results = await Perforce.GetChangeDetailsAsync(ClusterName, StreamName, new[] { ChangeNumber }, ImpersonateUser);
-			return Results[0];
+			List<ChangeDetails> results = await perforce.GetChangeDetailsAsync(clusterName, streamName, new[] { changeNumber }, impersonateUser);
+			return results[0];
 		}
 
 		/// <summary>
 		/// Gets the latest change for a particular stream
 		/// </summary>
-		/// <param name="Perforce">The perforce implementation</param>
-		/// <param name="ClusterName">Name of the Perforce cluster</param>
-		/// <param name="StreamName">The stream to query</param>
-		/// <param name="MinChange">The minimum changelist number</param>
-		/// <param name="MaxChange">The maximum changelist number</param>
-		/// <param name="Results">Number of results to return</param>
-		/// <param name="ImpersonateUser">Name of the user to impersonate</param>
+		/// <param name="perforce">The perforce implementation</param>
+		/// <param name="clusterName">Name of the Perforce cluster</param>
+		/// <param name="streamName">The stream to query</param>
+		/// <param name="minChange">The minimum changelist number</param>
+		/// <param name="maxChange">The maximum changelist number</param>
+		/// <param name="results">Number of results to return</param>
+		/// <param name="impersonateUser">Name of the user to impersonate</param>
 		/// <returns>Commit details</returns>
-		public static async Task<List<ChangeDetails>> GetChangeDetailsAsync(this IPerforceService Perforce, string ClusterName, string StreamName, int? MinChange, int? MaxChange, int Results, string? ImpersonateUser)
+		public static async Task<List<ChangeDetails>> GetChangeDetailsAsync(this IPerforceService perforce, string clusterName, string streamName, int? minChange, int? maxChange, int results, string? impersonateUser)
 		{
-			List<ChangeSummary> Changes = await Perforce.GetChangesAsync(ClusterName, StreamName, MinChange, MaxChange, Results, ImpersonateUser);
-			return await Perforce.GetChangeDetailsAsync(ClusterName, StreamName, Changes.ConvertAll(x => x.Number), ImpersonateUser);
+			List<ChangeSummary> changes = await perforce.GetChangesAsync(clusterName, streamName, minChange, maxChange, results, impersonateUser);
+			return await perforce.GetChangeDetailsAsync(clusterName, streamName, changes.ConvertAll(x => x.Number), impersonateUser);
 		}
 
 		/// <summary>
 		/// Get the latest submitted change to the stream
 		/// </summary>
-		/// <param name="Perforce">The perforce implementation</param>
-		/// <param name="ClusterName">Name of the Perforce cluster</param>
-		/// <param name="StreamName">The stream to query</param>
-		/// <param name="ImpersonateUser">Name of the user to impersonate</param>
+		/// <param name="perforce">The perforce implementation</param>
+		/// <param name="clusterName">Name of the Perforce cluster</param>
+		/// <param name="streamName">The stream to query</param>
+		/// <param name="impersonateUser">Name of the user to impersonate</param>
 		/// <returns>Latest changelist number</returns>
-		public static async Task<int> GetLatestChangeAsync(this IPerforceService Perforce, string ClusterName, string StreamName, string? ImpersonateUser)
+		public static async Task<int> GetLatestChangeAsync(this IPerforceService perforce, string clusterName, string streamName, string? impersonateUser)
 		{
-			List<ChangeSummary> Changes = await Perforce.GetChangesAsync(ClusterName, StreamName, null, null, 1, ImpersonateUser);
-			if (Changes.Count == 0)
+			List<ChangeSummary> changes = await perforce.GetChangesAsync(clusterName, streamName, null, null, 1, impersonateUser);
+			if (changes.Count == 0)
 			{
-				throw new Exception($"No changes have been submitted to stream {StreamName}");
+				throw new Exception($"No changes have been submitted to stream {streamName}");
 			}
-			return Changes[0].Number;
+			return changes[0].Number;
 		}
 	}
 }
diff --git a/Engine/Source/Programs/Horde/Horde.Build/Perforce/Impl/CommitCollection.cs b/Engine/Source/Programs/Horde/Horde.Build/Perforce/Impl/CommitCollection.cs
--- a/Engine/Source/Programs/Horde/Horde.Build/Perforce/Impl/CommitCollection.cs
+++ b/Engine/Source/Programs/Horde/Horde.Build/Perforce/Impl/CommitCollection.cs
@@ -1,17 +1,15 @@
 // Copyright Epic Games, Inc. All Rights Reserved.
 
-using EpicGames.Core;
-using Horde.Build.Models;
-using Horde.Build.Services;
-using Horde.Build.Utilities;
-using MongoDB.Bson;
-using MongoDB.Driver;
 using System;
 using System.Collections.Generic;
 using System.Linq;
 using System.Threading.Tasks;
+using Horde.Build.Models;
+using Horde.Build.Services;
+using Horde.Build.Utilities;
+using MongoDB.Bson;
 using MongoDB.Bson.Serialization.Attributes;
-using EpicGames.Horde.Storage;
+using MongoDB.Driver;
 
 namespace Horde.Build.Commits.Impl
 {
@@ -60,67 +58,67 @@ namespace Horde.Build.Commits.Impl
 			{
 			}
 
-			public Commit(NewCommit NewCommit)
+			public Commit(NewCommit newCommit)
 			{
-				this.Change = NewCommit.Change;
-				if (NewCommit.OriginalChange != NewCommit.Change)
+				Change = newCommit.Change;
+				if (newCommit.OriginalChange != newCommit.Change)
 				{
-					this.OriginalChange = NewCommit.OriginalChange;
+					OriginalChange = newCommit.OriginalChange;
 				}
-				this.StreamId = NewCommit.StreamId;
-				this.AuthorId = NewCommit.AuthorId;
-				if (NewCommit.OwnerId != NewCommit.AuthorId)
+				StreamId = newCommit.StreamId;
+				AuthorId = newCommit.AuthorId;
+				if (newCommit.OwnerId != newCommit.AuthorId)
 				{
-					this.OwnerId = NewCommit.OwnerId;
+					OwnerId = newCommit.OwnerId;
 				}
-				this.Description = NewCommit.Description;
-				this.BasePath = NewCommit.BasePath;
-				this.DateUtc = NewCommit.DateUtc;
+				Description = newCommit.Description;
+				BasePath = newCommit.BasePath;
+				DateUtc = newCommit.DateUtc;
 			}
 		}
 
-		IMongoCollection<Commit> Commits;
+		readonly IMongoCollection<Commit> _commits;
 
 		/// <summary>
 		/// Constructor
 		/// </summary>
-		public CommitCollection(DatabaseService DatabaseService)
+		public CommitCollection(DatabaseService databaseService)
 		{
-			Commits = DatabaseService.Database.GetCollection<Commit>("Commits");
+			_commits = databaseService.Database.GetCollection<Commit>("Commits");
 
-			if (!DatabaseService.ReadOnlyMode)
+			if (!databaseService.ReadOnlyMode)
 			{
-				Commits.Indexes.CreateOne(new CreateIndexModel<Commit>(Builders<Commit>.IndexKeys.Ascending(x => x.StreamId).Descending(x => x.Change), new CreateIndexOptions { Unique = true }));
+				_commits.Indexes.CreateOne(new CreateIndexModel<Commit>(Builders<Commit>.IndexKeys.Ascending(x => x.StreamId).Descending(x => x.Change), new CreateIndexOptions { Unique = true }));
 			}
 		}
 
 		/// <inheritdoc/>
-		public async Task<ICommit> AddOrReplaceAsync(NewCommit NewCommit)
+		public async Task<ICommit> AddOrReplaceAsync(NewCommit newCommit)
 		{
-			Commit Commit = new Commit(NewCommit);
-			FilterDefinition<Commit> Filter = Builders<Commit>.Filter.Expr(x => x.StreamId == NewCommit.StreamId && x.Change == NewCommit.Change);
-			return await Commits.FindOneAndReplaceAsync(Filter, new Commit(NewCommit), new FindOneAndReplaceOptions<Commit> { IsUpsert = true, ReturnDocument = ReturnDocument.After });
+			Commit commit = new Commit(newCommit);
+			FilterDefinition<Commit> filter = Builders<Commit>.Filter.Expr(x => x.StreamId == newCommit.StreamId && x.Change == newCommit.Change);
+			return await _commits.FindOneAndReplaceAsync(filter, new Commit(newCommit), new FindOneAndReplaceOptions<Commit> { IsUpsert = true, ReturnDocument = ReturnDocument.After });
 		}
 
 		/// <inheritdoc/>
-		public async Task<ICommit?> GetCommitAsync(CommitId Id)
+		public async Task<ICommit?> GetCommitAsync(CommitId id)
 		{
-			return await Commits.Find(x => x.Id == Id).FirstOrDefaultAsync();
+			return await _commits.Find(x => x.Id == id).FirstOrDefaultAsync();
 		}
 
 		/// <inheritdoc/>
-		public async Task<List<ICommit>> FindCommitsAsync(StreamId StreamId, int? MinChange = null, int? MaxChange = null, int? Index = null, int? Count = null)
+		public async Task<List<ICommit>> FindCommitsAsync(StreamId streamId, int? minChange = null, int? maxChange = null, int? index = null, int? count = null)
 		{
-			FilterDefinition<Commit> Filter = Builders<Commit>.Filter.Eq(x => x.StreamId, StreamId);
-			if (MaxChange != null)
+			FilterDefinition<Commit> filter = Builders<Commit>.Filter.Eq(x => x.StreamId, streamId);
+			if (maxChange != null)
 			{
-				Filter &= Builders<Commit>.Filter.Lte(x => x.Change, MaxChange.Value);
+				filter &= Builders<Commit>.Filter.Lte(x => x.Change, maxChange.Value);
 			}
-			if (MinChange != null)
+			if (minChange != null)
 			{
-				Filter &= Builders<Commit>.Filter.Lte(x => x.Change, MinChange.Value);
+				filter &= Builders<Commit>.Filter.Lte(x => x.Change, minChange.Value);
 			}
-			return await Commits.Find(Filter).SortByDescending(x => x.Change).Range(Index, Count).ToListAsync<Commit, ICommit>();
+			return await _commits.Find(filter).SortByDescending(x => x.Change).Range(index, count).ToListAsync<Commit, ICommit>();
 		}
 	}
 }
