Horde: Fix coding conventions within Horde.Build.

diff --git a/Engine/Source/Programs/Horde/Horde.Build/Perforce/Impl/PerforceService.cs b/Engine/Source/Programs/Horde/Horde.Build/Perforce/Impl/PerforceService.cs
--- a/Engine/Source/Programs/Horde/Horde.Build/Perforce/Impl/PerforceService.cs
+++ b/Engine/Source/Programs/Horde/Horde.Build/Perforce/Impl/PerforceService.cs
@@ -1,13 +1,5 @@
 // Copyright Epic Games, Inc. All Rights Reserved.
 
-using EpicGames.Core;
-using EpicGames.Perforce;
-using Horde.Build.Collections;
-using Horde.Build.Models;
-using Horde.Build.Utilities;
-using Microsoft.Extensions.Caching.Memory;
-using Microsoft.Extensions.Logging;
-using Microsoft.Extensions.Options;
 using System;
 using System.Collections.Generic;
 using System.Diagnostics.CodeAnalysis;
@@ -20,24 +12,31 @@ using System.Runtime.InteropServices;
 using System.Text;
 using System.Text.RegularExpressions;
 using System.Threading.Tasks;
+using EpicGames.Core;
+using EpicGames.Perforce;
+using Horde.Build.Collections;
+using Horde.Build.Models;
+using Horde.Build.Utilities;
+using Microsoft.Extensions.Caching.Memory;
+using Microsoft.Extensions.Logging;
+using Microsoft.Extensions.Options;
 using OpenTracing;
 using OpenTracing.Util;
 
 namespace Horde.Build.Services
 {
 	using P4 = Perforce.P4;
-	using UserId = ObjectId<IUser>;
 
 	static class PerforceExtensions
 	{
-		static FieldInfo Field = typeof(P4.Changelist).GetField("_baseForm", BindingFlags.Instance | BindingFlags.NonPublic)!;
+		static readonly FieldInfo s_field = typeof(P4.Changelist).GetField("_baseForm", BindingFlags.Instance | BindingFlags.NonPublic)!;
 
-		public static string GetPath(this P4.Changelist Changelist)
+		public static string GetPath(this P4.Changelist changelist)
 		{
-			P4.FormBase? FormBase = Field.GetValue(Changelist) as P4.FormBase;
-			object? PathValue = null;
-			FormBase?.TryGetValue("path", out PathValue);
-			return (PathValue as string) ?? "//...";
+			P4.FormBase? formBase = s_field.GetValue(changelist) as P4.FormBase;
+			object? pathValue = null;
+			formBase?.TryGetValue("path", out pathValue);
+			return (pathValue as string) ?? "//...";
 		}
 	}
 
@@ -48,341 +47,338 @@ namespace Horde.Build.Services
 	{
 		class CachedTicketInfo
 		{
-			public IPerforceServer Server;
-			public string UserName;
-			public P4.Credential Ticket;
+			public IPerforceServer _server;
+			public string _userName;
+			public P4.Credential _ticket;
 
-			public CachedTicketInfo(IPerforceServer Server, string UserName, P4.Credential Ticket)
+			public CachedTicketInfo(IPerforceServer server, string userName, P4.Credential ticket)
 			{
-				this.Server = Server;
-				this.UserName = UserName;
-				this.Ticket = Ticket;
+				_server = server;
+				_userName = userName;
+				_ticket = ticket;
 			}
 		}
 
-		PerforceLoadBalancer LoadBalancer;
-		LazyCachedValue<Task<Globals>> CachedGlobals;
-		ILogger Logger;
+		readonly PerforceLoadBalancer _loadBalancer;
+		readonly LazyCachedValue<Task<Globals>> _cachedGlobals;
+		readonly ILogger _logger;
 
 		/// <summary>
 		/// Object used for controlling access to the access user tickets
 		/// </summary>
-		static object TicketLock = new object();
+		static readonly object s_ticketLock = new object();
 
 		/// <summary>
 		/// Object used for controlling access to the p4 command output
 		/// </summary>
-		static object P4LogLock = new object();
+		static readonly object s_p4LogLock = new object();
 
 		/// <summary>
 		/// Native -> managed debug logging callback
 		/// </summary>
-		P4.P4CallBacks.LogMessageDelegate LogBridgeDelegate;
+		readonly P4.P4CallBacks.LogMessageDelegate _logBridgeDelegate;
 
 		/// <summary>
 		/// The server settings
 		/// </summary>
-		ServerSettings Settings;
-
-		Dictionary<string, Dictionary<string, CachedTicketInfo>> ClusterTickets = new Dictionary<string, Dictionary<string, CachedTicketInfo>>(StringComparer.OrdinalIgnoreCase);
-
-		IUserCollection UserCollection;
-		MemoryCache UserCache = new MemoryCache(new MemoryCacheOptions { SizeLimit = 2000 });
+		readonly ServerSettings _settings;
+		readonly Dictionary<string, Dictionary<string, CachedTicketInfo>> _clusterTickets = new Dictionary<string, Dictionary<string, CachedTicketInfo>>(StringComparer.OrdinalIgnoreCase);
+		readonly IUserCollection _userCollection;
+		readonly MemoryCache _userCache = new MemoryCache(new MemoryCacheOptions { SizeLimit = 2000 });
 
 		/// <summary>
 		/// Constructor
 		/// </summary>
-		public PerforceService(PerforceLoadBalancer LoadBalancer, DatabaseService DatabaseService, IUserCollection UserCollection, IOptions<ServerSettings> Settings, ILogger<PerforceService> Logger)
+		public PerforceService(PerforceLoadBalancer loadBalancer, DatabaseService databaseService, IUserCollection userCollection, IOptions<ServerSettings> settings, ILogger<PerforceService> logger)
 		{
-			this.LoadBalancer = LoadBalancer;
-			this.CachedGlobals = new LazyCachedValue<Task<Globals>>(() => DatabaseService.GetGlobalsAsync(), TimeSpan.FromSeconds(30.0));
-			this.UserCollection = UserCollection;
-			this.Settings = Settings.Value;
-			this.Logger = Logger;
+			_loadBalancer = loadBalancer;
+			_cachedGlobals = new LazyCachedValue<Task<Globals>>(() => databaseService.GetGlobalsAsync(), TimeSpan.FromSeconds(30.0));
+			_userCollection = userCollection;
+			_settings = settings.Value;
+			_logger = logger;
 
-			LogBridgeDelegate = new P4.P4CallBacks.LogMessageDelegate(LogBridgeMessage);
-			P4.P4Debugging.SetBridgeLogFunction(LogBridgeDelegate);
+			_logBridgeDelegate = new P4.P4CallBacks.LogMessageDelegate(LogBridgeMessage);
+			P4.P4Debugging.SetBridgeLogFunction(_logBridgeDelegate);
 
 			P4.LogFile.SetLoggingFunction(LogPerforce);
 		}
 
 		public void Dispose()
 		{
-			UserCache.Dispose();
+			_userCache.Dispose();
 		}
 
-		public async ValueTask<IUser> FindOrAddUserAsync(string ClusterName, string UserName)
+		public async ValueTask<IUser> FindOrAddUserAsync(string clusterName, string userName)
 		{
-			using IScope Scope = GlobalTracer.Instance.BuildSpan("PerforceService.FindOrAddUserAsync").StartActive();
-			Scope.Span.SetTag("ClusterName", ClusterName);
-			Scope.Span.SetTag("UserName", UserName);
+			using IScope scope = GlobalTracer.Instance.BuildSpan("PerforceService.FindOrAddUserAsync").StartActive();
+			scope.Span.SetTag("ClusterName", clusterName);
+			scope.Span.SetTag("UserName", userName);
 			
-			IUser? User;
-			if (!UserCache.TryGetValue((ClusterName, UserName), out User))
+			IUser? user;
+			if (!_userCache.TryGetValue((clusterName, userName), out user))
 			{
-				User = await UserCollection.FindUserByLoginAsync(UserName);
-				if (User == null)
+				user = await _userCollection.FindUserByLoginAsync(userName);
+				if (user == null)
 				{
-					PerforceUserInfo? UserInfo = await GetUserInfoAsync(ClusterName, UserName);
-					User = await UserCollection.FindOrAddUserByLoginAsync(UserName, UserInfo?.FullName, UserInfo?.Email);
+					PerforceUserInfo? userInfo = await GetUserInfoAsync(clusterName, userName);
+					user = await _userCollection.FindOrAddUserByLoginAsync(userName, userInfo?.FullName, userInfo?.Email);
 				}
 
-				using (ICacheEntry Entry = UserCache.CreateEntry((ClusterName, UserName)))
+				using (ICacheEntry entry = _userCache.CreateEntry((clusterName, userName)))
 				{
-					Entry.SetValue(User);
-					Entry.SetSize(1);
-					Entry.SetAbsoluteExpiration(TimeSpan.FromDays(1.0));
+					entry.SetValue(user);
+					entry.SetSize(1);
+					entry.SetAbsoluteExpiration(TimeSpan.FromDays(1.0));
 				}
 			}
-			return User!;
+			return user!;
 		}
 
-		async Task<PerforceCluster> GetClusterAsync(string? ClusterName, string? ServerAndPort = null)
+		async Task<PerforceCluster> GetClusterAsync(string? clusterName, string? serverAndPort = null)
 		{
-			Globals Globals = await CachedGlobals.GetCached();
+			Globals globals = await _cachedGlobals.GetCached();
 
-			PerforceCluster? Cluster = Globals.FindPerforceCluster(ClusterName, ServerAndPort);
-			if (Cluster == null)
+			PerforceCluster? cluster = globals.FindPerforceCluster(clusterName, serverAndPort);
+			if (cluster == null)
 			{
-				throw new Exception($"Unknown Perforce cluster '{ClusterName}'");
+				throw new Exception($"Unknown Perforce cluster '{clusterName}'");
 			}
 
-			return Cluster;
+			return cluster;
 		}
 
-		async Task<IPerforceServer> SelectServer(PerforceCluster Cluster)
+		async Task<IPerforceServer> SelectServer(PerforceCluster cluster)
 		{
-			using IScope Scope = GlobalTracer.Instance.BuildSpan("PerforceService.SelectServer").StartActive();
-			Scope.Span.SetTag("ClusterName", Cluster.Name);
+			using IScope scope = GlobalTracer.Instance.BuildSpan("PerforceService.SelectServer").StartActive();
+			scope.Span.SetTag("ClusterName", cluster.Name);
 			
-			IPerforceServer? Server = await LoadBalancer.SelectServerAsync(Cluster);
-			if (Server == null)
+			IPerforceServer? server = await _loadBalancer.SelectServerAsync(cluster);
+			if (server == null)
 			{
-				throw new Exception($"Unable to select server from '{Cluster.Name}'");
+				throw new Exception($"Unable to select server from '{cluster.Name}'");
 			}
-			return Server;
+			return server;
 		}
 
 		[SuppressMessage("Microsoft.Reliability", "CA2000:DisposeObjectsBeforeLosingScope")]
-		static P4.Repository CreateConnection(IPerforceServer Server, string? UserName, string? Ticket)
+		static P4.Repository CreateConnection(IPerforceServer server, string? userName, string? ticket)
 		{
-			using IScope Scope = GlobalTracer.Instance.BuildSpan("PerforceService.CreateConnection").StartActive();
-			Scope.Span.SetTag("Cluster", Server.Cluster);
-			Scope.Span.SetTag("ServerAndPort", Server.ServerAndPort);
-			Scope.Span.SetTag("UserName", UserName);
+			using IScope scope = GlobalTracer.Instance.BuildSpan("PerforceService.CreateConnection").StartActive();
+			scope.Span.SetTag("Cluster", server.Cluster);
+			scope.Span.SetTag("ServerAndPort", server.ServerAndPort);
+			scope.Span.SetTag("UserName", userName);
 			
-			P4.Repository Repository = new P4.Repository(new P4.Server(new P4.ServerAddress(Server.ServerAndPort)));
+			P4.Repository repository = new P4.Repository(new P4.Server(new P4.ServerAddress(server.ServerAndPort)));
 			try
 			{
-				P4.Connection Connection = Repository.Connection;
-				if (UserName != null)
+				P4.Connection connection = repository.Connection;
+				if (userName != null)
 				{
-					Connection.UserName = UserName;
+					connection.UserName = userName;
 				}
 
-				P4.Options Options = new P4.Options();
-				if (Ticket != null)
+				P4.Options options = new P4.Options();
+				if (ticket != null)
 				{
-					Options["Ticket"] = Ticket;
+					options["Ticket"] = ticket;
 				}
 
 				// connect to the server
-				if (!Connection.Connect(Options))
+				if (!connection.Connect(options))
 				{
 					throw new Exception("Unable to get P4 server connection");
 				}
 			}
 			catch
 			{
-				Repository.Dispose();
+				repository.Dispose();
 				throw;
 			}
-			return Repository;
+			return repository;
 		}
 
-
 		/// <summary>
 		/// 
 		/// </summary>
-		/// <param name="ClusterName"></param>
+		/// <param name="clusterName"></param>
 		/// <returns></returns>
-		public async Task<IPerforceConnection?> GetServiceUserConnection(string? ClusterName)
+		public async Task<IPerforceConnection?> GetServiceUserConnection(string? clusterName)
 		{
-			using IScope Scope = GlobalTracer.Instance.BuildSpan("PerforceService.GetServiceUserConnection").StartActive();
-			Scope.Span.SetTag("ClusterName", ClusterName);
+			using IScope scope = GlobalTracer.Instance.BuildSpan("PerforceService.GetServiceUserConnection").StartActive();
+			scope.Span.SetTag("ClusterName", clusterName);
 
-			PerforceCluster Cluster = await GetClusterAsync(ClusterName);
+			PerforceCluster cluster = await GetClusterAsync(clusterName);
 
-			string? UserName = Cluster.ServiceAccount ?? Environment.UserName;
+			string? userName = cluster.ServiceAccount ?? Environment.UserName;
 
-			string? Password = null;
-			if (Cluster.ServiceAccount != null)
+			string? password = null;
+			if (cluster.ServiceAccount != null)
 			{
-				PerforceCredentials? Credentials = Cluster.Credentials.FirstOrDefault(x => x.UserName.Equals(UserName, StringComparison.OrdinalIgnoreCase));
-				if (Credentials == null)
+				PerforceCredentials? credentials = cluster.Credentials.FirstOrDefault(x => x.UserName.Equals(userName, StringComparison.OrdinalIgnoreCase));
+				if (credentials == null)
 				{
-					throw new Exception($"No credentials defined for {Cluster.ServiceAccount} on {Cluster.Name}");
+					throw new Exception($"No credentials defined for {cluster.ServiceAccount} on {cluster.Name}");
 				}
-				Password = Credentials.Password;
+				password = credentials.Password;
 			}
 
-			IPerforceServer Server = await SelectServer(Cluster);
+			IPerforceServer server = await SelectServer(cluster);
 
-			PerforceSettings Settings = new PerforceSettings(Server.ServerAndPort, UserName);
-			Settings.Password = Password;
-			Settings.AppName = "Horde.Build";
-			Settings.ClientName = "__DOES_NOT_EXIST__";
-			Settings.PreferNativeClient = true;
+			PerforceSettings settings = new PerforceSettings(server.ServerAndPort, userName);
+			settings.Password = password;
+			settings.AppName = "Horde.Build";
+			settings.ClientName = "__DOES_NOT_EXIST__";
+			settings.PreferNativeClient = true;
 
-			return await PerforceConnection.CreateAsync(Settings, Logger);
+			return await PerforceConnection.CreateAsync(settings, _logger);
 		}
 
-		async Task<P4.Repository> GetServiceUserConnection(PerforceCluster Cluster)
+		async Task<P4.Repository> GetServiceUserConnection(PerforceCluster cluster)
 		{
-			using IScope Scope = GlobalTracer.Instance.BuildSpan("PerforceService.GetServiceUserConnectionAsync").StartActive();
-			Scope.Span.SetTag("ClusterName", Cluster.Name);
+			using IScope scope = GlobalTracer.Instance.BuildSpan("PerforceService.GetServiceUserConnectionAsync").StartActive();
+			scope.Span.SetTag("ClusterName", cluster.Name);
 			
-			IPerforceServer Server = await SelectServer(Cluster);
-			return GetServiceUserConnection(Cluster, Server);
+			IPerforceServer server = await SelectServer(cluster);
+			return GetServiceUserConnection(cluster, server);
 		}
 
-		P4.Repository GetServiceUserConnection(PerforceCluster Cluster, IPerforceServer Server)
+		P4.Repository GetServiceUserConnection(PerforceCluster cluster, IPerforceServer server)
 		{
-			using IScope Scope = GlobalTracer.Instance.BuildSpan("PerforceService.GetServiceUserConnection").StartActive();
-			Scope.Span.SetTag("ClusterName", Cluster.Name);
+			using IScope scope = GlobalTracer.Instance.BuildSpan("PerforceService.GetServiceUserConnection").StartActive();
+			scope.Span.SetTag("ClusterName", cluster.Name);
 
-			if (Cluster.Name == PerforceCluster.DefaultName && Settings.P4BridgeServiceUsername != null && Settings.P4BridgeServicePassword != null)
+			if (cluster.Name == PerforceCluster.DefaultName && _settings.P4BridgeServiceUsername != null && _settings.P4BridgeServicePassword != null)
 			{
-				return CreateConnection(Server, Settings.P4BridgeServiceUsername, Settings.P4BridgeServicePassword);
+				return CreateConnection(server, _settings.P4BridgeServiceUsername, _settings.P4BridgeServicePassword);
 			}
 
-			string? UserName = null;
-			string? Password = null;
-			if (Cluster.ServiceAccount != null)
+			string? userName = null;
+			string? password = null;
+			if (cluster.ServiceAccount != null)
 			{
-				PerforceCredentials? Credentials = Cluster.Credentials.FirstOrDefault(x => x.UserName.Equals(Cluster.ServiceAccount, StringComparison.OrdinalIgnoreCase));
-				if (Credentials == null)
+				PerforceCredentials? credentials = cluster.Credentials.FirstOrDefault(x => x.UserName.Equals(cluster.ServiceAccount, StringComparison.OrdinalIgnoreCase));
+				if (credentials == null)
 				{
-					throw new Exception($"No credentials defined for {Cluster.ServiceAccount} on {Cluster.Name}");
+					throw new Exception($"No credentials defined for {cluster.ServiceAccount} on {cluster.Name}");
 				}
-				UserName = Credentials.UserName;
-				Password = Credentials.Password;
+				userName = credentials.UserName;
+				password = credentials.Password;
 			}
-			return CreateConnection(Server, UserName, Password);
+			return CreateConnection(server, userName, password);
 		}
 
-		async Task<CachedTicketInfo> GetImpersonateCredential(PerforceCluster Cluster, string ImpersonateUser)
+		async Task<CachedTicketInfo> GetImpersonateCredential(PerforceCluster cluster, string impersonateUser)
 		{
-			using IScope Scope = GlobalTracer.Instance.BuildSpan("PerforceService.GetImpersonateCredential").StartActive();
-			Scope.Span.SetTag("ClusterName", Cluster.Name);
-			Scope.Span.SetTag("ImpersonateUser", ImpersonateUser);
+			using IScope scope = GlobalTracer.Instance.BuildSpan("PerforceService.GetImpersonateCredential").StartActive();
+			scope.Span.SetTag("ClusterName", cluster.Name);
+			scope.Span.SetTag("ImpersonateUser", impersonateUser);
 			
-			if (!Cluster.CanImpersonate)
+			if (!cluster.CanImpersonate)
 			{
-				throw new Exception($"Service account required to impersonate user {ImpersonateUser}");
+				throw new Exception($"Service account required to impersonate user {impersonateUser}");
 			}
 
-			CachedTicketInfo? TicketInfo = null;
+			CachedTicketInfo? ticketInfo = null;
 
-			Dictionary<string, CachedTicketInfo>? UserTickets;
-			lock (TicketLock)
+			Dictionary<string, CachedTicketInfo>? userTickets;
+			lock (s_ticketLock)
 			{
 				// Check if we have a ticket
-				if (!ClusterTickets.TryGetValue(Cluster.Name, out UserTickets))
+				if (!_clusterTickets.TryGetValue(cluster.Name, out userTickets))
 				{
-					UserTickets = new Dictionary<string, CachedTicketInfo>(StringComparer.OrdinalIgnoreCase);
-					ClusterTickets[Cluster.Name] = UserTickets;
+					userTickets = new Dictionary<string, CachedTicketInfo>(StringComparer.OrdinalIgnoreCase);
+					_clusterTickets[cluster.Name] = userTickets;
 				}
-				if (UserTickets.TryGetValue(ImpersonateUser, out TicketInfo))
+				if (userTickets.TryGetValue(impersonateUser, out ticketInfo))
 				{
 					// if the credential expires within the next 15 minutes, refresh
-					TimeSpan Time = new TimeSpan(0, 15, 0);
-					if (TicketInfo.Ticket.Expires.Subtract(Time) <= DateTime.UtcNow)
+					TimeSpan time = new TimeSpan(0, 15, 0);
+					if (ticketInfo._ticket.Expires.Subtract(time) <= DateTime.UtcNow)
 					{
-						UserTickets.Remove(ImpersonateUser);
-						TicketInfo = null;
+						userTickets.Remove(impersonateUser);
+						ticketInfo = null;
 					}
 				}
 			}
 
-			if (TicketInfo == null)
+			if (ticketInfo == null)
 			{
-				IPerforceServer Server = await SelectServer(Cluster);
+				IPerforceServer server = await SelectServer(cluster);
 
-				P4.Credential? Credential;
-				using (P4.Repository Repository = GetServiceUserConnection(Cluster, Server))
+				P4.Credential? credential;
+				using (P4.Repository repository = GetServiceUserConnection(cluster, server))
 				{
-					Credential = Repository.Connection.Login(null, new P4.LoginCmdOptions(P4.LoginCmdFlags.AllHosts | P4.LoginCmdFlags.DisplayTicket, null), ImpersonateUser);
+					credential = repository.Connection.Login(null, new P4.LoginCmdOptions(P4.LoginCmdFlags.AllHosts | P4.LoginCmdFlags.DisplayTicket, null), impersonateUser);
 				}
-				if (Credential == null)
+				if (credential == null)
 				{
-					throw new Exception($"GetImpersonateCredential - Unable to get impersonation credential for {ImpersonateUser} from {Dns.GetHostName()}");
+					throw new Exception($"GetImpersonateCredential - Unable to get impersonation credential for {impersonateUser} from {Dns.GetHostName()}");
 				}
 
-				TicketInfo = new CachedTicketInfo(Server, ImpersonateUser, Credential);
+				ticketInfo = new CachedTicketInfo(server, impersonateUser, credential);
 
-				lock (TicketLock)
+				lock (s_ticketLock)
 				{
-					UserTickets[ImpersonateUser] = TicketInfo;
+					userTickets[impersonateUser] = ticketInfo;
 				}
 			}
 
-			return TicketInfo;
+			return ticketInfo;
 		}
 
-		async Task<P4.Repository> GetImpersonatedConnection(PerforceCluster Cluster, string ImpersonateUser)
+		async Task<P4.Repository> GetImpersonatedConnection(PerforceCluster cluster, string impersonateUser)
 		{
-			using IScope Scope = GlobalTracer.Instance.BuildSpan("PerforceService.GetImpersonatedConnection").StartActive();
-			Scope.Span.SetTag("ClusterName", Cluster.Name);
-			Scope.Span.SetTag("ImpersonateUser", ImpersonateUser);
+			using IScope scope = GlobalTracer.Instance.BuildSpan("PerforceService.GetImpersonatedConnection").StartActive();
+			scope.Span.SetTag("ClusterName", cluster.Name);
+			scope.Span.SetTag("ImpersonateUser", impersonateUser);
 			
-			CachedTicketInfo TicketInfo = await GetImpersonateCredential(Cluster, ImpersonateUser);
-			return CreateConnection(TicketInfo.Server, TicketInfo.UserName, TicketInfo.Ticket.Ticket);
+			CachedTicketInfo ticketInfo = await GetImpersonateCredential(cluster, impersonateUser);
+			return CreateConnection(ticketInfo._server, ticketInfo._userName, ticketInfo._ticket.Ticket);
 		}
 		
-		async Task<string?> GetPortFromChange(PerforceCluster Cluster, int Change)
+		async Task<string?> GetPortFromChange(PerforceCluster cluster, int change)
 		{
-			using IScope Scope = GlobalTracer.Instance.BuildSpan("PerforceService.GetPortFromChange").StartActive();
-			Scope.Span.SetTag("ClusterName", Cluster.Name);
-			Scope.Span.SetTag("Change", Change);
+			using IScope scope = GlobalTracer.Instance.BuildSpan("PerforceService.GetPortFromChange").StartActive();
+			scope.Span.SetTag("ClusterName", cluster.Name);
+			scope.Span.SetTag("Change", change);
 
-			string? ClientPort = null;
+			string? clientPort = null;
 
-			using (P4.Repository Repository = await GetConnection(Cluster, NoClient: true))
+			using (P4.Repository repository = await GetConnection(cluster, noClient: true))
 			{				
-				P4.Changelist Changelist = Repository.GetChangelist(Change);
-				if (Changelist == null || string.IsNullOrEmpty(Changelist.ClientId))
+				P4.Changelist changelist = repository.GetChangelist(change);
+				if (changelist == null || String.IsNullOrEmpty(changelist.ClientId))
 				{
-					throw new Exception($"GetPortFromChange - Unable to get changelist for {Change}");
+					throw new Exception($"GetPortFromChange - Unable to get changelist for {change}");
 				}
 
-				P4.Client Client = Repository.GetClient(Changelist.ClientId);
+				P4.Client client = repository.GetClient(changelist.ClientId);
 
-				if (Client == null)
+				if (client == null)
 				{
-					throw new Exception($"GetPortFromChange - Unable to get client for {Change}");
+					throw new Exception($"GetPortFromChange - Unable to get client for {change}");
 				}
 
 				// Check whether not restricted to a specific server
-				if (string.IsNullOrEmpty(Client.ServerID))
+				if (String.IsNullOrEmpty(client.ServerID))
 				{
 					return null;
 				}
 
-				List<string> Args = new List<string> { "-o", Client.ServerID };
+				List<string> args = new List<string> { "-o", client.ServerID };
 
-				using (P4.P4Command Command = new P4.P4Command(Repository, "server", true, Args.ToArray()))
+				using (P4.P4Command command = new P4.P4Command(repository, "server", true, args.ToArray()))
 				{
-					P4.P4CommandResult Result = Command.Run();
+					P4.P4CommandResult result = command.Run();
 
-					if (!Result.Success || Result.TaggedOutput == null || Result.TaggedOutput.Count == 0)
+					if (!result.Success || result.TaggedOutput == null || result.TaggedOutput.Count == 0)
 					{
-						throw new Exception($"Unable to get tagged output for {Client.ServerID}");
+						throw new Exception($"Unable to get tagged output for {client.ServerID}");
 					}
 
-					foreach (P4.TaggedObject TaggedObject in Result.TaggedOutput)
+					foreach (P4.TaggedObject taggedObject in result.TaggedOutput)
 					{
-						if (TaggedObject.TryGetValue("Address", out ClientPort))
+						if (taggedObject.TryGetValue("Address", out clientPort))
 						{
 							break;
 						}
@@ -390,352 +386,350 @@ namespace Horde.Build.Services
 				}
 			}
 
-			if (string.IsNullOrEmpty(ClientPort))
+			if (String.IsNullOrEmpty(clientPort))
 			{
-				throw new Exception($"GetPortFromChange - Unable to get client port for {Change}");
+				throw new Exception($"GetPortFromChange - Unable to get client port for {change}");
 			}
 
-			return ClientPort;
+			return clientPort;
 		}
 
 		[SuppressMessage("Microsoft.Reliability", "CA2000:DisposeObjectsBeforeLosingScope")]
-		async Task<P4.Repository> GetConnection(PerforceCluster Cluster, string? Stream = null, string? Username = null, bool ReadOnly = true, bool CreateChange = false, int? ClientFromChange = null, bool UseClientFromChange = false, int? UsePortFromChange = null, string? ClientId = null, bool NoClient = false)
+		async Task<P4.Repository> GetConnection(PerforceCluster cluster, string? stream = null, string? username = null, bool readOnly = true, bool createChange = false, int? clientFromChange = null, bool useClientFromChange = false, int? usePortFromChange = null, string? clientId = null, bool noClient = false)
 		{
-			using IScope Scope = GlobalTracer.Instance.BuildSpan("PerforceService.GetConnection").StartActive();
-			Scope.Span.SetTag("ClusterName", Cluster.Name);
-			Scope.Span.SetTag("Stream", Stream);
-			Scope.Span.SetTag("Username", Username);
-			Scope.Span.SetTag("NoClient", NoClient);
-			Scope.Span.SetTag("ClientId", ClientId);
-
-			if (UsePortFromChange.HasValue)
+			using IScope scope = GlobalTracer.Instance.BuildSpan("PerforceService.GetConnection").StartActive();
+			scope.Span.SetTag("ClusterName", cluster.Name);
+			scope.Span.SetTag("Stream", stream);
+			scope.Span.SetTag("Username", username);
+			scope.Span.SetTag("NoClient", noClient);
+			scope.Span.SetTag("ClientId", clientId);
+
+			if (usePortFromChange.HasValue)
 			{
 				try
 				{
-					string? ClientPort = await GetPortFromChange(Cluster, UsePortFromChange.Value);
-					if (ClientPort != null)
+					string? clientPort = await GetPortFromChange(cluster, usePortFromChange.Value);
+					if (clientPort != null)
 					{
-						PerforceCluster ClientCluster = await GetClusterAsync(null, ClientPort);
-						if (!string.Equals(ClientCluster.Name, Cluster.Name, StringComparison.OrdinalIgnoreCase))
+						PerforceCluster clientCluster = await GetClusterAsync(null, clientPort);
+						if (!String.Equals(clientCluster.Name, cluster.Name, StringComparison.OrdinalIgnoreCase))
 						{
-							Logger.LogInformation("Perforce: Overriding Cluster {Cluster} with Client Cluster {ClientCluster} for server restricted change {Change}", Cluster.Name, ClientCluster.Name, UsePortFromChange.Value);
-							Cluster = ClientCluster;
+							_logger.LogInformation("Perforce: Overriding Cluster {Cluster} with Client Cluster {ClientCluster} for server restricted change {Change}", cluster.Name, clientCluster.Name, usePortFromChange.Value);
+							cluster = clientCluster;
 						}
 					}
 					else
 					{
-						Logger.LogInformation("Change {UsePortFromChange} is not restricted to a server", UsePortFromChange.Value);
+						_logger.LogInformation("Change {UsePortFromChange} is not restricted to a server", usePortFromChange.Value);
 					}
 				}
-				catch (Exception Ex)
+				catch (Exception ex)
 				{
-					Logger.LogError(Ex, "Unable to get client port for changelist {Change} using cluster {ClusterName}", UsePortFromChange, Cluster.Name);
+					_logger.LogError(ex, "Unable to get client port for changelist {Change} using cluster {ClusterName}", usePortFromChange, cluster.Name);
 				}
 			}
 
-			P4.Repository Repository;
-			if (Username == null || !Cluster.CanImpersonate || Username.Equals(Cluster.ServiceAccount, StringComparison.OrdinalIgnoreCase))
+			P4.Repository repository;
+			if (username == null || !cluster.CanImpersonate || username.Equals(cluster.ServiceAccount, StringComparison.OrdinalIgnoreCase))
 			{
-				Repository = await GetServiceUserConnection(Cluster);
+				repository = await GetServiceUserConnection(cluster);
 			}
 			else
 			{
-				Repository = await GetImpersonatedConnection(Cluster, Username);
+				repository = await GetImpersonatedConnection(cluster, username);
 			}
 
-			if (!NoClient)
+			if (!noClient)
 			{
 				try
 				{
-					if (ClientId != null)
+					if (clientId != null)
 					{
-						Repository.Connection.SetClient(ClientId);
+						repository.Connection.SetClient(clientId);
 					}
 					else
 					{
-						P4.Client Client = GetOrCreateClient(Cluster.ServiceAccount, Repository, Stream, Username, ReadOnly, CreateChange, ClientFromChange, UseClientFromChange);
-						Repository.Connection.SetClient(Client.Name);
+						P4.Client client = GetOrCreateClient(cluster.ServiceAccount, repository, stream, username, readOnly, createChange, clientFromChange, useClientFromChange);
+						repository.Connection.SetClient(client.Name);
 
 					}
 
-					if (UseClientFromChange)
+					if (useClientFromChange)
 					{
-						string? ClientHost = Repository.Connection.Client.Host;
-						if (!string.IsNullOrEmpty(ClientHost))
+						string? clientHost = repository.Connection.Client.Host;
+						if (!String.IsNullOrEmpty(clientHost))
 						{
-							Repository.Connection.getP4Server().SetConnectionHost(ClientHost);
+							repository.Connection.getP4Server().SetConnectionHost(clientHost);
 						}
 					}					
-
 				}
 				catch
 				{
-					Repository.Dispose();
+					repository.Dispose();
 					throw;
 				}
 			}
 
-			if (Repository == null)
+			if (repository == null)
 			{
-				throw new Exception($"Unable to get connection for Stream:{Stream} Username:{Username} ReadOnly:{ReadOnly} CreateChange:{CreateChange} ClientFromChange:{ClientFromChange} UseClientFromChange:{UseClientFromChange} UsePortFromChange:{UsePortFromChange}");
+				throw new Exception($"Unable to get connection for Stream:{stream} Username:{username} ReadOnly:{readOnly} CreateChange:{createChange} ClientFromChange:{clientFromChange} UseClientFromChange:{useClientFromChange} UsePortFromChange:{usePortFromChange}");
 			}
 
-			Repository.Connection.CommandEcho += LogPerforceCommand;
+			repository.Connection.CommandEcho += LogPerforceCommand;
 
-			return Repository;
+			return repository;
 
 		}
 
-		static string GetClientName(string? ServiceUserName, string Stream, bool ReadOnly, bool CreateChange, string? Username = null)
+		static string GetClientName(string? serviceUserName, string stream, bool readOnly, bool createChange, string? username = null)
 		{
-			string ClientName = $"horde-p4bridge-{Dns.GetHostName()}-{ServiceUserName ?? "default"}-{Stream.Replace("/", "+", StringComparison.OrdinalIgnoreCase)}";
+			string clientName = $"horde-p4bridge-{Dns.GetHostName()}-{serviceUserName ?? "default"}-{stream.Replace("/", "+", StringComparison.OrdinalIgnoreCase)}";
 
-			if (!ReadOnly)
+			if (!readOnly)
 			{
-				ClientName += "-write";
+				clientName += "-write";
 			}
 
-			if (CreateChange)
+			if (createChange)
 			{
-				ClientName += "-create";
+				clientName += "-create";
 			}
 
-			if (!string.IsNullOrEmpty(Username))
+			if (!String.IsNullOrEmpty(username))
 			{
-				ClientName += "-" + Username;
+				clientName += "-" + username;
 			}
 
-			return ClientName;
+			return clientName;
 
 		}
 
-		static void ResetClient(P4.Repository Repository)
+		static void ResetClient(P4.Repository repository)
 		{
-			Repository.Connection.Client.RevertFiles(new P4.Options(), new P4.DepotPath("//..."));
+			repository.Connection.Client.RevertFiles(new P4.Options(), new P4.DepotPath("//..."));
 
-			IList<P4.Changelist> Changes = Repository.GetChangelists(new P4.ChangesCmdOptions(P4.ChangesCmdFlags.None, Repository.Connection.Client.Name, 100, P4.ChangeListStatus.Pending, null));
-			if (Changes != null)
+			IList<P4.Changelist> changes = repository.GetChangelists(new P4.ChangesCmdOptions(P4.ChangesCmdFlags.None, repository.Connection.Client.Name, 100, P4.ChangeListStatus.Pending, null));
+			if (changes != null)
 			{
-				foreach (P4.Changelist Change in Changes)
+				foreach (P4.Changelist change in changes)
 				{
-					Repository.DeleteChangelist(Change, new P4.Options());
+					repository.DeleteChangelist(change, new P4.Options());
 				}
 			}
 		}
 
-		static P4.Client GetOrCreateClient(string? ServiceUserName, P4.Repository Repository, string? Stream, string? Username = null, bool ReadOnly = true, bool CreateChange = false, int? ClientFromChange = null, bool UseClientFromChange = false)
+		static P4.Client GetOrCreateClient(string? serviceUserName, P4.Repository repository, string? stream, string? username = null, bool readOnly = true, bool createChange = false, int? clientFromChange = null, bool useClientFromChange = false)
 		{
-			using IScope Scope = GlobalTracer.Instance.BuildSpan("PerforceService.GetOrCreateClient").StartActive();
-			Scope.Span.SetTag("ServiceUserName", ServiceUserName);
-			Scope.Span.SetTag("Stream", Stream);
-			Scope.Span.SetTag("Username", Username);
+			using IScope scope = GlobalTracer.Instance.BuildSpan("PerforceService.GetOrCreateClient").StartActive();
+			scope.Span.SetTag("ServiceUserName", serviceUserName);
+			scope.Span.SetTag("Stream", stream);
+			scope.Span.SetTag("Username", username);
 			
-			P4.Client? Client = null;
-			P4.Changelist? Changelist = null;
+			P4.Client? client = null;
+			P4.Changelist? changelist = null;
 
-			string? ClientHost = null;
+			string? clientHost = null;
 
-			if (ClientFromChange.HasValue)
+			if (clientFromChange.HasValue)
 			{
-				Changelist = Repository.GetChangelist(ClientFromChange.Value);
+				changelist = repository.GetChangelist(clientFromChange.Value);
 
-				if (Changelist == null)
+				if (changelist == null)
 				{
-					throw new Exception($"Unable to get changelist for client {ClientFromChange}");
+					throw new Exception($"Unable to get changelist for client {clientFromChange}");
 				}
 
-				Client = Repository.GetClient(Changelist.ClientId);
+				client = repository.GetClient(changelist.ClientId);
 
-				if (Client == null)
+				if (client == null)
 				{
-					throw new Exception($"Unable to get client for id {Changelist.ClientId}");
+					throw new Exception($"Unable to get client for id {changelist.ClientId}");
 				}
 
-				Stream = Client.Stream;
-				Username = Client.OwnerName;
-				ClientHost = Client.Host;
+				stream = client.Stream;
+				username = client.OwnerName;
+				clientHost = client.Host;
 
 			}
 
-			if (!UseClientFromChange)
+			if (!useClientFromChange)
 			{
-				if (Stream == null)
+				if (stream == null)
 				{
 					throw new Exception("Stream required for client");
 
 				}
-				string ClientName = GetClientName(ServiceUserName, Stream, ReadOnly, CreateChange, Username);
+				string clientName = GetClientName(serviceUserName, stream, readOnly, createChange, username);
 
-				IList<P4.Client>? Clients = Repository.GetClients(new P4.ClientsCmdOptions(P4.ClientsCmdFlags.None, Username, ClientName, 1, Stream));
+				IList<P4.Client>? clients = repository.GetClients(new P4.ClientsCmdOptions(P4.ClientsCmdFlags.None, username, clientName, 1, stream));
 
-				if (Clients != null && Clients.Count == 1)
+				if (clients != null && clients.Count == 1)
 				{
-					Client = Clients[0];
-					if (!string.IsNullOrEmpty(Client.Host))
+					client = clients[0];
+					if (!String.IsNullOrEmpty(client.Host))
 					{
-						Client.Host = string.Empty;
-						Repository.UpdateClient(Client);
+						client.Host = String.Empty;
+						repository.UpdateClient(client);
 					}
 				}
 				else
 				{
 
-					P4.Client NewClient = new P4.Client();
+					P4.Client newClient = new P4.Client();
 
-					NewClient.Name = ClientName;
-					NewClient.Root = RuntimeInformation.IsOSPlatform(OSPlatform.Linux) ? $"/tmp/{ClientName}/" : $"{Path.GetTempPath()}{ClientName}\\";
-					NewClient.Stream = Stream;
-					if (ReadOnly)
+					newClient.Name = clientName;
+					newClient.Root = RuntimeInformation.IsOSPlatform(OSPlatform.Linux) ? $"/tmp/{clientName}/" : $"{Path.GetTempPath()}{clientName}\\";
+					newClient.Stream = stream;
+					if (readOnly)
 					{
-						NewClient.ClientType = P4.ClientType.@readonly;
+						newClient.ClientType = P4.ClientType.@readonly;
 					}
 
-					NewClient.OwnerName = Username ?? ServiceUserName;
+					newClient.OwnerName = username ?? serviceUserName;
 
-					Client = Repository.CreateClient(NewClient);
+					client = repository.CreateClient(newClient);
 
-					if (Client == null)
+					if (client == null)
 					{
-						throw new Exception($"Unable to create client for ${Stream} : {Username ?? ServiceUserName}");
+						throw new Exception($"Unable to create client for ${stream} : {username ?? serviceUserName}");
 					}
 				}
 			}
 
-			if (Client == null)
+			if (client == null)
 			{
-				throw new Exception($"Unable to create client for Stream:{Stream} Username:{Username} ReadOnly:{ReadOnly} CreateChange:{CreateChange} ClientFromChange:{ClientFromChange} UseClientFromChange:{UseClientFromChange}");
+				throw new Exception($"Unable to create client for Stream:{stream} Username:{username} ReadOnly:{readOnly} CreateChange:{createChange} ClientFromChange:{clientFromChange} UseClientFromChange:{useClientFromChange}");
 			}
 
-			return Client;
+			return client;
 
 		}
 
 		/// <summary>
 		/// Logs perforce commands 
 		/// </summary>
-		/// <param name="Log">The p4 command log info</param>
-		static void LogPerforceCommand(string Log)
+		/// <param name="log">The p4 command log info</param>
+		static void LogPerforceCommand(string log)
 		{
-			lock (P4LogLock)
+			lock (s_p4LogLock)
 			{
-				Serilog.Log.Information("Perforce: {Message}", Log);
+				Serilog.Log.Information("Perforce: {Message}", log);
 			}
-
 		}
 
 		/// <summary>
 		/// Perforce logging funcion
 		/// </summary>
-		/// <param name="LogLevel">The level whether error, warning, message, or debug</param>
-		/// <param name="Source">The log source</param>
-		/// <param name="Message">The log message</param>
-		static void LogPerforce(int LogLevel, String Source, String Message)
+		/// <param name="logLevel">The level whether error, warning, message, or debug</param>
+		/// <param name="source">The log source</param>
+		/// <param name="message">The log message</param>
+		static void LogPerforce(int logLevel, string source, string message)
 		{
-			lock (P4LogLock)
+			lock (s_p4LogLock)
 			{
-				switch (LogLevel)
+				switch (logLevel)
 				{
 					case 0:
 					case 1:
-						Serilog.Log.Error("Perforce (Error): {Message} {Source} : ", Message, Source);
+						Serilog.Log.Error("Perforce (Error): {Message} {Source} : ", message, source);
 						break;
 					case 2:
-						Serilog.Log.Warning("Perforce (Warning): {Message} {Source} : ", Message, Source);
+						Serilog.Log.Warning("Perforce (Warning): {Message} {Source} : ", message, source);
 						break;
 					case 3:
-						Serilog.Log.Information("Perforce: {Message} {Source} : ", Message, Source);
+						Serilog.Log.Information("Perforce: {Message} {Source} : ", message, source);
 						break;
 					default:
-						Serilog.Log.Debug("Perforce (Debug): {Message} {Source} : ", Message, Source);
+						Serilog.Log.Debug("Perforce (Debug): {Message} {Source} : ", message, source);
 						break;
-				};
+				}
 			}
 		}
 
-		static void LogBridgeMessage(int LogLevel, String Filename, int Line, String Message)
+		static void LogBridgeMessage(int logLevel, string filename, int line, string message)
 		{
 
 			// Note, we do not get log level 4 unless it is defined in native code as it is very, very spammy (P4BridgeServer.cpp)
 
 			// remove the full path to the source, keep just the file name
-			String fileName = Path.GetFileName(Filename);
+			string fileName = Path.GetFileName(filename);
 
-			string category = String.Format(CultureInfo.CurrentCulture, "P4Bridge({0}:{1})", fileName, Line);
+			string category = String.Format(CultureInfo.CurrentCulture, "P4Bridge({0}:{1})", fileName, line);
 
-			P4.LogFile.LogMessage(LogLevel, category, Message);
+			P4.LogFile.LogMessage(logLevel, category, message);
 		}
 
 		class StreamView : IStreamView
 		{
-			P4.P4MapApi MapApi;
+			readonly P4.P4MapApi _mapApi;
 
-			public StreamView(P4.P4MapApi MapApi)
+			public StreamView(P4.P4MapApi mapApi)
 			{
-				this.MapApi = MapApi;
+				_mapApi = mapApi;
 			}
 
 			public void Dispose()
 			{
-				MapApi.Dispose();
+				_mapApi.Dispose();
 			}
 
-			public bool TryGetStreamPath(string DepotPath, out string StreamPath)
+			public bool TryGetStreamPath(string depotPath, out string streamPath)
 			{
-				StreamPath = MapApi.Translate(DepotPath, P4.P4MapApi.Direction.LeftRight);
-				return StreamPath != null;
+				streamPath = _mapApi.Translate(depotPath, P4.P4MapApi.Direction.LeftRight);
+				return streamPath != null;
 			}
 		}
 
 		/// <inheritdoc/>
-		public async Task<IStreamView> GetStreamViewAsync(string ClusterName, string StreamName)
+		public async Task<IStreamView> GetStreamViewAsync(string clusterName, string streamName)
 		{
-			using IScope Scope = GlobalTracer.Instance.BuildSpan("PerforceService.GetStreamViewAsync").StartActive();
-			Scope.Span.SetTag("ClusterName", ClusterName);
-			Scope.Span.SetTag("StreamName", StreamName);
+			using IScope scope = GlobalTracer.Instance.BuildSpan("PerforceService.GetStreamViewAsync").StartActive();
+			scope.Span.SetTag("ClusterName", clusterName);
+			scope.Span.SetTag("StreamName", streamName);
 			
-			PerforceCluster Cluster = await GetClusterAsync(ClusterName);
-			using (P4.Repository Repository = await GetConnection(Cluster, NoClient: true))
+			PerforceCluster cluster = await GetClusterAsync(clusterName);
+			using (P4.Repository repository = await GetConnection(cluster, noClient: true))
 			{
-				P4.Stream Stream = Repository.GetStream(StreamName, new P4.StreamCmdOptions(P4.StreamCmdFlags.View, null, null));
+				P4.Stream stream = repository.GetStream(streamName, new P4.StreamCmdOptions(P4.StreamCmdFlags.View, null, null));
 
-				P4.P4MapApi? MapApi = null;
+				P4.P4MapApi? mapApi = null;
 				try
 				{
-					MapApi = new P4.P4MapApi(Repository.Server.Metadata.UnicodeEnabled);
-					foreach (P4.MapEntry Entry in Stream.View)
+					mapApi = new P4.P4MapApi(repository.Server.Metadata.UnicodeEnabled);
+					foreach (P4.MapEntry entry in stream.View)
 					{
-						P4.P4MapApi.Type MapType;
-						switch (Entry.Type)
+						P4.P4MapApi.Type mapType;
+						switch (entry.Type)
 						{
 							case P4.MapType.Include:
-								MapType = P4.P4MapApi.Type.Include;
+								mapType = P4.P4MapApi.Type.Include;
 								break;
 							case P4.MapType.Exclude:
-								MapType = P4.P4MapApi.Type.Exclude;
+								mapType = P4.P4MapApi.Type.Exclude;
 								break;
 							case P4.MapType.Overlay:
-								MapType = P4.P4MapApi.Type.Overlay;
+								mapType = P4.P4MapApi.Type.Overlay;
 								break;
 							default:
-								throw new Exception($"Invalid map type: {Entry.Type}");
+								throw new Exception($"Invalid map type: {entry.Type}");
 						}
-						MapApi.Insert(Entry.Left.ToString(), Entry.Right.ToString(), MapType);
+						mapApi.Insert(entry.Left.ToString(), entry.Right.ToString(), mapType);
 					}
 
-					StreamView View = new StreamView(MapApi);
-					MapApi = null;
-					return View;
+					StreamView view = new StreamView(mapApi);
+					mapApi = null;
+					return view;
 				}
 				finally
 				{
-					MapApi?.Dispose();
+					mapApi?.Dispose();
 				}
 			}
 		}
 
-		static int GetSyncRevision(string Path, P4.FileAction HeadAction, int HeadRev)
+		static int GetSyncRevision(string path, P4.FileAction headAction, int headRev)
 		{
-			switch (HeadAction)
+			switch (headAction)
 			{
 				case P4.FileAction.None:
 				case P4.FileAction.Add:
@@ -743,202 +737,202 @@ namespace Horde.Build.Services
 				case P4.FileAction.MoveAdd:
 				case P4.FileAction.Edit:
 				case P4.FileAction.Integrate:
-					return HeadRev;
+					return headRev;
 				case P4.FileAction.Delete:
 				case P4.FileAction.MoveDelete:
 				case P4.FileAction.Purge:
 					return -1;
 				default:
-					throw new Exception($"Unrecognized P4 file change type '{HeadAction}' for file {Path}#{HeadRev}");
+					throw new Exception($"Unrecognized P4 file change type '{headAction}' for file {path}#{headRev}");
 			}
 		}
 
 		/// <summary>
 		/// Creates a <see cref="ChangeFile"/> from a <see cref="P4.FileMetaData"/>
 		/// </summary>
-		/// <param name="RelativePath"></param>
-		/// <param name="MetaData"></param>
+		/// <param name="relativePath"></param>
+		/// <param name="metaData"></param>
 		/// <returns></returns>
-		static ChangeFile CreateChangeFile(string RelativePath, P4.FileMetaData MetaData)
+		static ChangeFile CreateChangeFile(string relativePath, P4.FileMetaData metaData)
 		{
-			using IScope Scope = GlobalTracer.Instance.BuildSpan("PerforceService.CreateChangeFile (FileMetaData)").StartActive();
-			Scope.Span.SetTag("RelativePath", RelativePath);
+			using IScope scope = GlobalTracer.Instance.BuildSpan("PerforceService.CreateChangeFile (FileMetaData)").StartActive();
+			scope.Span.SetTag("RelativePath", relativePath);
 			
-			int Revision = GetSyncRevision(MetaData.DepotPath.Path, MetaData.HeadAction, MetaData.HeadRev);
-			Md5Hash? Digest = String.IsNullOrEmpty(MetaData.Digest) ? (Md5Hash?)null : Md5Hash.Parse(MetaData.Digest);
-			return new ChangeFile(RelativePath, MetaData.DepotPath.Path, Revision, MetaData.FileSize, Digest, (MetaData.Type ?? MetaData.HeadType).ToString());
+			int revision = GetSyncRevision(metaData.DepotPath.Path, metaData.HeadAction, metaData.HeadRev);
+			Md5Hash? digest = String.IsNullOrEmpty(metaData.Digest) ? (Md5Hash?)null : Md5Hash.Parse(metaData.Digest);
+			return new ChangeFile(relativePath, metaData.DepotPath.Path, revision, metaData.FileSize, digest, (metaData.Type ?? metaData.HeadType).ToString());
 		}
 
 		/// <summary>
 		/// Creates a <see cref="ChangeFile"/> from a <see cref="P4.FileMetaData"/>
 		/// </summary>
-		/// <param name="RelativePath"></param>
-		/// <param name="MetaData"></param>
+		/// <param name="relativePath"></param>
+		/// <param name="metaData"></param>
 		/// <returns></returns>
-		static ChangeFile CreateChangeFile(string RelativePath, P4.ShelvedFile MetaData)
+		static ChangeFile CreateChangeFile(string relativePath, P4.ShelvedFile metaData)
 		{
-			using IScope Scope = GlobalTracer.Instance.BuildSpan("PerforceService.CreateChangeFile (ShelvedFile)").StartActive();
-			Scope.Span.SetTag("RelativePath", RelativePath);
+			using IScope scope = GlobalTracer.Instance.BuildSpan("PerforceService.CreateChangeFile (ShelvedFile)").StartActive();
+			scope.Span.SetTag("RelativePath", relativePath);
 			
-			int Revision = GetSyncRevision(MetaData.Path.Path, MetaData.Action, MetaData.Revision);
-			Md5Hash? Digest = String.IsNullOrEmpty(MetaData.Digest) ? (Md5Hash?)null : Md5Hash.Parse(MetaData.Digest);
-			return new ChangeFile(RelativePath, MetaData.Path.Path, Revision, MetaData.Size, Digest, MetaData.Type.ToString());
+			int revision = GetSyncRevision(metaData.Path.Path, metaData.Action, metaData.Revision);
+			Md5Hash? digest = String.IsNullOrEmpty(metaData.Digest) ? (Md5Hash?)null : Md5Hash.Parse(metaData.Digest);
+			return new ChangeFile(relativePath, metaData.Path.Path, revision, metaData.Size, digest, metaData.Type.ToString());
 		}
 
 		/// <inheritdoc/>
-		public async Task<List<ChangeSummary>> GetChangesAsync(string ClusterName, int? MinChange, int? MaxChange, int MaxResults)
+		public async Task<List<ChangeSummary>> GetChangesAsync(string clusterName, int? minChange, int? maxChange, int maxResults)
 		{
-			using IScope Scope = GlobalTracer.Instance.BuildSpan("PerforceService.GetChangesAsync").StartActive();
-			Scope.Span.SetTag("ClusterName", ClusterName);
-			Scope.Span.SetTag("MinChange", MinChange ?? -1);
-			Scope.Span.SetTag("MaxResults", MaxResults);
+			using IScope scope = GlobalTracer.Instance.BuildSpan("PerforceService.GetChangesAsync").StartActive();
+			scope.Span.SetTag("ClusterName", clusterName);
+			scope.Span.SetTag("MinChange", minChange ?? -1);
+			scope.Span.SetTag("MaxResults", maxResults);
 			
-			PerforceCluster Cluster = await GetClusterAsync(ClusterName);
-			using (P4.Repository Repository = await GetConnection(Cluster, NoClient: true))
+			PerforceCluster cluster = await GetClusterAsync(clusterName);
+			using (P4.Repository repository = await GetConnection(cluster, noClient: true))
 			{
-				P4.ChangesCmdOptions Options = new P4.ChangesCmdOptions(P4.ChangesCmdFlags.IncludeTime | P4.ChangesCmdFlags.FullDescription, null, MaxResults, P4.ChangeListStatus.Submitted, null);
+				P4.ChangesCmdOptions options = new P4.ChangesCmdOptions(P4.ChangesCmdFlags.IncludeTime | P4.ChangesCmdFlags.FullDescription, null, maxResults, P4.ChangeListStatus.Submitted, null);
 
-				IList<P4.Changelist> Changelists = Repository.GetChangelists(Options, new P4.FileSpec(new P4.DepotPath(GetFilter("//...", MinChange, MaxChange)), null, null, null));
+				IList<P4.Changelist> changelists = repository.GetChangelists(options, new P4.FileSpec(new P4.DepotPath(GetFilter("//...", minChange, maxChange)), null, null, null));
 
-				List<ChangeSummary> Changes = new List<ChangeSummary>();
-				if (Changelists != null)
+				List<ChangeSummary> changes = new List<ChangeSummary>();
+				if (changelists != null)
 				{
-					foreach (P4.Changelist Changelist in Changelists)
+					foreach (P4.Changelist changelist in changelists)
 					{
-						IUser User = await FindOrAddUserAsync(ClusterName, Changelist.OwnerName);
-						Changes.Add(new ChangeSummary(Changelist.Id, User, Changelist.GetPath(), Changelist.Description));
+						IUser user = await FindOrAddUserAsync(clusterName, changelist.OwnerName);
+						changes.Add(new ChangeSummary(changelist.Id, user, changelist.GetPath(), changelist.Description));
 					}
 				}
-				return Changes;
+				return changes;
 			}
 		}
 
 		/// <inheritdoc/>
-		public async Task<List<ChangeSummary>> GetChangesAsync(string ClusterName, string StreamName, int? MinChange, int? MaxChange, int Results, string? ImpersonateUser)
+		public async Task<List<ChangeSummary>> GetChangesAsync(string clusterName, string streamName, int? minChange, int? maxChange, int results, string? impersonateUser)
 		{
-			using IScope Scope = GlobalTracer.Instance.BuildSpan("PerforceService.GetChangesAsync").StartActive();
-			Scope.Span.SetTag("ClusterName", ClusterName);
-			Scope.Span.SetTag("MinChange", MinChange ?? -1);
-			Scope.Span.SetTag("MaxChange", MaxChange ?? -1);
-			Scope.Span.SetTag("Results", Results);
-			Scope.Span.SetTag("ImpersonateUser", ImpersonateUser);
+			using IScope scope = GlobalTracer.Instance.BuildSpan("PerforceService.GetChangesAsync").StartActive();
+			scope.Span.SetTag("ClusterName", clusterName);
+			scope.Span.SetTag("MinChange", minChange ?? -1);
+			scope.Span.SetTag("MaxChange", maxChange ?? -1);
+			scope.Span.SetTag("Results", results);
+			scope.Span.SetTag("ImpersonateUser", impersonateUser);
 			
-			PerforceCluster Cluster = await GetClusterAsync(ClusterName);
-			using (P4.Repository Repository = await GetConnection(Cluster, StreamName, ImpersonateUser))
+			PerforceCluster cluster = await GetClusterAsync(clusterName);
+			using (P4.Repository repository = await GetConnection(cluster, streamName, impersonateUser))
 			{
 
-				P4.ChangesCmdOptions Options = new P4.ChangesCmdOptions(P4.ChangesCmdFlags.IncludeTime | P4.ChangesCmdFlags.FullDescription, null, Results, P4.ChangeListStatus.Submitted, null);
+				P4.ChangesCmdOptions options = new P4.ChangesCmdOptions(P4.ChangesCmdFlags.IncludeTime | P4.ChangesCmdFlags.FullDescription, null, results, P4.ChangeListStatus.Submitted, null);
 
-				string Filter = GetFilter($"//{Repository.Connection.Client.Name}/...", MinChange, MaxChange);
+				string filter = GetFilter($"//{repository.Connection.Client.Name}/...", minChange, maxChange);
 
-				P4.FileSpec FileSpec = new P4.FileSpec(new P4.DepotPath(Filter), null, null, null);
+				P4.FileSpec fileSpec = new P4.FileSpec(new P4.DepotPath(filter), null, null, null);
 
-				IList<P4.Changelist> Changelists = Repository.GetChangelists(Options, FileSpec);
+				IList<P4.Changelist> changelists = repository.GetChangelists(options, fileSpec);
 
-				List<ChangeSummary> Changes = new List<ChangeSummary>();
+				List<ChangeSummary> changes = new List<ChangeSummary>();
 
-				if (Changelists != null)
+				if (changelists != null)
 				{
-					foreach (P4.Changelist Changelist in Changelists)
+					foreach (P4.Changelist changelist in changelists)
 					{
-						IUser User = await FindOrAddUserAsync(ClusterName, Changelist.OwnerName);
-						Changes.Add(new ChangeSummary(Changelist.Id, User, Changelist.GetPath(), Changelist.Description));
+						IUser user = await FindOrAddUserAsync(clusterName, changelist.OwnerName);
+						changes.Add(new ChangeSummary(changelist.Id, user, changelist.GetPath(), changelist.Description));
 					}
 				}
 
-				return Changes;
+				return changes;
 			}
 		}
 
 		/// <inheritdoc/>
-		public async Task<ChangeDetails> GetChangeDetailsAsync(string ClusterName, string StreamName, int ChangeNumber)
+		public async Task<ChangeDetails> GetChangeDetailsAsync(string clusterName, string streamName, int changeNumber)
 		{
-			using IScope Scope = GlobalTracer.Instance.BuildSpan("PerforceService.GetChangeDetailsAsync").StartActive();
-			Scope.Span.SetTag("ClusterName", ClusterName);
-			Scope.Span.SetTag("StreamName", StreamName);
-			Scope.Span.SetTag("ChangeNumber", ChangeNumber);
+			using IScope scope = GlobalTracer.Instance.BuildSpan("PerforceService.GetChangeDetailsAsync").StartActive();
+			scope.Span.SetTag("ClusterName", clusterName);
+			scope.Span.SetTag("StreamName", streamName);
+			scope.Span.SetTag("ChangeNumber", changeNumber);
 			
-			PerforceCluster Cluster = await GetClusterAsync(ClusterName);
-			using (P4.Repository Repository = await GetConnection(Cluster, NoClient: true))
+			PerforceCluster cluster = await GetClusterAsync(clusterName);
+			using (P4.Repository repository = await GetConnection(cluster, noClient: true))
 			{
-				List<ChangeDetails> Results = new List<ChangeDetails>();
+				List<ChangeDetails> results = new List<ChangeDetails>();
 
-				List<string> Args = new List<string> { "-s", $"{ChangeNumber}" };
+				List<string> args = new List<string> { "-s", $"{changeNumber}" };
 
-				using (P4.P4Command Command = new P4.P4Command(Repository, "describe", true, Args.ToArray()))
+				using (P4.P4Command command = new P4.P4Command(repository, "describe", true, args.ToArray()))
 				{
-					P4.P4CommandResult Result = Command.Run();
+					P4.P4CommandResult result = command.Run();
 
-					if (!Result.Success || Result.TaggedOutput == null || Result.TaggedOutput.Count == 0)
+					if (!result.Success || result.TaggedOutput == null || result.TaggedOutput.Count == 0)
 					{
 						throw new Exception("Unable to get changes");
 					}
 
-					List<P4.Changelist> Changelists = new List<P4.Changelist>() { };
+					List<P4.Changelist> changelists = new List<P4.Changelist>() { };
 
-					bool DSTMismatch = false;
-					string Offset = string.Empty;
+					bool dstMismatch = false;
+					string offset = String.Empty;
 
-					P4.Server Server = Repository.Server;
-					if (Server != null && Server.Metadata != null)
+					P4.Server server = repository.Server;
+					if (server != null && server.Metadata != null)
 					{
-						Offset = Server.Metadata.DateTimeOffset;
-						DSTMismatch = P4.FormBase.DSTMismatch(Server.Metadata);
+						offset = server.Metadata.DateTimeOffset;
+						dstMismatch = P4.FormBase.DSTMismatch(server.Metadata);
 					}
 
-					P4.TaggedObject TaggedObject = Result.TaggedOutput[0];
+					P4.TaggedObject taggedObject = result.TaggedOutput[0];
 					{
-						List<ChangeFile> Files = new List<ChangeFile>();
+						List<ChangeFile> files = new List<ChangeFile>();
 
-						P4.Changelist Change = new P4.Changelist();
-						Change.FromChangeCmdTaggedOutput(TaggedObject, false, Offset, DSTMismatch);
+						P4.Changelist change = new P4.Changelist();
+						change.FromChangeCmdTaggedOutput(taggedObject, false, offset, dstMismatch);
 
-						foreach (P4.FileMetaData DescribeFile in Change.Files)
+						foreach (P4.FileMetaData describeFile in change.Files)
 						{
-							string? RelativePath;
-							if (TryGetStreamRelativePath(DescribeFile.DepotPath.Path, StreamName, out RelativePath))
+							string? relativePath;
+							if (TryGetStreamRelativePath(describeFile.DepotPath.Path, streamName, out relativePath))
 							{
-								Files.Add(CreateChangeFile(RelativePath, DescribeFile));
+								files.Add(CreateChangeFile(relativePath, describeFile));
 							}
 						}
 
-						IUser User = await FindOrAddUserAsync(ClusterName, Change.OwnerName);
-						return new ChangeDetails(Change.Id, User, Change.GetPath(), Change.Description, Files, Change.ModifiedDate);
+						IUser user = await FindOrAddUserAsync(clusterName, change.OwnerName);
+						return new ChangeDetails(change.Id, user, change.GetPath(), change.Description, files, change.ModifiedDate);
 					}
 				}
 			}
 		}
 
 		/// <inheritdoc/>
-		public async Task<CheckShelfResult> CheckShelfAsync(string ClusterName, string StreamName, int ChangeNumber, string? ImpersonateUser)
+		public async Task<CheckShelfResult> CheckShelfAsync(string clusterName, string streamName, int changeNumber, string? impersonateUser)
 		{
-			using IScope Scope = GlobalTracer.Instance.BuildSpan("PerforceService.CheckPreflightAsync").StartActive();
-			Scope.Span.SetTag("ClusterName", ClusterName);
-			Scope.Span.SetTag("StreamName", StreamName);
-			Scope.Span.SetTag("ChangeNumber", ChangeNumber);
-			Scope.Span.SetTag("ImpersonateUser", ImpersonateUser);
-
-			PerforceCluster Cluster = await GetClusterAsync(ClusterName);
-			using (P4.Repository Repository = await GetConnection(Cluster, Stream: StreamName, Username: ImpersonateUser))
+			using IScope scope = GlobalTracer.Instance.BuildSpan("PerforceService.CheckPreflightAsync").StartActive();
+			scope.Span.SetTag("ClusterName", clusterName);
+			scope.Span.SetTag("StreamName", streamName);
+			scope.Span.SetTag("ChangeNumber", changeNumber);
+			scope.Span.SetTag("ImpersonateUser", impersonateUser);
+
+			PerforceCluster cluster = await GetClusterAsync(clusterName);
+			using (P4.Repository repository = await GetConnection(cluster, stream: streamName, username: impersonateUser))
 			{
-				P4.Changelist Change = Repository.GetChangelist(ChangeNumber, new P4.DescribeCmdOptions(P4.DescribeChangelistCmdFlags.Omit | P4.DescribeChangelistCmdFlags.Shelved, 0, 0));
-				if(Change == null)
+				P4.Changelist change = repository.GetChangelist(changeNumber, new P4.DescribeCmdOptions(P4.DescribeChangelistCmdFlags.Omit | P4.DescribeChangelistCmdFlags.Shelved, 0, 0));
+				if(change == null)
 				{
 					return CheckShelfResult.NoChange;
 				}
-				if (Change.ShelvedFiles == null || Change.ShelvedFiles.Count == 0)
+				if (change.ShelvedFiles == null || change.ShelvedFiles.Count == 0)
 				{
 					return CheckShelfResult.NoShelvedFiles;
 				}
 
-				P4.Stream Stream = Repository.GetStream(StreamName, new P4.StreamCmdOptions(P4.StreamCmdFlags.View, null, null));
-				ViewMap Map = new ViewMap(Stream.View);
+				P4.Stream stream = repository.GetStream(streamName, new P4.StreamCmdOptions(P4.StreamCmdFlags.View, null, null));
+				ViewMap map = new ViewMap(stream.View);
 
 				bool bHasMappedFile = false;
 				bool bHasUnmappedFile = false;
-				foreach (P4.ShelvedFile ShelvedFile in Change.ShelvedFiles)
+				foreach (P4.ShelvedFile shelvedFile in change.ShelvedFiles)
 				{
-					if (Map.TryMapFile(ShelvedFile.Path.Path, Utf8StringComparer.OrdinalIgnoreCase, out Utf8String _))
+					if (map.TryMapFile(shelvedFile.Path.Path, Utf8StringComparer.OrdinalIgnoreCase, out Utf8String _))
 					{
 						bHasMappedFile = true;
 					}
@@ -957,725 +951,720 @@ namespace Horde.Build.Services
 		}
 
 		/// <inheritdoc/>
-		public async Task<List<ChangeDetails>> GetChangeDetailsAsync(string ClusterName, string StreamName, IReadOnlyList<int> ChangeNumbers, string? ImpersonateUser)
+		public async Task<List<ChangeDetails>> GetChangeDetailsAsync(string clusterName, string streamName, IReadOnlyList<int> changeNumbers, string? impersonateUser)
 		{
-			using IScope Scope = GlobalTracer.Instance.BuildSpan("PerforceService.GetChangeDetailsAsync").StartActive();
-			Scope.Span.SetTag("ClusterName", ClusterName);
-			Scope.Span.SetTag("StreamName", StreamName);
-			Scope.Span.SetTag("ChangeNumbers.Count", ChangeNumbers.Count);
-			Scope.Span.SetTag("ImpersonateUser", ImpersonateUser);
+			using IScope scope = GlobalTracer.Instance.BuildSpan("PerforceService.GetChangeDetailsAsync").StartActive();
+			scope.Span.SetTag("ClusterName", clusterName);
+			scope.Span.SetTag("StreamName", streamName);
+			scope.Span.SetTag("ChangeNumbers.Count", changeNumbers.Count);
+			scope.Span.SetTag("ImpersonateUser", impersonateUser);
 			
-			PerforceCluster Cluster = await GetClusterAsync(ClusterName);
-			using (P4.Repository Repository = await GetConnection(Cluster, Stream: StreamName, Username: ImpersonateUser))
+			PerforceCluster cluster = await GetClusterAsync(clusterName);
+			using (P4.Repository repository = await GetConnection(cluster, stream: streamName, username: impersonateUser))
 			{
-				List<ChangeDetails> Results = new List<ChangeDetails>();
+				List<ChangeDetails> results = new List<ChangeDetails>();
 
-				List<string> Args = new List<string> { "-s", "-S" };
-				Args.AddRange(ChangeNumbers.Select(Change => Change.ToString(CultureInfo.InvariantCulture)));
+				List<string> args = new List<string> { "-s", "-S" };
+				args.AddRange(changeNumbers.Select(change => change.ToString(CultureInfo.InvariantCulture)));
 
-				using (P4.P4Command Command = new P4.P4Command(Repository, "describe", true, Args.ToArray()))
+				using (P4.P4Command command = new P4.P4Command(repository, "describe", true, args.ToArray()))
 				{
-					P4.P4CommandResult Result = Command.Run();
+					P4.P4CommandResult result = command.Run();
 
-					if (!Result.Success)
+					if (!result.Success)
 					{
 						throw new Exception("Unable to get changes");
 					}
 
-					if (Result.TaggedOutput == null || Result.TaggedOutput.Count <= 0)
+					if (result.TaggedOutput == null || result.TaggedOutput.Count <= 0)
 					{
-						return Results;
+						return results;
 					}
 
-					List<P4.Changelist> Changelists = new List<P4.Changelist>() { };
+					List<P4.Changelist> changelists = new List<P4.Changelist>() { };
 
-					bool DSTMismatch = false;
-					string Offset = string.Empty;
+					bool dstMismatch = false;
+					string offset = String.Empty;
 
-					P4.Server Server = Repository.Server;
-					if (Server != null && Server.Metadata != null)
+					P4.Server server = repository.Server;
+					if (server != null && server.Metadata != null)
 					{
-						Offset = Server.Metadata.DateTimeOffset;
-						DSTMismatch = P4.FormBase.DSTMismatch(Server.Metadata);
+						offset = server.Metadata.DateTimeOffset;
+						dstMismatch = P4.FormBase.DSTMismatch(server.Metadata);
 					}
 
-					foreach (P4.TaggedObject TaggedObject in Result.TaggedOutput)
+					foreach (P4.TaggedObject taggedObject in result.TaggedOutput)
 					{
-						List<ChangeFile> Files = new List<ChangeFile>();
+						List<ChangeFile> files = new List<ChangeFile>();
 
-						P4.Changelist Change = new P4.Changelist();
-						Change.FromChangeCmdTaggedOutput(TaggedObject, true, Offset, DSTMismatch);
+						P4.Changelist change = new P4.Changelist();
+						change.FromChangeCmdTaggedOutput(taggedObject, true, offset, dstMismatch);
 
-						foreach (P4.FileMetaData DescribeFile in Change.Files)
+						foreach (P4.FileMetaData describeFile in change.Files)
 						{
-							string? RelativePath;
-							if (TryGetStreamRelativePath(DescribeFile.DepotPath.Path, StreamName, out RelativePath))
+							string? relativePath;
+							if (TryGetStreamRelativePath(describeFile.DepotPath.Path, streamName, out relativePath))
 							{
-								Files.Add(CreateChangeFile(RelativePath, DescribeFile));
+								files.Add(CreateChangeFile(relativePath, describeFile));
 							}
 						}
 
-						if (Change.ShelvedFiles != null && Change.ShelvedFiles.Count > 0)
+						if (change.ShelvedFiles != null && change.ShelvedFiles.Count > 0)
 						{
-							foreach (P4.ShelvedFile ShelvedFile in Change.ShelvedFiles)
+							foreach (P4.ShelvedFile shelvedFile in change.ShelvedFiles)
 							{
-								string? RelativePath;
-								if (TryGetStreamRelativePath(ShelvedFile.Path.ToString(), StreamName, out RelativePath))
+								string? relativePath;
+								if (TryGetStreamRelativePath(shelvedFile.Path.ToString(), streamName, out relativePath))
 								{
-									Files.Add(CreateChangeFile(RelativePath, ShelvedFile));
+									files.Add(CreateChangeFile(relativePath, shelvedFile));
 								}
 							}
-
 						}
 
-						IUser User = await FindOrAddUserAsync(ClusterName, Change.OwnerName);
-						Results.Add(new ChangeDetails(Change.Id, User, Change.GetPath(), Change.Description, Files, Change.ModifiedDate));
+						IUser user = await FindOrAddUserAsync(clusterName, change.OwnerName);
+						results.Add(new ChangeDetails(change.Id, user, change.GetPath(), change.Description, files, change.ModifiedDate));
 					}
-
 				}
 
-				return Results;
+				return results;
 			}
 		}
 
 		/// <inheritdoc/>
-		public async Task<List<FileSummary>> FindFilesAsync(string ClusterName, IEnumerable<string> Paths)
+		public async Task<List<FileSummary>> FindFilesAsync(string clusterName, IEnumerable<string> paths)
 		{
-			using IScope Scope = GlobalTracer.Instance.BuildSpan("PerforceService.FindFilesAsync").StartActive();
-			Scope.Span.SetTag("ClusterName", ClusterName);
+			using IScope scope = GlobalTracer.Instance.BuildSpan("PerforceService.FindFilesAsync").StartActive();
+			scope.Span.SetTag("ClusterName", clusterName);
 			
-			List<FileSummary> Results = new List<FileSummary>();
+			List<FileSummary> results = new List<FileSummary>();
 
-			PerforceCluster Cluster = await GetClusterAsync(ClusterName);
-			using (P4.Repository Repository = await GetConnection(Cluster, NoClient: true))
+			PerforceCluster cluster = await GetClusterAsync(clusterName);
+			using (P4.Repository repository = await GetConnection(cluster, noClient: true))
 			{
-				P4.GetFileMetaDataCmdOptions Options = new P4.GetFileMetaDataCmdOptions(P4.GetFileMetadataCmdFlags.ExcludeClientData, "", "", 0, "", "", "");
+				P4.GetFileMetaDataCmdOptions options = new P4.GetFileMetaDataCmdOptions(P4.GetFileMetadataCmdFlags.ExcludeClientData, "", "", 0, "", "", "");
 
-				IList<P4.FileMetaData>? Files = Repository.GetFileMetaData(Options, Paths.Select(Path => { P4.FileSpec FileSpec = new P4.FileSpec(new P4.DepotPath(Path)); return FileSpec; }).ToArray());
+				IList<P4.FileMetaData>? files = repository.GetFileMetaData(options, paths.Select(path => 
+				{ 
+					P4.FileSpec fileSpec = new P4.FileSpec(new P4.DepotPath(path)); 
+					return fileSpec; 
+				}).ToArray());
 
-				if (Files == null)
+				if (files == null)
 				{
-					Files = new List<P4.FileMetaData>();
+					files = new List<P4.FileMetaData>();
 				}
 
-				foreach (string Path in Paths)
+				foreach (string path in paths)
 				{
-					P4.FileMetaData? Meta = Files.FirstOrDefault(File => File.DepotPath?.Path == Path);
+					P4.FileMetaData? meta = files.FirstOrDefault(file => file.DepotPath?.Path == path);
 
-					if (Meta == null)
+					if (meta == null)
 					{
-						Results.Add(new FileSummary(Path, false, 0));
+						results.Add(new FileSummary(path, false, 0));
 					}
 					else
 					{
-						Results.Add(new FileSummary(Path, Meta.HeadAction != P4.FileAction.Delete && Meta.HeadAction != P4.FileAction.MoveDelete, Meta.HeadChange));
+						results.Add(new FileSummary(path, meta.HeadAction != P4.FileAction.Delete && meta.HeadAction != P4.FileAction.MoveDelete, meta.HeadChange));
 					}
-
 				}
 
-				return Results;
+				return results;
 			}
 		}
 
 		/// <inheritdoc/>
-		public async Task<byte[]> PrintAsync(string ClusterName, string DepotPath)
+		public async Task<byte[]> PrintAsync(string clusterName, string depotPath)
 		{
-			using IScope Scope = GlobalTracer.Instance.BuildSpan("PerforceService.PrintAsync").StartActive();
-			Scope.Span.SetTag("ClusterName", ClusterName);
-			Scope.Span.SetTag("DepotPath", DepotPath);
+			using IScope scope = GlobalTracer.Instance.BuildSpan("PerforceService.PrintAsync").StartActive();
+			scope.Span.SetTag("ClusterName", clusterName);
+			scope.Span.SetTag("DepotPath", depotPath);
 			
-			if (DepotPath.EndsWith("...", StringComparison.OrdinalIgnoreCase) || DepotPath.EndsWith("*", StringComparison.OrdinalIgnoreCase))
+			if (depotPath.EndsWith("...", StringComparison.OrdinalIgnoreCase) || depotPath.EndsWith("*", StringComparison.OrdinalIgnoreCase))
 			{
 				throw new Exception("PrintAsync requires exactly one file to be specified");
 			}
 
-			PerforceCluster Cluster = await GetClusterAsync(ClusterName);
-			using (P4.Repository Repository = await GetConnection(Cluster, NoClient: true))
+			PerforceCluster cluster = await GetClusterAsync(clusterName);
+			using (P4.Repository repository = await GetConnection(cluster, noClient: true))
 			{
-				using (P4.P4Command Command = new P4.P4Command(Repository, "print", false, new string[] { "-q", DepotPath }))
+				using (P4.P4Command command = new P4.P4Command(repository, "print", false, new string[] { "-q", depotPath }))
 				{
-					P4.P4CommandResult Result = Command.Run();
+					P4.P4CommandResult result = command.Run();
 
-					if (Result.BinaryOutput != null)
+					if (result.BinaryOutput != null)
 					{
-						return Result.BinaryOutput;
+						return result.BinaryOutput;
 					}
 
-					return Encoding.Default.GetBytes(Result.TextOutput);
-
+					return Encoding.Default.GetBytes(result.TextOutput);
 				}
-
 			}
 		}
 
 		/// <inheritdoc/>
-		public async Task<int> DuplicateShelvedChangeAsync(string ClusterName, int ShelvedChange)
+		public async Task<int> DuplicateShelvedChangeAsync(string clusterName, int shelvedChange)
 		{
-			using IScope Scope = GlobalTracer.Instance.BuildSpan("PerforceService.DuplicateShelvedChangeAsync").StartActive();
-			Scope.Span.SetTag("ClusterName", ClusterName);
-			Scope.Span.SetTag("ShelvedChange", ShelvedChange);
+			using IScope scope = GlobalTracer.Instance.BuildSpan("PerforceService.DuplicateShelvedChangeAsync").StartActive();
+			scope.Span.SetTag("ClusterName", clusterName);
+			scope.Span.SetTag("ShelvedChange", shelvedChange);
 
-			string? ChangeOwner = null;
+			string? changeOwner = null;
 
 			// Get the owner of the shelf
-			PerforceCluster Cluster = await GetClusterAsync(ClusterName);
-			using (P4.Repository Repository = await GetConnection(Cluster, NoClient: true))
+			PerforceCluster cluster = await GetClusterAsync(clusterName);
+			using (P4.Repository repository = await GetConnection(cluster, noClient: true))
 			{
-				List<string> Args = new List<string> { "-S", ShelvedChange.ToString(CultureInfo.InvariantCulture) };
+				List<string> args = new List<string> { "-S", shelvedChange.ToString(CultureInfo.InvariantCulture) };
 
-				using (P4.P4Command Command = new P4.P4Command(Repository, "describe", true, Args.ToArray()))
+				using (P4.P4Command command = new P4.P4Command(repository, "describe", true, args.ToArray()))
 				{
-					P4.P4CommandResult Result = Command.Run();
+					P4.P4CommandResult result = command.Run();
 
-					if (!Result.Success)
+					if (!result.Success)
 					{
-						throw new Exception($"Unable to get change {ShelvedChange}");
+						throw new Exception($"Unable to get change {shelvedChange}");
 					}
 
-					if (Result.TaggedOutput == null || Result.TaggedOutput.Count != 1)
+					if (result.TaggedOutput == null || result.TaggedOutput.Count != 1)
 					{
 
-						throw new Exception($"Unable to get tagged output for change: {ShelvedChange}");
+						throw new Exception($"Unable to get tagged output for change: {shelvedChange}");
 					}
 
-					P4.Changelist Changelist = new P4.Changelist();
-					Changelist.FromChangeCmdTaggedOutput(Result.TaggedOutput[0], true, string.Empty, false);
-					ChangeOwner = Changelist.OwnerName;
+					P4.Changelist changelist = new P4.Changelist();
+					changelist.FromChangeCmdTaggedOutput(result.TaggedOutput[0], true, String.Empty, false);
+					changeOwner = changelist.OwnerName;
 
 				}
 			}
 
-			if (string.IsNullOrEmpty(ChangeOwner))
+			if (String.IsNullOrEmpty(changeOwner))
 			{
-				throw new Exception($"Unable to get owner for shelved change {ShelvedChange}");
+				throw new Exception($"Unable to get owner for shelved change {shelvedChange}");
 			}
 
-			using (P4.Repository Repository = await GetConnection(Cluster, ReadOnly: false, UsePortFromChange: ShelvedChange, ClientFromChange: ShelvedChange, UseClientFromChange: false, Username: ChangeOwner))
+			using (P4.Repository repository = await GetConnection(cluster, readOnly: false, usePortFromChange: shelvedChange, clientFromChange: shelvedChange, useClientFromChange: false, username: changeOwner))
 			{
-				return ReshelveChange(Repository, ShelvedChange);
+				return ReshelveChange(repository, shelvedChange);
 			}
 
-			throw new Exception($"Unable to duplicate shelve change {ShelvedChange}");
+			throw new Exception($"Unable to duplicate shelve change {shelvedChange}");
 		}
 
 		/// <inheritdoc/>
-		public async Task DeleteShelvedChangeAsync(string ClusterName, int ShelvedChange)
+		public async Task DeleteShelvedChangeAsync(string clusterName, int shelvedChange)
 		{
-			using IScope Scope = GlobalTracer.Instance.BuildSpan("PerforceService.DeleteShelvedChangeAsync").StartActive();
-			Scope.Span.SetTag("ClusterName", ClusterName);
-			Scope.Span.SetTag("ShelvedChange", ShelvedChange);
+			using IScope scope = GlobalTracer.Instance.BuildSpan("PerforceService.DeleteShelvedChangeAsync").StartActive();
+			scope.Span.SetTag("ClusterName", clusterName);
+			scope.Span.SetTag("ShelvedChange", shelvedChange);
 			
-			PerforceCluster Cluster = await GetClusterAsync(ClusterName);
-			using (P4.Repository Repository = await GetConnection(Cluster, NoClient: true))
+			PerforceCluster cluster = await GetClusterAsync(clusterName);
+			using (P4.Repository repository = await GetConnection(cluster, noClient: true))
 			{
 
-				P4.Changelist? Changelist = Repository.GetChangelist(ShelvedChange, new P4.DescribeCmdOptions(P4.DescribeChangelistCmdFlags.Omit, 0, 0));
+				P4.Changelist? changelist = repository.GetChangelist(shelvedChange, new P4.DescribeCmdOptions(P4.DescribeChangelistCmdFlags.Omit, 0, 0));
 
-				if (Changelist == null)
+				if (changelist == null)
 				{
-					throw new Exception($"Unable to get shelved changelist for delete: {ShelvedChange}");
+					throw new Exception($"Unable to get shelved changelist for delete: {shelvedChange}");
 				}
 
-				string ClientId = Changelist.ClientId;
+				string clientId = changelist.ClientId;
 
-				if (String.IsNullOrEmpty(ClientId))
+				if (String.IsNullOrEmpty(clientId))
 				{
-					throw new Exception($"Unable to get shelved changelist client id for delete: {ShelvedChange}");
+					throw new Exception($"Unable to get shelved changelist client id for delete: {shelvedChange}");
 				}
 
-				using (P4.Repository ClientRepository = await GetConnection(Cluster, Changelist.Stream, Changelist.OwnerName, ClientId: ClientId, UseClientFromChange: true, UsePortFromChange: ShelvedChange))
+				using (P4.Repository clientRepository = await GetConnection(cluster, changelist.Stream, changelist.OwnerName, clientId: clientId, useClientFromChange: true, usePortFromChange: shelvedChange))
 				{
 
-					if (Changelist.Shelved)
+					if (changelist.Shelved)
 					{
-						IList<P4.FileSpec> Files = ClientRepository.Connection.Client.ShelveFiles(new P4.ShelveFilesCmdOptions(P4.ShelveFilesCmdFlags.Delete, null, ShelvedChange));
+						IList<P4.FileSpec> files = clientRepository.Connection.Client.ShelveFiles(new P4.ShelveFilesCmdOptions(P4.ShelveFilesCmdFlags.Delete, null, shelvedChange));
 					}
 
-					ClientRepository.DeleteChangelist(Changelist, null);
+					clientRepository.DeleteChangelist(changelist, null);
 				}
 			}
 		}
 
 		/// <inheritdoc/>
-		public async Task UpdateChangelistDescription(string ClusterName, int Change, string Description)
+		public async Task UpdateChangelistDescription(string clusterName, int change, string description)
 		{
-			using IScope Scope = GlobalTracer.Instance.BuildSpan("PerforceService.UpdateChangelistDescription").StartActive();
-			Scope.Span.SetTag("ClusterName", ClusterName);
-			Scope.Span.SetTag("Change", Change);
+			using IScope scope = GlobalTracer.Instance.BuildSpan("PerforceService.UpdateChangelistDescription").StartActive();
+			scope.Span.SetTag("ClusterName", clusterName);
+			scope.Span.SetTag("Change", change);
 
 			try
 			{
-				PerforceCluster Cluster = await GetClusterAsync(ClusterName);
+				PerforceCluster cluster = await GetClusterAsync(clusterName);
 
-				using (P4.Repository Repository = await GetConnection(Cluster, NoClient: true))
+				using (P4.Repository repository = await GetConnection(cluster, noClient: true))
 				{
-					P4.Changelist Changelist = Repository.GetChangelist(Change);
+					P4.Changelist changelist = repository.GetChangelist(change);
 
-					Repository.Connection.Disconnect();
+					repository.Connection.Disconnect();
 
 					// the client must exist for the change list, otherwise will fail (for example, CreateNewChangeAsync deletes the client before returning)
-					using (P4.Repository UpdateRepository = await GetConnection(Cluster, ClientFromChange: Change, UseClientFromChange: true, UsePortFromChange: Change, Username: Changelist.OwnerName))
+					using (P4.Repository updateRepository = await GetConnection(cluster, clientFromChange: change, useClientFromChange: true, usePortFromChange: change, username: changelist.OwnerName))
 					{
-						P4.Changelist UpdatedChangelist = UpdateRepository.GetChangelist(Change);
-						UpdatedChangelist.Description = Description;
-						UpdateRepository.UpdateChangelist(UpdatedChangelist);
+						P4.Changelist updatedChangelist = updateRepository.GetChangelist(change);
+						updatedChangelist.Description = description;
+						updateRepository.UpdateChangelist(updatedChangelist);
 					}
 				}
 			}
-			catch (Exception Ex)
+			catch (Exception ex)
 			{
-				LogPerforce(1, "", $"Unable to update Changelist for CL {Change} to ${Description}, {Ex.Message}");
+				LogPerforce(1, "", $"Unable to update Changelist for CL {change} to ${description}, {ex.Message}");
 			}
 		}
 
 		/// <inheritdoc/>
-		public async Task<int> CreateNewChangeAsync(string ClusterName, string StreamName, string FilePath, string Description)
+		public async Task<int> CreateNewChangeAsync(string clusterName, string streamName, string filePath, string description)
 		{
-			using IScope Scope = GlobalTracer.Instance.BuildSpan("PerforceService.CreateNewChangeAsync").StartActive();
-			Scope.Span.SetTag("ClusterName", ClusterName);
-			Scope.Span.SetTag("StreamName", StreamName);
-			Scope.Span.SetTag("FilePath", FilePath);
+			using IScope scope = GlobalTracer.Instance.BuildSpan("PerforceService.CreateNewChangeAsync").StartActive();
+			scope.Span.SetTag("ClusterName", clusterName);
+			scope.Span.SetTag("StreamName", streamName);
+			scope.Span.SetTag("FilePath", filePath);
 			
-			PerforceCluster Cluster = await GetClusterAsync(ClusterName);
-			P4.SubmitResults? SubmitResults = null;
+			PerforceCluster cluster = await GetClusterAsync(clusterName);
+			P4.SubmitResults? submitResults = null;
 
-			using (P4.Repository Repository = await GetConnection(Cluster, Stream: StreamName, ReadOnly: false, CreateChange: true))
+			using (P4.Repository repository = await GetConnection(cluster, stream: streamName, readOnly: false, createChange: true))
 			{
-				P4.Client Client = Repository.Connection.Client;
+				P4.Client client = repository.Connection.Client;
 
-				string WorkspaceFilePath = $"//{Client.Name}/{FilePath.TrimStart('/')}";
-				string DiskFilePath = $"{Client.Root + FilePath.TrimStart('/')}";
+				string workspaceFilePath = $"//{client.Name}/{filePath.TrimStart('/')}";
+				string diskFilePath = $"{client.Root + filePath.TrimStart('/')}";
 
-				P4.FileSpec WorkspaceFileSpec = new P4.FileSpec(new P4.ClientPath(WorkspaceFilePath));
+				P4.FileSpec workspaceFileSpec = new P4.FileSpec(new P4.ClientPath(workspaceFilePath));
 
-				IList<P4.File> Files = Repository.GetFiles(new P4.FilesCmdOptions(P4.FilesCmdFlags.None, 1), WorkspaceFileSpec);
+				IList<P4.File> files = repository.GetFiles(new P4.FilesCmdOptions(P4.FilesCmdFlags.None, 1), workspaceFileSpec);
 
-				P4.Changelist? SubmitChangelist = null;
-				P4.DepotPath? DepotPath = null;
+				P4.Changelist? submitChangelist = null;
+				P4.DepotPath? depotPath = null;
 
 				const int MaxRetries = 5;
-				int Retry = 0;
+				int retry = 0;
 
 				for (; ; )
 				{
-					ResetClient(Repository);
+					ResetClient(repository);
 
-					DepotPath = null;
+					depotPath = null;
 
-					if (Retry == MaxRetries)
+					if (retry == MaxRetries)
 					{
 						break;
 					}
 
 					try
 					{
-						if (Files == null || Files.Count == 0)
+						if (files == null || files.Count == 0)
 						{
 							// File does not exist, create it
-							string? DirectoryName = Path.GetDirectoryName(DiskFilePath);
+							string? directoryName = Path.GetDirectoryName(diskFilePath);
 
-							if (string.IsNullOrEmpty(DirectoryName))
+							if (String.IsNullOrEmpty(directoryName))
 							{
-								throw new Exception($"Invalid directory name for local client file, disk file path: {DiskFilePath}");
+								throw new Exception($"Invalid directory name for local client file, disk file path: {diskFilePath}");
 							}
 
 							// Create the directory
-							if (!Directory.Exists(DirectoryName))
+							if (!Directory.Exists(directoryName))
 							{
-								Directory.CreateDirectory(DirectoryName);
+								Directory.CreateDirectory(directoryName);
 
-								if (!Directory.Exists(DirectoryName))
+								if (!Directory.Exists(directoryName))
 								{
-									throw new Exception($"Unable to create directrory: {DirectoryName}");
+									throw new Exception($"Unable to create directrory: {directoryName}");
 								}
 							}
 
 							// Create the file
-							if (!File.Exists(DiskFilePath))
+							if (!File.Exists(diskFilePath))
 							{
-								using (FileStream FileStream = File.OpenWrite(DiskFilePath))
+								using (FileStream fileStream = File.OpenWrite(diskFilePath))
 								{
-									FileStream.Close();
+									fileStream.Close();
 								}
 
-								if (!File.Exists(DiskFilePath))
+								if (!File.Exists(diskFilePath))
 								{
-									throw new Exception($"Unable to create local change file: {DiskFilePath}");
+									throw new Exception($"Unable to create local change file: {diskFilePath}");
 								}
-
 							}
 
-							IList<P4.FileSpec> DepotFiles = Client.AddFiles(new P4.Options(), WorkspaceFileSpec);
+							IList<P4.FileSpec> depotFiles = client.AddFiles(new P4.Options(), workspaceFileSpec);
 
-							if (DepotFiles == null || DepotFiles.Count != 1)
+							if (depotFiles == null || depotFiles.Count != 1)
 							{
-								throw new Exception($"Unable to add local change file,  local: {DiskFilePath} : workspace: {WorkspaceFileSpec}");
+								throw new Exception($"Unable to add local change file,  local: {diskFilePath} : workspace: {workspaceFileSpec}");
 							}
 
-							DepotPath = DepotFiles[0].DepotPath;
+							depotPath = depotFiles[0].DepotPath;
 
 						}
 						else
 						{
-							IList<P4.FileSpec> SyncResults = Client.SyncFiles(new P4.SyncFilesCmdOptions(P4.SyncFilesCmdFlags.Force), WorkspaceFileSpec);
+							IList<P4.FileSpec> syncResults = client.SyncFiles(new P4.SyncFilesCmdOptions(P4.SyncFilesCmdFlags.Force), workspaceFileSpec);
 
-							if (SyncResults == null || SyncResults.Count != 1)
+							if (syncResults == null || syncResults.Count != 1)
 							{
-								throw new Exception($"Unable to sync file, workspace: {WorkspaceFileSpec}");
+								throw new Exception($"Unable to sync file, workspace: {workspaceFileSpec}");
 							}
 
-							IList<P4.FileSpec> EditResults = Client.EditFiles(new P4.FileSpec[] { new P4.FileSpec(SyncResults[0].DepotPath) }, new P4.Options());
+							IList<P4.FileSpec> editResults = client.EditFiles(new P4.FileSpec[] { new P4.FileSpec(syncResults[0].DepotPath) }, new P4.Options());
 
-							if (EditResults == null || EditResults.Count != 1)
+							if (editResults == null || editResults.Count != 1)
 							{
-								throw new Exception($"Unable to edit file, workspace: {WorkspaceFileSpec}");
+								throw new Exception($"Unable to edit file, workspace: {workspaceFileSpec}");
 							}
 
-							DepotPath = EditResults[0].DepotPath;
+							depotPath = editResults[0].DepotPath;
 
 						}
 
-						if (DepotPath == null || string.IsNullOrEmpty(DepotPath.Path))
+						if (depotPath == null || String.IsNullOrEmpty(depotPath.Path))
 						{
-							throw new Exception($"Unable to get depot path for: {WorkspaceFileSpec}");
+							throw new Exception($"Unable to get depot path for: {workspaceFileSpec}");
 						}
 
 						// create a new change
-						P4.Changelist Changelist = new P4.Changelist();
-						Changelist.Description = Description;
-						Changelist.Files.Add(new P4.FileMetaData(new P4.FileSpec(DepotPath)));
-						SubmitChangelist = Repository.CreateChangelist(Changelist);
+						P4.Changelist changelist = new P4.Changelist();
+						changelist.Description = description;
+						changelist.Files.Add(new P4.FileMetaData(new P4.FileSpec(depotPath)));
+						submitChangelist = repository.CreateChangelist(changelist);
 
-						if (SubmitChangelist == null)
+						if (submitChangelist == null)
 						{
-							throw new Exception($"Unable to create a changelist for: {DepotPath}");
+							throw new Exception($"Unable to create a changelist for: {depotPath}");
 						}
 
-						SubmitResults = SubmitChangelist.Submit(null);
+						submitResults = submitChangelist.Submit(null);
 
-						if (SubmitResults == null)
+						if (submitResults == null)
 						{
-							throw new Exception($"Unable to submit changelist for: {DepotPath}");
+							throw new Exception($"Unable to submit changelist for: {depotPath}");
 						}
 
 						break;
 					}
 					catch
 					{
-						Retry++;
+						retry++;
 						continue;
 					}
-
 				}
 
-				string ClientName = Client.Name;
+				string clientName = client.Name;
 				try
 				{
-					Repository.DeleteClient(Client, new P4.Options());
+					repository.DeleteClient(client, new P4.Options());
 				}
 				catch
 				{
-					Logger.LogError("Unable to delete client {ClientName}", ClientName);
+					_logger.LogError("Unable to delete client {ClientName}", clientName);
 				}
-
 			}
 
-			if (SubmitResults == null)
+			if (submitResults == null)
 			{
-				throw new Exception($"Unable to submit change for {StreamName} {FilePath}");
+				throw new Exception($"Unable to submit change for {streamName} {filePath}");
 			}
 
-			return SubmitResults.ChangeIdAfterSubmit;
+			return submitResults.ChangeIdAfterSubmit;
 		}
 
 		/// <inheritdoc/>
-		public async Task<(int? Change, string Message)> SubmitShelvedChangeAsync(string ClusterName, int Change, int OriginalChange)
+		public async Task<(int? Change, string Message)> SubmitShelvedChangeAsync(string clusterName, int change, int originalChange)
 		{
-			using IScope Scope = GlobalTracer.Instance.BuildSpan("PerforceService.SubmitShelvedChangeAsync").StartActive();
-			Scope.Span.SetTag("ClusterName", ClusterName);
-			Scope.Span.SetTag("Change", Change);
-			Scope.Span.SetTag("OriginalChange", OriginalChange);
+			using IScope scope = GlobalTracer.Instance.BuildSpan("PerforceService.SubmitShelvedChangeAsync").StartActive();
+			scope.Span.SetTag("ClusterName", clusterName);
+			scope.Span.SetTag("Change", change);
+			scope.Span.SetTag("OriginalChange", originalChange);
 			
-			PerforceCluster Cluster = await GetClusterAsync(ClusterName);
-			using (P4.Repository Repository = await GetConnection(Cluster, NoClient: true))
+			PerforceCluster cluster = await GetClusterAsync(clusterName);
+			using (P4.Repository repository = await GetConnection(cluster, noClient: true))
 			{
 
-				List<string> Args = new List<string> { "-S", Change.ToString(CultureInfo.InvariantCulture), OriginalChange.ToString(CultureInfo.InvariantCulture) };
+				List<string> args = new List<string> { "-S", change.ToString(CultureInfo.InvariantCulture), originalChange.ToString(CultureInfo.InvariantCulture) };
 
-				using (P4.P4Command Command = new P4.P4Command(Repository, "describe", true, Args.ToArray()))
+				using (P4.P4Command command = new P4.P4Command(repository, "describe", true, args.ToArray()))
 				{
-					P4.P4CommandResult Result = Command.Run();
+					P4.P4CommandResult result = command.Run();
 
-					if (!Result.Success)
+					if (!result.Success)
 					{
-						return (null, $"Unable to get change {Change}");
+						return (null, $"Unable to get change {change}");
 
 					}
 
-					if (Result.TaggedOutput == null || Result.TaggedOutput.Count != 2)
+					if (result.TaggedOutput == null || result.TaggedOutput.Count != 2)
 					{
 
-						return (null, $"Unable to get tagged output for change: {Change} and original change: {OriginalChange}");
+						return (null, $"Unable to get tagged output for change: {change} and original change: {originalChange}");
 					}
 
-					bool DSTMismatch = false;
-					string Offset = string.Empty;
+					bool dstMismatch = false;
+					string offset = String.Empty;
 
-					P4.Server Server = Repository.Server;
-					if (Server != null && Server.Metadata != null)
+					P4.Server server = repository.Server;
+					if (server != null && server.Metadata != null)
 					{
-						Offset = Server.Metadata.DateTimeOffset;
-						DSTMismatch = P4.FormBase.DSTMismatch(Server.Metadata);
+						offset = server.Metadata.DateTimeOffset;
+						dstMismatch = P4.FormBase.DSTMismatch(server.Metadata);
 					}
 
-					P4.Changelist Changelist = new P4.Changelist();
-					Changelist.FromChangeCmdTaggedOutput(Result.TaggedOutput[0], true, Offset, DSTMismatch);
+					P4.Changelist changelist = new P4.Changelist();
+					changelist.FromChangeCmdTaggedOutput(result.TaggedOutput[0], true, offset, dstMismatch);
 
-					P4.Changelist OriginalChangelist = new P4.Changelist();
-					OriginalChangelist.FromChangeCmdTaggedOutput(Result.TaggedOutput[1], true, Offset, DSTMismatch);
+					P4.Changelist originalChangelist = new P4.Changelist();
+					originalChangelist.FromChangeCmdTaggedOutput(result.TaggedOutput[1], true, offset, dstMismatch);
 
-					if (OriginalChangelist.ShelvedFiles.Count != Changelist.ShelvedFiles.Count)
+					if (originalChangelist.ShelvedFiles.Count != changelist.ShelvedFiles.Count)
 					{
-						return (null, $"Mismatched number of shelved files for change: {Change} and original change: {OriginalChange}");
+						return (null, $"Mismatched number of shelved files for change: {change} and original change: {originalChange}");
 					}
 
-					if (OriginalChangelist.ShelvedFiles.Count == 0)
+					if (originalChangelist.ShelvedFiles.Count == 0)
 					{
-						return (null, $"No shelved file for change: {Change} and original change: {OriginalChange}");
+						return (null, $"No shelved file for change: {change} and original change: {originalChange}");
 					}
 
-					foreach (P4.ShelvedFile ShelvedFile in Changelist.ShelvedFiles)
+					foreach (P4.ShelvedFile shelvedFile in changelist.ShelvedFiles)
 					{
-						P4.ShelvedFile? Found = OriginalChangelist.ShelvedFiles.FirstOrDefault(Original => Original.Digest == ShelvedFile.Digest && Original.Action == ShelvedFile.Action);
+						P4.ShelvedFile? found = originalChangelist.ShelvedFiles.FirstOrDefault(original => original.Digest == shelvedFile.Digest && original.Action == shelvedFile.Action);
 
-						if (Found == null)
+						if (found == null)
 						{
-							return (null, $"Mismatch in shelved file digest or action for {ShelvedFile.Path}");
+							return (null, $"Mismatch in shelved file digest or action for {shelvedFile.Path}");
 						}
 					}
 
-					Repository.Connection.Disconnect();
+					repository.Connection.Disconnect();
 
-					using (P4.Repository SubmitRepository = await GetConnection(Cluster, ReadOnly: false, ClientFromChange: Change, UseClientFromChange: true, UsePortFromChange: Change, Username: Changelist.OwnerName))
+					using (P4.Repository submitRepository = await GetConnection(cluster, readOnly: false, clientFromChange: change, useClientFromChange: true, usePortFromChange: change, username: changelist.OwnerName))
 					{
 						// we might not need a client here, possibly -e below facilitates this, check!
-						using (P4.P4Command SubmitCommand = new P4.P4Command(SubmitRepository, "submit", null, true, new string[] { "-e", Change.ToString(CultureInfo.InvariantCulture) }))
+						using (P4.P4Command submitCommand = new P4.P4Command(submitRepository, "submit", null, true, new string[] { "-e", change.ToString(CultureInfo.InvariantCulture) }))
 						{
 							try
 							{
-								Result = SubmitCommand.Run();
+								result = submitCommand.Run();
 							}
-							catch (Exception Ex)
+							catch (Exception ex)
 							{
-								return (null, $"Submit command failed: {Ex.Message}");
+								return (null, $"Submit command failed: {ex.Message}");
 							}
 
-							if (!Result.Success)
+							if (!result.Success)
 							{
-								string Message = (Result.ErrorList != null && Result.ErrorList.Count > 0) ? Result.ErrorList[0].ErrorMessage : "Unknown error, no errors in list";
-								return (null, $"Unable to submit {Change}, {Message}");
+								string message = (result.ErrorList != null && result.ErrorList.Count > 0) ? result.ErrorList[0].ErrorMessage : "Unknown error, no errors in list";
+								return (null, $"Unable to submit {change}, {message}");
 							}
 
-							int? SubmittedChangeId = null;
+							int? submittedChangeId = null;
 
-							foreach (P4.TaggedObject TaggedObject in Result.TaggedOutput)
+							foreach (P4.TaggedObject taggedObject in result.TaggedOutput)
 							{
-								string? Submitted;
-								if (TaggedObject.TryGetValue("submittedChange", out Submitted))
+								string? submitted;
+								if (taggedObject.TryGetValue("submittedChange", out submitted))
 								{
-									SubmittedChangeId = int.Parse(Submitted, CultureInfo.InvariantCulture);
+									submittedChangeId = Int32.Parse(submitted, CultureInfo.InvariantCulture);
 								}
 							}
 
-							if (SubmittedChangeId == null)
+							if (submittedChangeId == null)
 							{
 								return (null, $"Submit command succeeded, though unable to parse submitted change number");
 							}
 
-							return (SubmittedChangeId, Changelist.Description);
+							return (submittedChangeId, changelist.Description);
 
 						}
 					}
 				}
 			}
 
-			throw new Exception($"Unable to get shelve change: {Change} for original change: {OriginalChange}");
+			throw new Exception($"Unable to get shelve change: {change} for original change: {originalChange}");
 		}
 
 		/// <inheritdoc/>		
-		public async Task<PerforceUserInfo?> GetUserInfoAsync(string ClusterName, string UserName)
+		public async Task<PerforceUserInfo?> GetUserInfoAsync(string clusterName, string userName)
 		{
-			using IScope Scope = GlobalTracer.Instance.BuildSpan("PerforceService.GetUserInfoAsync").StartActive();
-			Scope.Span.SetTag("ClusterName", ClusterName);
-			Scope.Span.SetTag("UserName", UserName);
+			using IScope scope = GlobalTracer.Instance.BuildSpan("PerforceService.GetUserInfoAsync").StartActive();
+			scope.Span.SetTag("ClusterName", clusterName);
+			scope.Span.SetTag("UserName", userName);
 			
-			PerforceCluster Cluster = await GetClusterAsync(ClusterName);
-			using (P4.Repository Repository = await GetConnection(Cluster, NoClient: true))
+			PerforceCluster cluster = await GetClusterAsync(clusterName);
+			using (P4.Repository repository = await GetConnection(cluster, noClient: true))
 			{
-				P4.User User = Repository.GetUser(UserName, new P4.UserCmdOptions(P4.UserCmdFlags.Output));
+				P4.User user = repository.GetUser(userName, new P4.UserCmdOptions(P4.UserCmdFlags.Output));
 
-				if (User == null)
+				if (user == null)
 				{
 					return null;
 				}
 
-				return new PerforceUserInfo { Login = UserName, FullName = User.FullName, Email = User.EmailAddress };
+				return new PerforceUserInfo { Login = userName, FullName = user.FullName, Email = user.EmailAddress };
 			}
 		}
 
-		static int ReshelveChange(P4.Repository Repository, int Change)
+		static int ReshelveChange(P4.Repository repository, int change)
 		{
-			using IScope Scope = GlobalTracer.Instance.BuildSpan("PerforceService.ReshelveChange").StartActive();
-			Scope.Span.SetTag("Change", Change);
+			using IScope scope = GlobalTracer.Instance.BuildSpan("PerforceService.ReshelveChange").StartActive();
+			scope.Span.SetTag("Change", change);
 
-			bool EdgeServer = false;
-			string? Value;
-			if (Repository.Server.Metadata.RawData.TryGetValue("serverServices", out Value))
+			bool edgeServer = false;
+			string? value;
+			if (repository.Server.Metadata.RawData.TryGetValue("serverServices", out value))
 			{
-				if (Value == "edge-server")
+				if (value == "edge-server")
 				{
-					EdgeServer = true;
+					edgeServer = true;
 				}
 			}
 
-			List<string> Arguments = new List<string>();
+			List<string> arguments = new List<string>();
 
 			// promote shelf if we're on an edge server
-			if (EdgeServer)
+			if (edgeServer)
 			{
-				Arguments.Add("-p");
+				arguments.Add("-p");
 			}
 
-			Arguments.AddRange(new string[] { "-s", Change.ToString(CultureInfo.InvariantCulture), "-f" });
+			arguments.AddRange(new string[] { "-s", change.ToString(CultureInfo.InvariantCulture), "-f" });
 
-			using (P4.P4Command Cmd = new P4.P4Command(Repository, "reshelve", false, Arguments.ToArray()))
+			using (P4.P4Command cmd = new P4.P4Command(repository, "reshelve", false, arguments.ToArray()))
 			{
-				P4.P4CommandResult Results = Cmd.Run();
+				P4.P4CommandResult results = cmd.Run();
 
-				if (Results.Success)
+				if (results.Success)
 				{
 
-					if (Results.InfoOutput.Count == 0)
+					if (results.InfoOutput.Count == 0)
 					{
 						Serilog.Log.Logger.Information("Perforce: Unexpected info output when reshelving change");
 						throw new Exception("Unexpected info output when reshelving change");
 					}
 
-					bool Error = true;
-					int ReshelvedChange = 0;
-					string Message = Results.InfoOutput[Results.InfoOutput.Count - 1].Message;
-					Match ChangeMatch = Regex.Match(Message, @"Change (\d+) files shelved");
+					bool error = true;
+					int reshelvedChange = 0;
+					string message = results.InfoOutput[^1].Message;
+					Match changeMatch = Regex.Match(message, @"Change (\d+) files shelved");
 
-					if (ChangeMatch.Success && ChangeMatch.Groups.Count == 2)
+					if (changeMatch.Success && changeMatch.Groups.Count == 2)
 					{
-						if (int.TryParse(ChangeMatch.Groups[1].Value, out ReshelvedChange))
+						if (Int32.TryParse(changeMatch.Groups[1].Value, out reshelvedChange))
 						{
-							Error = false;
+							error = false;
 						}
 					}
 
-					if (Error)
+					if (error)
 					{
-						Serilog.Log.Logger.Information("Perforce: Unable to parse cl for reshelf: {Message}", Message);
-						throw new Exception($"Unable to parse cl for reshelf: {Message}");
+						Serilog.Log.Logger.Information("Perforce: Unable to parse cl for reshelf: {Message}", message);
+						throw new Exception($"Unable to parse cl for reshelf: {message}");
 					}
 
-					return ReshelvedChange;
+					return reshelvedChange;
 				}
 			}
 
 			Serilog.Log.Logger.Information("Perforce: General reshelving failure");
 
-			throw new Exception($"Unable to reshelve CL {Change}");
+			throw new Exception($"Unable to reshelve CL {change}");
 
 		}
 
 		/// <inheritdoc/>
-		public virtual async Task<int> GetCodeChangeAsync(string ClusterName, string StreamName, int Change)
+		public virtual async Task<int> GetCodeChangeAsync(string clusterName, string streamName, int change)
 		{
-			using IScope Scope = GlobalTracer.Instance.BuildSpan("PerforceService.GetCodeChangeAsync").StartActive();
-			Scope.Span.SetTag("ClusterName", ClusterName);
-			Scope.Span.SetTag("StreamName", StreamName);
-			Scope.Span.SetTag("Change", Change);
+			using IScope scope = GlobalTracer.Instance.BuildSpan("PerforceService.GetCodeChangeAsync").StartActive();
+			scope.Span.SetTag("ClusterName", clusterName);
+			scope.Span.SetTag("StreamName", streamName);
+			scope.Span.SetTag("Change", change);
 			
-			int MaxChange = Change;
+			int maxChange = change;
 			for (; ; )
 			{
 				// Query for the changes before this point
-				List<ChangeSummary> Changes = await GetChangesAsync(ClusterName, StreamName, null, MaxChange, 10, null);
-				Serilog.Log.Logger.Information("Finding last code change in {Stream} before {MaxChange}: {NumResults}", StreamName, MaxChange, Changes.Count);
-				if (Changes.Count == 0)
+				List<ChangeSummary> changes = await GetChangesAsync(clusterName, streamName, null, maxChange, 10, null);
+				Serilog.Log.Logger.Information("Finding last code change in {Stream} before {MaxChange}: {NumResults}", streamName, maxChange, changes.Count);
+				if (changes.Count == 0)
 				{
 					return 0;
 				}
 
 				// Get the details for them
-				List<ChangeDetails> DetailsList = await GetChangeDetailsAsync(ClusterName, StreamName, Changes.ConvertAll(x => x.Number), null);
-				foreach (ChangeDetails Details in DetailsList.OrderByDescending(x => x.Number))
+				List<ChangeDetails> detailsList = await GetChangeDetailsAsync(clusterName, streamName, changes.ConvertAll(x => x.Number), null);
+				foreach (ChangeDetails details in detailsList.OrderByDescending(x => x.Number))
 				{
-					ChangeContentFlags ContentFlags = Details.GetContentFlags();
-					Serilog.Log.Logger.Information("Change {Change} = {Flags}", Details.Number, ContentFlags.ToString());
-					if ((Details.GetContentFlags() & ChangeContentFlags.ContainsCode) != 0)
+					ChangeContentFlags contentFlags = details.GetContentFlags();
+					Serilog.Log.Logger.Information("Change {Change} = {Flags}", details.Number, contentFlags.ToString());
+					if ((details.GetContentFlags() & ChangeContentFlags.ContainsCode) != 0)
 					{
-						return Details.Number;
+						return details.Number;
 					}
 				}
 
 				// Loop round again, adjusting our maximum changelist number
-				MaxChange = Changes.Min(x => x.Number) - 1;
+				maxChange = changes.Min(x => x.Number) - 1;
 			}
 		}
 
 		/// <summary>
 		/// Gets the wildcard filter for a particular range of changes
 		/// </summary>
-		/// <param name="BasePath">Base path to find files for</param>
-		/// <param name="MinChange">Minimum changelist number to query</param>
-		/// <param name="MaxChange">Maximum changelist number to query</param>
+		/// <param name="basePath">Base path to find files for</param>
+		/// <param name="minChange">Minimum changelist number to query</param>
+		/// <param name="maxChange">Maximum changelist number to query</param>
 		/// <returns>Filter string</returns>
-		public static string GetFilter(string BasePath, int? MinChange, int? MaxChange)
+		public static string GetFilter(string basePath, int? minChange, int? maxChange)
 		{
-			StringBuilder Filter = new StringBuilder(BasePath);
-			if (MinChange != null && MaxChange != null)
+			StringBuilder filter = new StringBuilder(basePath);
+			if (minChange != null && maxChange != null)
 			{
-				Filter.Append(CultureInfo.InvariantCulture, $"@{MinChange},{MaxChange}");
+				filter.Append(CultureInfo.InvariantCulture, $"@{minChange},{maxChange}");
 			}
-			else if (MinChange != null)
+			else if (minChange != null)
 			{
-				Filter.Append(CultureInfo.InvariantCulture, $"@>={MinChange}");
+				filter.Append(CultureInfo.InvariantCulture, $"@>={minChange}");
 			}
-			else if (MaxChange != null)
+			else if (maxChange != null)
 			{
-				Filter.Append(CultureInfo.InvariantCulture, $"@<={MaxChange}");
+				filter.Append(CultureInfo.InvariantCulture, $"@<={maxChange}");
 			}
-			return Filter.ToString();
+			return filter.ToString();
 		}
 
 		/// <summary>
 		/// Gets a stream-relative path from a depot path
 		/// </summary>
-		/// <param name="DepotFile">The depot file to check</param>
-		/// <param name="StreamName">Name of the stream</param>
-		/// <param name="RelativePath">Receives the relative path to the file</param>
+		/// <param name="depotFile">The depot file to check</param>
+		/// <param name="streamName">Name of the stream</param>
+		/// <param name="relativePath">Receives the relative path to the file</param>
 		/// <returns>True if the stream-relative path was returned</returns>
-		public static bool TryGetStreamRelativePath(string DepotFile, string StreamName, [NotNullWhen(true)] out string? RelativePath)
+		public static bool TryGetStreamRelativePath(string depotFile, string streamName, [NotNullWhen(true)] out string? relativePath)
 		{
-			if (DepotFile.StartsWith(StreamName, StringComparison.OrdinalIgnoreCase) && DepotFile.Length > StreamName.Length && DepotFile[StreamName.Length] == '/')
+			if (depotFile.StartsWith(streamName, StringComparison.OrdinalIgnoreCase) && depotFile.Length > streamName.Length && depotFile[streamName.Length] == '/')
 			{
-				RelativePath = DepotFile.Substring(StreamName.Length);
+				relativePath = depotFile.Substring(streamName.Length);
 				return true;
 			}
 
-			Match Match = Regex.Match(DepotFile, "^//[^/]+/[^/]+(/.*)$");
-			if (Match.Success)
+			Match match = Regex.Match(depotFile, "^//[^/]+/[^/]+(/.*)$");
+			if (match.Success)
 			{
-				RelativePath = Match.Groups[1].Value;
+				relativePath = match.Groups[1].Value;
 				return true;
 			}
 
-			RelativePath = null;
+			relativePath = null;
 			return false;
 		}
 	}
-
 }
