Horde: Fix coding conventions within Horde.Build.

diff --git a/Engine/Source/Programs/Horde/Horde.Build/Perforce/PerforceLoadBalancer.cs b/Engine/Source/Programs/Horde/Horde.Build/Perforce/PerforceLoadBalancer.cs
--- a/Engine/Source/Programs/Horde/Horde.Build/Perforce/PerforceLoadBalancer.cs
+++ b/Engine/Source/Programs/Horde/Horde.Build/Perforce/PerforceLoadBalancer.cs
@@ -1,18 +1,6 @@
 // Copyright Epic Games, Inc. All Rights Reserved.
 
-using Google.Protobuf.WellKnownTypes;
-using HordeCommon;
-using HordeCommon.Rpc.Tasks;
-using Horde.Build.Collections;
-using Horde.Build.Models;
-using Horde.Build.Utilities;
-using Microsoft.Extensions.Hosting;
-using Microsoft.Extensions.Logging;
-using MongoDB.Bson;
-using MongoDB.Bson.Serialization.Attributes;
-using MongoDB.Driver;
 using System;
-using System.Collections.Concurrent;
 using System.Collections.Generic;
 using System.Globalization;
 using System.Linq;
@@ -21,6 +9,16 @@ using System.Net.Http;
 using System.Text.Json;
 using System.Threading;
 using System.Threading.Tasks;
+using Google.Protobuf.WellKnownTypes;
+using Horde.Build.Collections;
+using Horde.Build.Models;
+using Horde.Build.Utilities;
+using HordeCommon;
+using HordeCommon.Rpc.Tasks;
+using Microsoft.Extensions.Hosting;
+using Microsoft.Extensions.Logging;
+using MongoDB.Bson.Serialization.Attributes;
+using MongoDB.Driver;
 
 namespace Horde.Build.Services
 {
@@ -115,15 +113,15 @@ namespace Horde.Build.Services
 			{
 			}
 
-			public PerforceServerEntry(string ServerAndPort, string BaseServerAndPort, string? HealthCheckUrl, string Cluster, PerforceServerStatus Status, string? Detail, DateTime? LastUpdateTime)
+			public PerforceServerEntry(string serverAndPort, string baseServerAndPort, string? healthCheckUrl, string cluster, PerforceServerStatus status, string? detail, DateTime? lastUpdateTime)
 			{
-				this.ServerAndPort = ServerAndPort;
-				this.BaseServerAndPort = BaseServerAndPort;
-				this.HealthCheckUrl = HealthCheckUrl;
-				this.Cluster = Cluster;
-				this.Status = Status;
-				this.Detail = Detail;
-				this.LastUpdateTime = LastUpdateTime;
+				ServerAndPort = serverAndPort;
+				BaseServerAndPort = baseServerAndPort;
+				HealthCheckUrl = healthCheckUrl;
+				Cluster = cluster;
+				Status = status;
+				Detail = detail;
+				LastUpdateTime = lastUpdateTime;
 			}
 		}
 
@@ -136,40 +134,40 @@ namespace Horde.Build.Services
 			public List<PerforceServerEntry> Servers { get; set; } = new List<PerforceServerEntry>();
 		}
 
-		DatabaseService DatabaseService;
-		ILeaseCollection LeaseCollection;
-		SingletonDocument<PerforceServerList> ServerListSingleton;
-		Random Random = new Random();
-		ILogger Logger;
-		ITicker Ticker;
+		readonly DatabaseService _databaseService;
+		readonly ILeaseCollection _leaseCollection;
+		readonly SingletonDocument<PerforceServerList> _serverListSingleton;
+		readonly Random _random = new Random();
+		readonly ILogger _logger;
+		readonly ITicker _ticker;
 
 		/// <summary>
 		/// Constructor
 		/// </summary>
-		public PerforceLoadBalancer(DatabaseService DatabaseService, ILeaseCollection LeaseCollection, IClock Clock, ILogger<PerforceLoadBalancer> Logger)
+		public PerforceLoadBalancer(DatabaseService databaseService, ILeaseCollection leaseCollection, IClock clock, ILogger<PerforceLoadBalancer> logger)
 		{
-			this.DatabaseService = DatabaseService;
-			this.LeaseCollection = LeaseCollection;
-			this.ServerListSingleton = new SingletonDocument<PerforceServerList>(DatabaseService);
-			this.Logger = Logger;
-			if (DatabaseService.ReadOnlyMode)
+			_databaseService = databaseService;
+			_leaseCollection = leaseCollection;
+			_serverListSingleton = new SingletonDocument<PerforceServerList>(databaseService);
+			_logger = logger;
+			if (databaseService.ReadOnlyMode)
 			{
-				this.Ticker = new NullTicker();
+				_ticker = new NullTicker();
 			}
 			else
 			{
-				this.Ticker = Clock.AddSharedTicker<PerforceLoadBalancer>(TimeSpan.FromMinutes(1.0), TickInternalAsync, Logger);
+				_ticker = clock.AddSharedTicker<PerforceLoadBalancer>(TimeSpan.FromMinutes(1.0), TickInternalAsync, logger);
 			}
 		}
 
 		/// <inheritdoc/>
-		public Task StartAsync(CancellationToken cancellationToken) => Ticker.StartAsync();
+		public Task StartAsync(CancellationToken cancellationToken) => _ticker.StartAsync();
 
 		/// <inheritdoc/>
-		public Task StopAsync(CancellationToken cancellationToken) => Ticker.StopAsync();
+		public Task StopAsync(CancellationToken cancellationToken) => _ticker.StopAsync();
 
 		/// <inheritdoc/>
-		public void Dispose() => Ticker.Dispose();
+		public void Dispose() => _ticker.Dispose();
 
 		/// <summary>
 		/// Get the current server list
@@ -177,19 +175,19 @@ namespace Horde.Build.Services
 		/// <returns></returns>
 		async Task<PerforceServerList> GetServerListAsync()
 		{
-			PerforceServerList ServerList = await ServerListSingleton.GetAsync();
+			PerforceServerList serverList = await _serverListSingleton.GetAsync();
 
-			DateTime MinLastUpdateTime = DateTime.UtcNow - TimeSpan.FromMinutes(2.5);
-			foreach (PerforceServerEntry Server in ServerList.Servers)
+			DateTime minLastUpdateTime = DateTime.UtcNow - TimeSpan.FromMinutes(2.5);
+			foreach (PerforceServerEntry server in serverList.Servers)
 			{
-				if (Server.Status == PerforceServerStatus.Healthy && Server.LastUpdateTime != null && Server.LastUpdateTime.Value < MinLastUpdateTime)
+				if (server.Status == PerforceServerStatus.Healthy && server.LastUpdateTime != null && server.LastUpdateTime.Value < minLastUpdateTime)
 				{
-					Server.Status = PerforceServerStatus.Degraded;
-					Server.Detail = "Server has not responded to health check";
+					server.Status = PerforceServerStatus.Degraded;
+					server.Detail = "Server has not responded to health check";
 				}
 			}
 
-			return ServerList;
+			return serverList;
 		}
 
 		/// <summary>
@@ -198,23 +196,23 @@ namespace Horde.Build.Services
 		/// <returns></returns>
 		public async Task<List<IPerforceServer>> GetServersAsync()
 		{
-			PerforceServerList ServerList = await GetServerListAsync();
-			return ServerList.Servers.ConvertAll<IPerforceServer>(x => x);
+			PerforceServerList serverList = await GetServerListAsync();
+			return serverList.Servers.ConvertAll<IPerforceServer>(x => x);
 		}
 
 		/// <summary>
 		/// Allocates a server for use by a lease
 		/// </summary>
 		/// <returns>The server to use. Null if there is no healthy server available.</returns>
-		public async Task<IPerforceServer?> GetServer(string Cluster)
+		public async Task<IPerforceServer?> GetServer(string cluster)
 		{
-			PerforceServerList ServerList = await GetServerListAsync();
+			PerforceServerList serverList = await GetServerListAsync();
 
-			List<PerforceServerEntry> Candidates = ServerList.Servers.Where(x => x.Cluster == Cluster && x.Status >= PerforceServerStatus.Healthy).ToList();
-			if(Candidates.Count == 0)
+			List<PerforceServerEntry> candidates = serverList.Servers.Where(x => x.Cluster == cluster && x.Status >= PerforceServerStatus.Healthy).ToList();
+			if(candidates.Count == 0)
 			{
-				int Idx = Random.Next(0, Candidates.Count);
-				return Candidates[Idx];
+				int idx = _random.Next(0, candidates.Count);
+				return candidates[idx];
 			}
 			return null;
 		}
@@ -222,265 +220,265 @@ namespace Horde.Build.Services
 		/// <summary>
 		/// Select a Perforce server to use by the Horde server
 		/// </summary>
-		/// <param name="Cluster"></param>
+		/// <param name="cluster"></param>
 		/// <returns></returns>
-		public Task<IPerforceServer?> SelectServerAsync(PerforceCluster Cluster)
+		public Task<IPerforceServer?> SelectServerAsync(PerforceCluster cluster)
 		{
-			List<string> Properties = new List<string>{ "HordeServer=1" };
-			return SelectServerAsync(Cluster, Properties);
+			List<string> properties = new List<string>{ "HordeServer=1" };
+			return SelectServerAsync(cluster, properties);
 		}
 
 		/// <summary>
 		/// Select a Perforce server to use
 		/// </summary>
-		/// <param name="Cluster"></param>
-		/// <param name="Agent"></param>
+		/// <param name="cluster"></param>
+		/// <param name="agent"></param>
 		/// <returns></returns>
-		public Task<IPerforceServer?> SelectServerAsync(PerforceCluster Cluster, IAgent Agent)
+		public Task<IPerforceServer?> SelectServerAsync(PerforceCluster cluster, IAgent agent)
 		{
-			return SelectServerAsync(Cluster, Agent.Properties);
+			return SelectServerAsync(cluster, agent.Properties);
 		}
 
 		/// <summary>
 		/// Select a Perforce server to use
 		/// </summary>
-		/// <param name="Cluster"></param>
-		/// <param name="Properties"></param>
+		/// <param name="cluster"></param>
+		/// <param name="properties"></param>
 		/// <returns></returns>
-		public async Task<IPerforceServer?> SelectServerAsync(PerforceCluster Cluster, IReadOnlyList<string> Properties)
+		public async Task<IPerforceServer?> SelectServerAsync(PerforceCluster cluster, IReadOnlyList<string> properties)
 		{
 			// Find all the valid servers for this agent
-			List<PerforceServer> ValidServers = new List<PerforceServer>();
-			if (Properties != null)
+			List<PerforceServer> validServers = new List<PerforceServer>();
+			if (properties != null)
 			{
-				ValidServers.AddRange(Cluster.Servers.Where(x => x.Properties != null && x.Properties.Count > 0 && x.Properties.All(y => Properties.Contains(y))));
+				validServers.AddRange(cluster.Servers.Where(x => x.Properties != null && x.Properties.Count > 0 && x.Properties.All(y => properties.Contains(y))));
 			}
-			if (ValidServers.Count == 0)
+			if (validServers.Count == 0)
 			{
-				ValidServers.AddRange(Cluster.Servers.Where(x => x.Properties == null || x.Properties.Count == 0));
+				validServers.AddRange(cluster.Servers.Where(x => x.Properties == null || x.Properties.Count == 0));
 			}
 
-			HashSet<string> ValidServerNames = new HashSet<string>(ValidServers.Select(x => x.ServerAndPort), StringComparer.OrdinalIgnoreCase);
+			HashSet<string> validServerNames = new HashSet<string>(validServers.Select(x => x.ServerAndPort), StringComparer.OrdinalIgnoreCase);
 
 			// Find all the matching servers.
-			PerforceServerList ServerList = await GetServerListAsync();
+			PerforceServerList serverList = await GetServerListAsync();
 
-			List<PerforceServerEntry> Candidates = ServerList.Servers.Where(x => x.Cluster == Cluster.Name && ValidServerNames.Contains(x.BaseServerAndPort)).ToList();
-			if (Candidates.Count == 0)
+			List<PerforceServerEntry> candidates = serverList.Servers.Where(x => x.Cluster == cluster.Name && validServerNames.Contains(x.BaseServerAndPort)).ToList();
+			if (candidates.Count == 0)
 			{
-				foreach (PerforceServer ValidServer in ValidServers)
+				foreach (PerforceServer validServer in validServers)
 				{
-					Logger.LogDebug("Fetching server info for {ServerAndPort}", ValidServer.ServerAndPort);
-					await UpdateServerAsync(Cluster, ValidServer, Candidates);
+					_logger.LogDebug("Fetching server info for {ServerAndPort}", validServer.ServerAndPort);
+					await UpdateServerAsync(cluster, validServer, candidates);
 				}
-				if (Candidates.Count == 0)
+				if (candidates.Count == 0)
 				{
-					Logger.LogWarning("Unable to resolve any Perforce servers from valid list");
+					_logger.LogWarning("Unable to resolve any Perforce servers from valid list");
 					return null; 
 				}
 			}
 
 			// Remove any servers that are unhealthy
-			if (Candidates.Any(x => x.Status == PerforceServerStatus.Healthy))
+			if (candidates.Any(x => x.Status == PerforceServerStatus.Healthy))
 			{
-				Candidates.RemoveAll(x => x.Status != PerforceServerStatus.Healthy);
+				candidates.RemoveAll(x => x.Status != PerforceServerStatus.Healthy);
 			}
 			else
 			{
-				Candidates.RemoveAll(x => x.Status != PerforceServerStatus.Unknown);
+				candidates.RemoveAll(x => x.Status != PerforceServerStatus.Unknown);
 			}
 
-			if (Candidates.Count == 0)
+			if (candidates.Count == 0)
 			{
-				Logger.LogWarning("Unable to find any healthy Perforce server in cluster {ClusterName}", Cluster.Name);
+				_logger.LogWarning("Unable to find any healthy Perforce server in cluster {ClusterName}", cluster.Name);
 				return null;
 			}
 
 			// Select which server to use with a weighted average of the number of active leases
-			int Index = 0;
-			if (Candidates.Count > 1)
+			int index = 0;
+			if (candidates.Count > 1)
 			{
 				const int BaseWeight = 20;
 
-				int TotalLeases = Candidates.Sum(x => x.NumLeases);
-				int TotalWeight = Candidates.Sum(x => (TotalLeases + BaseWeight) - x.NumLeases);
+				int totalLeases = candidates.Sum(x => x.NumLeases);
+				int totalWeight = candidates.Sum(x => (totalLeases + BaseWeight) - x.NumLeases);
 
-				int Weight = Random.Next(TotalWeight);
-				for (; Index + 1 < Candidates.Count; Index++)
+				int weight = _random.Next(totalWeight);
+				for (; index + 1 < candidates.Count; index++)
 				{
-					Weight -= (TotalLeases + BaseWeight) - Candidates[Index].NumLeases;
-					if(Weight < 0)
+					weight -= (totalLeases + BaseWeight) - candidates[index].NumLeases;
+					if(weight < 0)
 					{
 						break;
 					}
 				}
 			}
 
-			return Candidates[Index];
+			return candidates[index];
 		}
 
 		/// <inheritdoc/>
-		async ValueTask TickInternalAsync(CancellationToken CancellationToken)
+		async ValueTask TickInternalAsync(CancellationToken cancellationToken)
 		{
-			Globals Globals = await DatabaseService.GetGlobalsAsync();
+			Globals globals = await _databaseService.GetGlobalsAsync();
 
 			// Set of new server entries
-			List<PerforceServerEntry> NewServers = new List<PerforceServerEntry>();
+			List<PerforceServerEntry> newServers = new List<PerforceServerEntry>();
 
 			// Update the state of all the valid servers
-			foreach (PerforceCluster Cluster in Globals.PerforceClusters)
+			foreach (PerforceCluster cluster in globals.PerforceClusters)
 			{
-				foreach (PerforceServer Server in Cluster.Servers)
+				foreach (PerforceServer server in cluster.Servers)
 				{
 					try
 					{
-						await UpdateServerAsync(Cluster, Server, NewServers);
+						await UpdateServerAsync(cluster, server, newServers);
 					}
-					catch (Exception Ex)
+					catch (Exception ex)
 					{
-						Logger.LogError(Ex, "Exception while updating Perforce server status");
+						_logger.LogError(ex, "Exception while updating Perforce server status");
 					}
 				}
 			}
 
 			// Update the number of leases for each entry
-			List<PerforceServerEntry> NewEntries = NewServers.OrderBy(x => x.Cluster).ThenBy(x => x.BaseServerAndPort).ThenBy(x => x.ServerAndPort).ToList();
-			await UpdateLeaseCounts(NewEntries);
-			PerforceServerList List = await ServerListSingleton.UpdateAsync(List => MergeServerList(List, NewEntries));
+			List<PerforceServerEntry> newEntries = newServers.OrderBy(x => x.Cluster).ThenBy(x => x.BaseServerAndPort).ThenBy(x => x.ServerAndPort).ToList();
+			await UpdateLeaseCounts(newEntries);
+			PerforceServerList list = await _serverListSingleton.UpdateAsync(list => MergeServerList(list, newEntries));
 
 			// Now update the health of each entry
-			List<Task> Tasks = new List<Task>();
-			foreach (PerforceServerEntry Entry in List.Servers)
+			List<Task> tasks = new List<Task>();
+			foreach (PerforceServerEntry entry in list.Servers)
 			{
-				Tasks.Add(Task.Run(() => UpdateHealthAsync(Entry), CancellationToken.None));
+				tasks.Add(Task.Run(() => UpdateHealthAsync(entry), CancellationToken.None));
 			}
-			await Task.WhenAll(Tasks);
+			await Task.WhenAll(tasks);
 		}
 
-		static void MergeServerList(PerforceServerList ServerList, List<PerforceServerEntry> NewEntries)
+		static void MergeServerList(PerforceServerList serverList, List<PerforceServerEntry> newEntries)
 		{
-			Dictionary<string, PerforceServerEntry> ExistingEntries = ServerList.Servers.ToDictionary(x => x.ServerAndPort, x => x, StringComparer.OrdinalIgnoreCase);
-			foreach (PerforceServerEntry NewEntry in NewEntries)
+			Dictionary<string, PerforceServerEntry> existingEntries = serverList.Servers.ToDictionary(x => x.ServerAndPort, x => x, StringComparer.OrdinalIgnoreCase);
+			foreach (PerforceServerEntry newEntry in newEntries)
 			{
-				PerforceServerEntry? ExistingEntry;
-				if (ExistingEntries.TryGetValue(NewEntry.ServerAndPort, out ExistingEntry))
+				PerforceServerEntry? existingEntry;
+				if (existingEntries.TryGetValue(newEntry.ServerAndPort, out existingEntry))
 				{
-					NewEntry.Status = ExistingEntry.Status;
-					NewEntry.Detail = ExistingEntry.Detail;
-					NewEntry.LastUpdateTime = ExistingEntry.LastUpdateTime;
+					newEntry.Status = existingEntry.Status;
+					newEntry.Detail = existingEntry.Detail;
+					newEntry.LastUpdateTime = existingEntry.LastUpdateTime;
 				}
 			}
-			ServerList.Servers = NewEntries;
+			serverList.Servers = newEntries;
 		}
 
-		async Task UpdateLeaseCounts(IEnumerable<PerforceServerEntry> NewServerEntries)
+		async Task UpdateLeaseCounts(IEnumerable<PerforceServerEntry> newServerEntries)
 		{
-			Dictionary<string, Dictionary<string, PerforceServerEntry>> NewServerLookup = new Dictionary<string, Dictionary<string, PerforceServerEntry>>();
-			foreach (PerforceServerEntry NewServerEntry in NewServerEntries)
+			Dictionary<string, Dictionary<string, PerforceServerEntry>> newServerLookup = new Dictionary<string, Dictionary<string, PerforceServerEntry>>();
+			foreach (PerforceServerEntry newServerEntry in newServerEntries)
 			{
-				if (NewServerEntry.Cluster != null)
+				if (newServerEntry.Cluster != null)
 				{
-					Dictionary<string, PerforceServerEntry>? ClusterServers;
-					if (!NewServerLookup.TryGetValue(NewServerEntry.Cluster, out ClusterServers))
+					Dictionary<string, PerforceServerEntry>? clusterServers;
+					if (!newServerLookup.TryGetValue(newServerEntry.Cluster, out clusterServers))
 					{
-						ClusterServers = new Dictionary<string, PerforceServerEntry>();
-						NewServerLookup.Add(NewServerEntry.Cluster, ClusterServers);
+						clusterServers = new Dictionary<string, PerforceServerEntry>();
+						newServerLookup.Add(newServerEntry.Cluster, clusterServers);
 					}
-					ClusterServers[NewServerEntry.ServerAndPort] = NewServerEntry;
+					clusterServers[newServerEntry.ServerAndPort] = newServerEntry;
 				}
 			}
 
-			List<ILease> Leases = await LeaseCollection.FindActiveLeasesAsync();
-			foreach (ILease Lease in Leases)
+			List<ILease> leases = await _leaseCollection.FindActiveLeasesAsync();
+			foreach (ILease lease in leases)
 			{
-				Any Any = Any.Parser.ParseFrom(Lease.Payload.ToArray());
+				Any any = Any.Parser.ParseFrom(lease.Payload.ToArray());
 
-				HashSet<(string, string)> Servers = new HashSet<(string, string)>();
-				if (Any.TryUnpack(out ConformTask ConformTask))
+				HashSet<(string, string)> servers = new HashSet<(string, string)>();
+				if (any.TryUnpack(out ConformTask conformTask))
 				{
-					foreach (HordeCommon.Rpc.Messages.AgentWorkspace Workspace in ConformTask.Workspaces)
+					foreach (HordeCommon.Rpc.Messages.AgentWorkspace workspace in conformTask.Workspaces)
 					{
-						Servers.Add((Workspace.Cluster, Workspace.ServerAndPort));
+						servers.Add((workspace.Cluster, workspace.ServerAndPort));
 					}
 				}
-				if (Any.TryUnpack(out ExecuteJobTask ExecuteJobTask))
+				if (any.TryUnpack(out ExecuteJobTask executeJobTask))
 				{
-					if (ExecuteJobTask.AutoSdkWorkspace != null)
+					if (executeJobTask.AutoSdkWorkspace != null)
 					{
-						Servers.Add((ExecuteJobTask.AutoSdkWorkspace.Cluster, ExecuteJobTask.AutoSdkWorkspace.ServerAndPort));
+						servers.Add((executeJobTask.AutoSdkWorkspace.Cluster, executeJobTask.AutoSdkWorkspace.ServerAndPort));
 					}
-					if (ExecuteJobTask.Workspace != null)
+					if (executeJobTask.Workspace != null)
 					{
-						Servers.Add((ExecuteJobTask.Workspace.Cluster, ExecuteJobTask.Workspace.ServerAndPort));
+						servers.Add((executeJobTask.Workspace.Cluster, executeJobTask.Workspace.ServerAndPort));
 					}
 				}
 
-				foreach ((string Cluster, string ServerAndPort) in Servers)
+				foreach ((string cluster, string serverAndPort) in servers)
 				{
-					IncrementLeaseCount(Cluster, ServerAndPort, NewServerLookup);
+					IncrementLeaseCount(cluster, serverAndPort, newServerLookup);
 				}
 			}
 		}
 
-		static void IncrementLeaseCount(string Cluster, string ServerAndPort, Dictionary<string, Dictionary<string, PerforceServerEntry>> NewServers)
+		static void IncrementLeaseCount(string cluster, string serverAndPort, Dictionary<string, Dictionary<string, PerforceServerEntry>> newServers)
 		{
-			Dictionary<string, PerforceServerEntry>? ClusterServers;
-			if (Cluster != null && NewServers.TryGetValue(Cluster, out ClusterServers))
+			Dictionary<string, PerforceServerEntry>? clusterServers;
+			if (cluster != null && newServers.TryGetValue(cluster, out clusterServers))
 			{
-				PerforceServerEntry? Entry;
-				if (ClusterServers.TryGetValue(ServerAndPort, out Entry))
+				PerforceServerEntry? entry;
+				if (clusterServers.TryGetValue(serverAndPort, out entry))
 				{
-					Entry.NumLeases++;
+					entry.NumLeases++;
 				}
 			}
 		}
 
-		async Task UpdateServerAsync(PerforceCluster Cluster, PerforceServer Server, List<PerforceServerEntry> NewServers)
+		async Task UpdateServerAsync(PerforceCluster cluster, PerforceServer server, List<PerforceServerEntry> newServers)
 		{
-			string InitialHostName = Server.ServerAndPort;
-			int Port = 1666;
+			string initialHostName = server.ServerAndPort;
+			int port = 1666;
 
-			int PortIdx = InitialHostName.LastIndexOf(':');
-			if (PortIdx != -1)
+			int portIdx = initialHostName.LastIndexOf(':');
+			if (portIdx != -1)
 			{
-				Port = int.Parse(InitialHostName.Substring(PortIdx + 1), NumberStyles.Integer, CultureInfo.InvariantCulture);
-				InitialHostName = InitialHostName.Substring(0, PortIdx);
+				port = Int32.Parse(initialHostName.Substring(portIdx + 1), NumberStyles.Integer, CultureInfo.InvariantCulture);
+				initialHostName = initialHostName.Substring(0, portIdx);
 			}
 
-			List<string> HostNames = new List<string>();
-			if (!Server.ResolveDns || IPAddress.TryParse(InitialHostName, out IPAddress? HostAddress))
+			List<string> hostNames = new List<string>();
+			if (!server.ResolveDns || IPAddress.TryParse(initialHostName, out IPAddress? _))
 			{
-				HostNames.Add(InitialHostName);
+				hostNames.Add(initialHostName);
 			}
 			else
 			{
-				await ResolveServersAsync(InitialHostName, HostNames);
+				await ResolveServersAsync(initialHostName, hostNames);
 			}
 
-			foreach (string HostName in HostNames)
+			foreach (string hostName in hostNames)
 			{
-				string? HealthCheckUrl = null;
-				if (Server.HealthCheck)
+				string? healthCheckUrl = null;
+				if (server.HealthCheck)
 				{
-					HealthCheckUrl = $"http://{HostName}:5000/healthcheck";
+					healthCheckUrl = $"http://{hostName}:5000/healthcheck";
 				}
-				NewServers.Add(new PerforceServerEntry($"{HostName}:{Port}", Server.ServerAndPort, HealthCheckUrl, Cluster.Name, PerforceServerStatus.Unknown, "", DateTime.UtcNow));
+				newServers.Add(new PerforceServerEntry($"{hostName}:{port}", server.ServerAndPort, healthCheckUrl, cluster.Name, PerforceServerStatus.Unknown, "", DateTime.UtcNow));
 			}
 		}
 
-		async Task ResolveServersAsync(string HostName, List<string> HostNames)
+		async Task ResolveServersAsync(string hostName, List<string> hostNames)
 		{
 			// Find all the addresses of the hosts
-			IPHostEntry Entry = await Dns.GetHostEntryAsync(HostName);
-			foreach (IPAddress Address in Entry.AddressList)
+			IPHostEntry entry = await Dns.GetHostEntryAsync(hostName);
+			foreach (IPAddress address in entry.AddressList)
 			{
 				try
 				{
-					HostNames.Add(Address.ToString());
+					hostNames.Add(address.ToString());
 				}
-				catch (Exception Ex)
+				catch (Exception ex)
 				{
-					Logger.LogDebug(Ex, "Unable to resolve host name for Perforce server {Address}", Address);
+					_logger.LogDebug(ex, "Unable to resolve host name for Perforce server {Address}", address);
 				}
 			}
 		}
@@ -488,66 +486,66 @@ namespace Horde.Build.Services
 		/// <summary>
 		/// Updates 
 		/// </summary>
-		/// <param name="Entry"></param>
+		/// <param name="entry"></param>
 		/// <returns></returns>
-		async Task UpdateHealthAsync(PerforceServerEntry Entry)
+		async Task UpdateHealthAsync(PerforceServerEntry entry)
 		{
-			DateTime? UpdateTime = null;
-			string Detail = "Health check disabled";
+			DateTime? updateTime = null;
+			string detail = "Health check disabled";
 
 			// Get the health of the server
-			PerforceServerStatus Health = PerforceServerStatus.Healthy;
-			if (Entry.HealthCheckUrl != null)
+			PerforceServerStatus health = PerforceServerStatus.Healthy;
+			if (entry.HealthCheckUrl != null)
 			{
-				UpdateTime = DateTime.UtcNow;
-				Uri HealthCheckUrl = new Uri(Entry.HealthCheckUrl);
+				updateTime = DateTime.UtcNow;
+				Uri healthCheckUrl = new Uri(entry.HealthCheckUrl);
 				try
 				{
-					(Health, Detail) = await GetServerHealthAsync(HealthCheckUrl);
+					(health, detail) = await GetServerHealthAsync(healthCheckUrl);
 				}
 				catch
 				{
-					(Health, Detail) = (PerforceServerStatus.Unhealthy, $"Failed to query status at {HealthCheckUrl}");
+					(health, detail) = (PerforceServerStatus.Unhealthy, $"Failed to query status at {healthCheckUrl}");
 				}
 			}
 
 			// Update the server record
-			if (Health != Entry.Status || Detail != Entry.Detail || UpdateTime != Entry.LastUpdateTime)
+			if (health != entry.Status || detail != entry.Detail || updateTime != entry.LastUpdateTime)
 			{
-				await ServerListSingleton.UpdateAsync(x => UpdateHealth(x, Entry.ServerAndPort, Health, Detail, UpdateTime));
+				await _serverListSingleton.UpdateAsync(x => UpdateHealth(x, entry.ServerAndPort, health, detail, updateTime));
 			}
 		}
 
-		static void UpdateHealth(PerforceServerList ServerList, string ServerAndPort, PerforceServerStatus Status, string Detail, DateTime? UpdateTime)
+		static void UpdateHealth(PerforceServerList serverList, string serverAndPort, PerforceServerStatus status, string detail, DateTime? updateTime)
 		{
-			PerforceServerEntry? Entry = ServerList.Servers.FirstOrDefault(x => x.ServerAndPort.Equals(ServerAndPort, StringComparison.Ordinal));
-			if (Entry != null)
+			PerforceServerEntry? entry = serverList.Servers.FirstOrDefault(x => x.ServerAndPort.Equals(serverAndPort, StringComparison.Ordinal));
+			if (entry != null)
 			{
-				if (Entry.LastUpdateTime == null || UpdateTime == null || Entry.LastUpdateTime.Value < UpdateTime.Value)
+				if (entry.LastUpdateTime == null || updateTime == null || entry.LastUpdateTime.Value < updateTime.Value)
 				{
-					Entry.Status = Status;
-					Entry.Detail = Detail;
-					Entry.LastUpdateTime = UpdateTime;
+					entry.Status = status;
+					entry.Detail = detail;
+					entry.LastUpdateTime = updateTime;
 				}
 			}
 		}
 
-		static async Task<(PerforceServerStatus, string)> GetServerHealthAsync(Uri HealthCheckUrl)
+		static async Task<(PerforceServerStatus, string)> GetServerHealthAsync(Uri healthCheckUrl)
 		{
-			using HttpClient Client = new HttpClient();
-			HttpResponseMessage Response = await Client.GetAsync(HealthCheckUrl);
+			using HttpClient client = new HttpClient();
+			HttpResponseMessage response = await client.GetAsync(healthCheckUrl);
 
-			byte[] Data = await Response.Content.ReadAsByteArrayAsync();
-			JsonDocument Document = JsonDocument.Parse(Data);
+			byte[] data = await response.Content.ReadAsByteArrayAsync();
+			JsonDocument document = JsonDocument.Parse(data);
 
-			foreach (JsonElement Element in Document.RootElement.GetProperty("results").EnumerateArray())
+			foreach (JsonElement element in document.RootElement.GetProperty("results").EnumerateArray())
 			{
-				if (Element.TryGetProperty("checker", out JsonElement Checker) && Checker.ValueEquals("edge_traffic_lights"))
+				if (element.TryGetProperty("checker", out JsonElement checker) && checker.ValueEquals("edge_traffic_lights"))
 				{
-					if (Element.TryGetProperty("output", out JsonElement Output))
+					if (element.TryGetProperty("output", out JsonElement output))
 					{
-						string Status = Output.GetString() ?? String.Empty;
-						switch(Status)
+						string status = output.GetString() ?? String.Empty;
+						switch(status)
 						{
 							case "green":
 								return (PerforceServerStatus.Healthy, "Server is healthy");
@@ -556,7 +554,7 @@ namespace Horde.Build.Services
 							case "red":
 								return (PerforceServerStatus.Unhealthy, "Server is being drained");
 							default:
-								return (PerforceServerStatus.Unknown, $"Expected state for health check ({Status})");
+								return (PerforceServerStatus.Unknown, $"Expected state for health check ({status})");
 						}
 					}
 				}
diff --git a/Engine/Source/Programs/Horde/Horde.Build/Program.cs b/Engine/Source/Programs/Horde/Horde.Build/Program.cs
--- a/Engine/Source/Programs/Horde/Horde.Build/Program.cs
+++ b/Engine/Source/Programs/Horde/Horde.Build/Program.cs
@@ -1,126 +1,101 @@
 // Copyright Epic Games, Inc. All Rights Reserved.
 
 using System;
-using System.Collections.Concurrent;
 using System.Collections.Generic;
-using System.Diagnostics;
-using System.Globalization;
 using System.IO;
-using System.Linq;
-using System.Net;
 using System.Reflection;
 using System.Runtime.InteropServices;
-using System.Security.Cryptography;
-using System.Security.Cryptography.X509Certificates;
-using System.Text;
-using System.Text.Json;
-using System.Threading;
-using System.Threading.Channels;
-using System.Text.RegularExpressions;
 using System.Threading.Tasks;
 using EpicGames.Core;
-using EpicGames.Perforce;
-using HordeCommon;
-using Horde.Build.Models;
-using Horde.Build.Utilities;
-using Microsoft.AspNetCore.Hosting;
-using Microsoft.AspNetCore.Server.Kestrel.Core;
+using Horde.Build.Commands;
 using Microsoft.Extensions.Configuration;
+using Microsoft.Extensions.DependencyInjection;
 using Microsoft.Extensions.Hosting;
-using Microsoft.Extensions.Logging;
-using Microsoft.Extensions.Logging.Abstractions;
-using MongoDB.Bson;
-using MongoDB.Bson.Serialization;
-using MongoDB.Driver;
 using OpenTracing;
+using OpenTracing.Propagation;
 using OpenTracing.Util;
 using Serilog;
 using Serilog.Configuration;
 using Serilog.Core;
-using Serilog.Enrichers.OpenTracing;
 using Serilog.Events;
-using Serilog.Filters;
 using Serilog.Formatting.Json;
 using Serilog.Sinks.SystemConsole.Themes;
-using OpenTracing.Propagation;
-using Microsoft.Extensions.DependencyInjection;
-using Horde.Build.Commands;
 
 namespace Horde.Build
 {
 	static class LoggerExtensions
 	{
-		public static LoggerConfiguration Console(this LoggerSinkConfiguration SinkConfig, ServerSettings Settings)
+		public static LoggerConfiguration Console(this LoggerSinkConfiguration sinkConfig, ServerSettings settings)
 		{
-			if (Settings.LogJsonToStdOut)
+			if (settings.LogJsonToStdOut)
 			{
-				return SinkConfig.Console(new JsonFormatter(renderMessage: true));
+				return sinkConfig.Console(new JsonFormatter(renderMessage: true));
 			}
 			else
 			{
-				ConsoleTheme Theme;
+				ConsoleTheme theme;
 				if (RuntimeInformation.IsOSPlatform(OSPlatform.Windows) && Environment.OSVersion.Version < new Version(10, 0))
 				{
-					Theme = SystemConsoleTheme.Literate;
+					theme = SystemConsoleTheme.Literate;
 				}
 				else
 				{
-					Theme = AnsiConsoleTheme.Code;
+					theme = AnsiConsoleTheme.Code;
 				}
-				return SinkConfig.Console(outputTemplate: "[{Timestamp:HH:mm:ss} {Level:w3}] {Indent}{Message:l}{NewLine}{Exception}", theme: Theme, restrictedToMinimumLevel: LogEventLevel.Debug);
+				return sinkConfig.Console(outputTemplate: "[{Timestamp:HH:mm:ss} {Level:w3}] {Indent}{Message:l}{NewLine}{Exception}", theme: theme, restrictedToMinimumLevel: LogEventLevel.Debug);
 			}
 		}
 
-		public static LoggerConfiguration WithHordeConfig(this LoggerConfiguration Configuration, ServerSettings Settings)
+		public static LoggerConfiguration WithHordeConfig(this LoggerConfiguration configuration, ServerSettings settings)
 		{
-			if (Settings.WithDatadog)
+			if (settings.WithDatadog)
 			{
-				Configuration = Configuration.Enrich.With<DatadogLogEnricher>();
+				configuration = configuration.Enrich.With<DatadogLogEnricher>();
 			}
-			return Configuration;
+			return configuration;
 		}
 	}
 
 	class DatadogLogEnricher : ILogEventEnricher
 	{
-		public void Enrich(Serilog.Events.LogEvent LogEvent, ILogEventPropertyFactory PropertyFactory)
+		public void Enrich(Serilog.Events.LogEvent logEvent, ILogEventPropertyFactory propertyFactory)
 		{
-			ISpan? Span = GlobalTracer.Instance?.ActiveSpan;
-			if (Span != null)
+			ISpan? span = GlobalTracer.Instance?.ActiveSpan;
+			if (span != null)
 			{
-				LogEvent.AddPropertyIfAbsent(PropertyFactory.CreateProperty("dd.trace_id", Span.Context.TraceId));
-				LogEvent.AddPropertyIfAbsent(PropertyFactory.CreateProperty("dd.span_id", Span.Context.SpanId));
+				logEvent.AddPropertyIfAbsent(propertyFactory.CreateProperty("dd.trace_id", span.Context.TraceId));
+				logEvent.AddPropertyIfAbsent(propertyFactory.CreateProperty("dd.span_id", span.Context.SpanId));
 			}
 		}
 	}
 
 	class TestTracer : ITracer
 	{
-		ITracer Inner;
+		readonly ITracer _inner;
 
-		public TestTracer(ITracer Inner)
+		public TestTracer(ITracer inner)
 		{
-			this.Inner = Inner;
+			_inner = inner;
 		}
 
-		public IScopeManager ScopeManager => Inner.ScopeManager;
+		public IScopeManager ScopeManager => _inner.ScopeManager;
 
-		public ISpan ActiveSpan => Inner.ActiveSpan;
+		public ISpan ActiveSpan => _inner.ActiveSpan;
 
 		public ISpanBuilder BuildSpan(string operationName)
 		{
 			Serilog.Log.Debug("Creating span {Name}", operationName);
-			return Inner.BuildSpan(operationName);
+			return _inner.BuildSpan(operationName);
 		}
 
 		public ISpanContext Extract<TCarrier>(IFormat<TCarrier> format, TCarrier carrier)
 		{
-			return Inner.Extract<TCarrier>(format, carrier);
+			return _inner.Extract<TCarrier>(format, carrier);
 		}
 
 		public void Inject<TCarrier>(ISpanContext spanContext, IFormat<TCarrier> format, TCarrier carrier)
 		{
-			Inner.Inject<TCarrier>(spanContext, format, carrier);
+			_inner.Inject<TCarrier>(spanContext, format, carrier);
 		}
 	}
 
@@ -136,22 +111,22 @@ namespace Horde.Build
 
 		static Type[] FindSchemaTypes()
 		{
-			List<Type> SchemaTypes = new List<Type>();
-			foreach (Type Type in Assembly.GetExecutingAssembly().GetTypes())
+			List<Type> schemaTypes = new List<Type>();
+			foreach (Type type in Assembly.GetExecutingAssembly().GetTypes())
 			{
-				if (Type.GetCustomAttribute<JsonSchemaAttribute>() != null)
+				if (type.GetCustomAttribute<JsonSchemaAttribute>() != null)
 				{
-					SchemaTypes.Add(Type);
+					schemaTypes.Add(type);
 				}
 			}
-			return SchemaTypes.ToArray();
+			return schemaTypes.ToArray();
 		}
 
-		public static async Task<int> Main(string[] Args)
+		public static async Task<int> Main(string[] args)
 		{
-			CommandLineArguments Arguments = new CommandLineArguments(Args);
+			CommandLineArguments arguments = new CommandLineArguments(args);
 
-			IConfiguration Config = new ConfigurationBuilder()
+			IConfiguration config = new ConfigurationBuilder()
 				.SetBasePath(AppDir.FullName)
 				.AddJsonFile("appsettings.json", optional: false) 
 				.AddJsonFile("appsettings.Build.json", optional: true) // specific settings for builds (installer/dockerfile)
@@ -161,50 +136,50 @@ namespace Horde.Build
 				.AddEnvironmentVariables()
 				.Build();
 
-			ServerSettings HordeSettings = new ServerSettings();
-			Config.GetSection("Horde").Bind(HordeSettings);
+			ServerSettings hordeSettings = new ServerSettings();
+			config.GetSection("Horde").Bind(hordeSettings);
 
-			InitializeDefaults(HordeSettings);
+			InitializeDefaults(hordeSettings);
 
-			DirectoryReference LogDir = AppDir;
+			DirectoryReference logDir = AppDir;
 			if (RuntimeInformation.IsOSPlatform(OSPlatform.Windows))
 			{
-				LogDir = DirectoryReference.Combine(DataDir);
+				logDir = DirectoryReference.Combine(DataDir);
 			}
 
 			Serilog.Log.Logger = new LoggerConfiguration()
-				.WithHordeConfig(HordeSettings)
+				.WithHordeConfig(hordeSettings)
 				.Enrich.FromLogContext()
-				.WriteTo.Console(HordeSettings)
-				.WriteTo.File(Path.Combine(LogDir.FullName, "Log.txt"), outputTemplate: "[{Timestamp:HH:mm:ss} {Level:w3}] {Indent}{Message:l}{NewLine}{Exception} [{SourceContext}]", rollingInterval: RollingInterval.Day, rollOnFileSizeLimit: true, fileSizeLimitBytes: 20 * 1024 * 1024, retainedFileCountLimit: 10)
-				.WriteTo.File(new JsonFormatter(renderMessage: true), Path.Combine(LogDir.FullName, "Log.json"), rollingInterval: RollingInterval.Day, rollOnFileSizeLimit: true, fileSizeLimitBytes: 20 * 1024 * 1024, retainedFileCountLimit: 10)
-				.ReadFrom.Configuration(Config)
+				.WriteTo.Console(hordeSettings)
+				.WriteTo.File(Path.Combine(logDir.FullName, "Log.txt"), outputTemplate: "[{Timestamp:HH:mm:ss} {Level:w3}] {Indent}{Message:l}{NewLine}{Exception} [{SourceContext}]", rollingInterval: RollingInterval.Day, rollOnFileSizeLimit: true, fileSizeLimitBytes: 20 * 1024 * 1024, retainedFileCountLimit: 10)
+				.WriteTo.File(new JsonFormatter(renderMessage: true), Path.Combine(logDir.FullName, "Log.json"), rollingInterval: RollingInterval.Day, rollOnFileSizeLimit: true, fileSizeLimitBytes: 20 * 1024 * 1024, retainedFileCountLimit: 10)
+				.ReadFrom.Configuration(config)
 				.CreateLogger();
 
-			if (HordeSettings.WithDatadog)
+			if (hordeSettings.WithDatadog)
 			{
-				using var _ = Datadog.Trace.Tracer.Instance.StartActive("Trace Test");
+				using Datadog.Trace.IScope? _ = Datadog.Trace.Tracer.Instance.StartActive("Trace Test");
 				Serilog.Log.Logger.Information("Enabling datadog tracing");
 				GlobalTracer.Register(Datadog.Trace.OpenTracing.OpenTracingTracerFactory.WrapTracer(Datadog.Trace.Tracer.Instance));
-				using IScope Scope = GlobalTracer.Instance.BuildSpan("OpenTrace Test").StartActive();
-				Scope.Span.SetTag("TestProp", "hello");
+				using IScope scope = GlobalTracer.Instance.BuildSpan("OpenTrace Test").StartActive();
+				scope.Span.SetTag("TestProp", "hello");
 				Serilog.Log.Logger.Information("Enabling datadog tracing (OpenTrace)");
 			}
 
-			IServiceCollection Services = new ServiceCollection();
-			Services.AddCommandsFromAssembly(Assembly.GetExecutingAssembly());
-			Services.AddLogging(Builder => Builder.AddSerilog());
-			Services.AddSingleton<IConfiguration>(Config);
-			Services.AddSingleton<ServerSettings>(HordeSettings);
+			IServiceCollection services = new ServiceCollection();
+			services.AddCommandsFromAssembly(Assembly.GetExecutingAssembly());
+			services.AddLogging(builder => builder.AddSerilog());
+			services.AddSingleton<IConfiguration>(config);
+			services.AddSingleton<ServerSettings>(hordeSettings);
 
 #pragma warning disable ASP0000 // Do not call 'IServiceCollection.BuildServiceProvider' in 'ConfigureServices'
-			IServiceProvider ServiceProvider = Services.BuildServiceProvider();
-			return await CommandHost.RunAsync(Arguments, ServiceProvider, typeof(ServerCommand));
+			IServiceProvider serviceProvider = services.BuildServiceProvider();
+			return await CommandHost.RunAsync(arguments, serviceProvider, typeof(ServerCommand));
 #pragma warning restore ASP0000 // Do not call 'IServiceCollection.BuildServiceProvider' in 'ConfigureServices'
 		}
 
 		// Used by WebApplicationFactory in controller tests. Uses reflection to call this exact function signature.
-		public static IHostBuilder CreateHostBuilder(string[] Args) => ServerCommand.CreateHostBuilderForTesting(Args);
+		public static IHostBuilder CreateHostBuilder(string[] args) => ServerCommand.CreateHostBuilderForTesting(args);
 
 		/// <summary>
 		/// Get the application directory
@@ -223,10 +198,10 @@ namespace Horde.Build
 		{
 			if (RuntimeInformation.IsOSPlatform(OSPlatform.Windows))
 			{
-				DirectoryReference? Dir = DirectoryReference.GetSpecialFolder(Environment.SpecialFolder.CommonApplicationData);
-				if (Dir != null)
+				DirectoryReference? dir = DirectoryReference.GetSpecialFolder(Environment.SpecialFolder.CommonApplicationData);
+				if (dir != null)
 				{
-					return DirectoryReference.Combine(Dir, "HordeServer");
+					return DirectoryReference.Combine(dir, "HordeServer");
 				}
 			}
 			return DirectoryReference.Combine(GetAppDir(), "Data");
@@ -236,54 +211,52 @@ namespace Horde.Build
 		/// Handles bootstrapping of defaults for local servers, which can't be generated during build/installation process (or are better handled here where they can be updated)
 		/// This stuff will change as we get settings into database and could be considered discovery for installer/dockerfile builds 
 		/// </summary>		
-		static void InitializeDefaults(ServerSettings Settings)
+		static void InitializeDefaults(ServerSettings settings)
 		{			
-			if (Settings.SingleInstance)
+			if (settings.SingleInstance)
 			{
-				FileReference GlobalConfig = FileReference.Combine(Program.DataDir, "Config/globals.json");
+				FileReference globalConfig = FileReference.Combine(Program.DataDir, "Config/globals.json");
 
-				if (!FileReference.Exists(GlobalConfig))
+				if (!FileReference.Exists(globalConfig))
 				{
-					DirectoryReference.CreateDirectory(GlobalConfig.Directory);
-					FileReference.WriteAllText(GlobalConfig, "{}");
+					DirectoryReference.CreateDirectory(globalConfig.Directory);
+					FileReference.WriteAllText(globalConfig, "{}");
 				}
 
-				FileReference PrivateCertFile = FileReference.Combine(Program.DataDir, "Agent/ServerToAgent.pfx");
-				string PrivateCertFileJsonPath = PrivateCertFile.ToString().Replace("\\", "/", StringComparison.Ordinal);
+				FileReference privateCertFile = FileReference.Combine(Program.DataDir, "Agent/ServerToAgent.pfx");
+				string privateCertFileJsonPath = privateCertFile.ToString().Replace("\\", "/", StringComparison.Ordinal);
 
 				if (!FileReference.Exists(UserConfigFile))
 				{
 					// create new user configuration
 					DirectoryReference.CreateDirectory(UserConfigFile.Directory);
-					FileReference.WriteAllText(UserConfigFile, $"{{\"Horde\": {{ \"ConfigPath\" : \"{GlobalConfig.ToString().Replace("\\", "/", StringComparison.Ordinal)}\", \"ServerPrivateCert\" : \"{PrivateCertFileJsonPath}\", \"HttpPort\": 8080}}}}");
+					FileReference.WriteAllText(UserConfigFile, $"{{\"Horde\": {{ \"ConfigPath\" : \"{globalConfig.ToString().Replace("\\", "/", StringComparison.Ordinal)}\", \"ServerPrivateCert\" : \"{privateCertFileJsonPath}\", \"HttpPort\": 8080}}}}");
 				}
 
 				// make sure the cert exists
-				if (!FileReference.Exists(PrivateCertFile))
+				if (!FileReference.Exists(privateCertFile))
 				{
-					string DnsName = System.Net.Dns.GetHostName();
-					Serilog.Log.Logger.Information("Creating certificate for {DnsName}", DnsName);
+					string dnsName = System.Net.Dns.GetHostName();
+					Serilog.Log.Logger.Information("Creating certificate for {DnsName}", dnsName);
 
-					byte[] PrivateCertData = CertificateUtils.CreateSelfSignedCert(DnsName, "Horde Server");
+					byte[] privateCertData = CertificateUtils.CreateSelfSignedCert(dnsName, "Horde Server");
 					
-					Serilog.Log.Logger.Information("Writing private cert: {PrivateCert}", PrivateCertFile.FullName);
+					Serilog.Log.Logger.Information("Writing private cert: {PrivateCert}", privateCertFile.FullName);
 
-					if (!DirectoryReference.Exists(PrivateCertFile.Directory))
+					if (!DirectoryReference.Exists(privateCertFile.Directory))
 					{
-						DirectoryReference.CreateDirectory(PrivateCertFile.Directory);
+						DirectoryReference.CreateDirectory(privateCertFile.Directory);
 					}
 
-					FileReference.WriteAllBytes(PrivateCertFile, PrivateCertData);
+					FileReference.WriteAllBytes(privateCertFile, privateCertData);
 				}
 
 				// note: this isn't great, though we need it early in server startup, and this is only hit on first server boot where the grpc cert isn't generated/set 
-				if (Settings.ServerPrivateCert == null)
+				if (settings.ServerPrivateCert == null)
 				{
-					Settings.ServerPrivateCert = PrivateCertFile.ToString();
+					settings.ServerPrivateCert = privateCertFile.ToString();
 				}				
-
 			}
 		}
 	}
 }
-
diff --git a/Engine/Source/Programs/Horde/Horde.Build/Projects/IProject.cs b/Engine/Source/Programs/Horde/Horde.Build/Projects/IProject.cs
--- a/Engine/Source/Programs/Horde/Horde.Build/Projects/IProject.cs
+++ b/Engine/Source/Programs/Horde/Horde.Build/Projects/IProject.cs
@@ -1,17 +1,14 @@
 ﻿// Copyright Epic Games, Inc. All Rights Reserved.
 
+using System.Collections.Generic;
+using System.Linq;
+using System.Text.RegularExpressions;
 using EpicGames.Core;
 using Horde.Build.Acls;
 using Horde.Build.Api;
 using Horde.Build.Utilities;
-using MongoDB.Bson;
 using MongoDB.Bson.Serialization.Attributes;
 using MongoDB.Driver;
-using System;
-using System.Collections.Generic;
-using System.Linq;
-using System.Text.RegularExpressions;
-using System.Threading.Tasks;
 
 namespace Horde.Build.Models
 {
@@ -59,34 +56,34 @@ namespace Horde.Build.Models
 		/// <summary>
 		/// Constructor
 		/// </summary>
-		/// <param name="Name">Name of this group</param>
-		public StreamCategory(string Name)
+		/// <param name="name">Name of this group</param>
+		public StreamCategory(string name)
 		{
-			this.Name = Name;
+			Name = name;
 		}
 
 		/// <summary>
 		/// Constructs a category from a request object
 		/// </summary>
-		/// <param name="Request">The request object</param>
-		public StreamCategory(CreateProjectCategoryRequest Request)
+		/// <param name="request">The request object</param>
+		public StreamCategory(CreateProjectCategoryRequest request)
 		{
-			this.Name = Request.Name;
-			this.Row = Request.Row;
-			this.ShowOnNavMenu = Request.ShowOnNavMenu;
-			this.IncludePatterns = Request.IncludePatterns;
-			this.ExcludePatterns = Request.ExcludePatterns;
+			Name = request.Name;
+			Row = request.Row;
+			ShowOnNavMenu = request.ShowOnNavMenu;
+			IncludePatterns = request.IncludePatterns;
+			ExcludePatterns = request.ExcludePatterns;
 		}
 
 		/// <summary>
 		/// Tests if a given name matches a pattern
 		/// </summary>
-		/// <param name="Name">The name to test</param>
-		/// <param name="Pattern">The pattern to match against</param>
+		/// <param name="name">The name to test</param>
+		/// <param name="pattern">The pattern to match against</param>
 		/// <returns>True if the pattern matches, false otherwise</returns>
-		public static bool MatchPattern(string Name, string Pattern)
+		public static bool MatchPattern(string name, string pattern)
 		{
-			return Regex.IsMatch(Name, Pattern);
+			return Regex.IsMatch(name, pattern);
 		}
 	}
 
@@ -170,64 +167,64 @@ namespace Horde.Build.Models
 		/// <summary>
 		/// Converts this object to a public response
 		/// </summary>
-		/// <param name="Project">The project instance</param>
+		/// <param name="project">The project instance</param>
 		/// <param name="bIncludeStreams">Whether to include streams in the response</param>
 		/// <param name="bIncludeCategories">Whether to include categories in the response</param>
-		/// <param name="Streams">The list of streams</param>
+		/// <param name="streams">The list of streams</param>
 		/// <param name="bIncludeAcl">Whether to include the ACL in the response</param>
 		/// <returns>Response instance</returns>
-		public static GetProjectResponse ToResponse(this IProject Project, bool bIncludeStreams, bool bIncludeCategories, List<IStream>? Streams, bool bIncludeAcl)
+		public static GetProjectResponse ToResponse(this IProject project, bool bIncludeStreams, bool bIncludeCategories, List<IStream>? streams, bool bIncludeAcl)
 		{
-			List<GetProjectStreamResponse>? StreamResponses = null;
+			List<GetProjectStreamResponse>? streamResponses = null;
 			if(bIncludeStreams)
 			{
-				StreamResponses = Streams!.ConvertAll(x => new GetProjectStreamResponse(x.Id.ToString(), x.Name));
+				streamResponses = streams!.ConvertAll(x => new GetProjectStreamResponse(x.Id.ToString(), x.Name));
 			}
 
-			List<GetProjectCategoryResponse>? CategoryResponses = null;
+			List<GetProjectCategoryResponse>? categoryResponses = null;
 			if (bIncludeCategories)
 			{
-				CategoryResponses = Project.Categories.ConvertAll(x => new GetProjectCategoryResponse(x));
-				if (Streams != null)
+				categoryResponses = project.Categories.ConvertAll(x => new GetProjectCategoryResponse(x));
+				if (streams != null)
 				{
-					foreach (IStream Stream in Streams)
+					foreach (IStream stream in streams)
 					{
-						GetProjectCategoryResponse? CategoryResponse = CategoryResponses.FirstOrDefault(x => MatchCategory(Stream.Name, x));
-						if(CategoryResponse == null)
+						GetProjectCategoryResponse? categoryResponse = categoryResponses.FirstOrDefault(x => MatchCategory(stream.Name, x));
+						if(categoryResponse == null)
 						{
-							int Row = (CategoryResponses.Count > 0) ? CategoryResponses.Max(x => x.Row) : 0;
-							if (CategoryResponses.Count(x => x.Row == Row) >= 3)
+							int row = (categoryResponses.Count > 0) ? categoryResponses.Max(x => x.Row) : 0;
+							if (categoryResponses.Count(x => x.Row == row) >= 3)
 							{
-								Row++;
+								row++;
 							}
 
-							StreamCategory OtherCategory = new StreamCategory("Other");
-							OtherCategory.Row = Row;
-							OtherCategory.IncludePatterns.Add(".*");
+							StreamCategory otherCategory = new StreamCategory("Other");
+							otherCategory.Row = row;
+							otherCategory.IncludePatterns.Add(".*");
 
-							CategoryResponse = new GetProjectCategoryResponse(OtherCategory);
-							CategoryResponses.Add(CategoryResponse);
+							categoryResponse = new GetProjectCategoryResponse(otherCategory);
+							categoryResponses.Add(categoryResponse);
 						}
-						CategoryResponse.Streams!.Add(Stream.Id.ToString());
+						categoryResponse.Streams!.Add(stream.Id.ToString());
 					}
 				}
 			}
 
-			GetAclResponse? AclResponse = (bIncludeAcl && Project.Acl != null) ? new GetAclResponse(Project.Acl) : null;
-			return new GetProjectResponse(Project.Id.ToString(), Project.Name, Project.Order, StreamResponses, CategoryResponses, AclResponse);
+			GetAclResponse? aclResponse = (bIncludeAcl && project.Acl != null) ? new GetAclResponse(project.Acl) : null;
+			return new GetProjectResponse(project.Id.ToString(), project.Name, project.Order, streamResponses, categoryResponses, aclResponse);
 		}
 
 		/// <summary>
 		/// Tests if a category response matches a given stream name
 		/// </summary>
-		/// <param name="Name">The stream name</param>
-		/// <param name="Category">The category response</param>
+		/// <param name="name">The stream name</param>
+		/// <param name="category">The category response</param>
 		/// <returns>True if the category matches</returns>
-		static bool MatchCategory(string Name, GetProjectCategoryResponse Category)
+		static bool MatchCategory(string name, GetProjectCategoryResponse category)
 		{
-			if (Category.IncludePatterns.Any(x => StreamCategory.MatchPattern(Name, x)))
+			if (category.IncludePatterns.Any(x => StreamCategory.MatchPattern(name, x)))
 			{
-				if (!Category.ExcludePatterns.Any(x => StreamCategory.MatchPattern(Name, x)))
+				if (!category.ExcludePatterns.Any(x => StreamCategory.MatchPattern(name, x)))
 				{
 					return true;
 				}
diff --git a/Engine/Source/Programs/Horde/Horde.Build/Projects/IProjectCollection.cs b/Engine/Source/Programs/Horde/Horde.Build/Projects/IProjectCollection.cs
--- a/Engine/Source/Programs/Horde/Horde.Build/Projects/IProjectCollection.cs
+++ b/Engine/Source/Programs/Horde/Horde.Build/Projects/IProjectCollection.cs
@@ -1,19 +1,9 @@
 // Copyright Epic Games, Inc. All Rights Reserved.
 
-using Horde.Build.Api;
-using Horde.Build.Models;
-using Horde.Build.Services;
-using Horde.Build.Utilities;
-using Microsoft.Extensions.Options;
-using MongoDB.Bson;
-using MongoDB.Bson.Serialization.Attributes;
-using MongoDB.Driver;
-using System;
 using System.Collections.Generic;
-using System.Diagnostics.CodeAnalysis;
-using System.Linq;
-using System.Security.Claims;
 using System.Threading.Tasks;
+using Horde.Build.Models;
+using Horde.Build.Utilities;
 
 namespace Horde.Build.Collections
 {
@@ -27,13 +17,13 @@ namespace Horde.Build.Collections
 		/// <summary>
 		/// Updates the project configuration
 		/// </summary>
-		/// <param name="Id">The project id</param>
-		/// <param name="ConfigPath">Path to the config file used to configure this project</param>
-		/// <param name="Revision">The config file revision</param>
-		/// <param name="Order">Order of the project</param>
-		/// <param name="Config">The configuration</param>
+		/// <param name="id">The project id</param>
+		/// <param name="configPath">Path to the config file used to configure this project</param>
+		/// <param name="revision">The config file revision</param>
+		/// <param name="order">Order of the project</param>
+		/// <param name="config">The configuration</param>
 		/// <returns>New project instance</returns>
-		Task<IProject?> AddOrUpdateAsync(ProjectId Id, string ConfigPath, string Revision, int Order, ProjectConfig Config);
+		Task<IProject?> AddOrUpdateAsync(ProjectId id, string configPath, string revision, int order, ProjectConfig config);
 
 		/// <summary>
 		/// Gets all the available projects
@@ -44,40 +34,40 @@ namespace Horde.Build.Collections
 		/// <summary>
 		/// Gets a project by ID
 		/// </summary>
-		/// <param name="ProjectId">Unique id of the project</param>
+		/// <param name="projectId">Unique id of the project</param>
 		/// <returns>The project document</returns>
-		Task<IProject?> GetAsync(ProjectId ProjectId);
+		Task<IProject?> GetAsync(ProjectId projectId);
 
 		/// <summary>
 		/// Gets the logo for a project
 		/// </summary>
-		/// <param name="ProjectId">The project id</param>
+		/// <param name="projectId">The project id</param>
 		/// <returns>The project logo document</returns>
-		Task<IProjectLogo?> GetLogoAsync(ProjectId ProjectId);
+		Task<IProjectLogo?> GetLogoAsync(ProjectId projectId);
 
 		/// <summary>
 		/// Sets the logo for a project
 		/// </summary>
-		/// <param name="ProjectId">The project id</param>
-		/// <param name="Path">Path to the source file</param>
-		/// <param name="Revision">Revision of the file</param>
-		/// <param name="MimeType"></param>
-		/// <param name="Data"></param>
+		/// <param name="projectId">The project id</param>
+		/// <param name="path">Path to the source file</param>
+		/// <param name="revision">Revision of the file</param>
+		/// <param name="mimeType"></param>
+		/// <param name="data"></param>
 		/// <returns></returns>
-		Task SetLogoAsync(ProjectId ProjectId, string Path, string Revision, string MimeType, byte[] Data);
+		Task SetLogoAsync(ProjectId projectId, string path, string revision, string mimeType, byte[] data);
 
 		/// <summary>
 		/// Gets a project's permissions info by ID
 		/// </summary>
-		/// <param name="ProjectId">Unique id of the project</param>
+		/// <param name="projectId">Unique id of the project</param>
 		/// <returns>The project document</returns>
-		Task<IProjectPermissions?> GetPermissionsAsync(ProjectId ProjectId);
+		Task<IProjectPermissions?> GetPermissionsAsync(ProjectId projectId);
 
 		/// <summary>
 		/// Deletes a project by id
 		/// </summary>
-		/// <param name="ProjectId">Unique id of the project</param>
+		/// <param name="projectId">Unique id of the project</param>
 		/// <returns>True if the project was deleted</returns>
-		Task DeleteAsync(ProjectId ProjectId);
+		Task DeleteAsync(ProjectId projectId);
 	}
 }
diff --git a/Engine/Source/Programs/Horde/Horde.Build/Projects/ISchedule.cs b/Engine/Source/Programs/Horde/Horde.Build/Projects/ISchedule.cs
--- a/Engine/Source/Programs/Horde/Horde.Build/Projects/ISchedule.cs
+++ b/Engine/Source/Programs/Horde/Horde.Build/Projects/ISchedule.cs
@@ -1,14 +1,9 @@
 // Copyright Epic Games, Inc. All Rights Reserved.
 
-using EpicGames.Core;
-using Horde.Build.Api;
-using Horde.Build.Utilities;
-using MongoDB.Bson;
-using MongoDB.Bson.Serialization.Attributes;
 using System;
 using System.Collections.Generic;
-using System.Linq;
-using System.Threading.Tasks;
+using Horde.Build.Utilities;
+using MongoDB.Bson.Serialization.Attributes;
 
 namespace Horde.Build.Models
 {
@@ -50,62 +45,62 @@ namespace Horde.Build.Models
 		/// <summary>
 		/// Constructor
 		/// </summary>
-		/// <param name="DaysOfWeek">Which days of the week the schedule should run</param>
-		/// <param name="MinTime">Time during the day for the first schedule to trigger. Measured in minutes from midnight.</param>
-		/// <param name="MaxTime">Time during the day for the last schedule to trigger. Measured in minutes from midnight.</param>
-		/// <param name="Interval">Interval between each schedule triggering</param>
-		public SchedulePattern(List<DayOfWeek>? DaysOfWeek, int MinTime, int? MaxTime, int? Interval)
+		/// <param name="daysOfWeek">Which days of the week the schedule should run</param>
+		/// <param name="minTime">Time during the day for the first schedule to trigger. Measured in minutes from midnight.</param>
+		/// <param name="maxTime">Time during the day for the last schedule to trigger. Measured in minutes from midnight.</param>
+		/// <param name="interval">Interval between each schedule triggering</param>
+		public SchedulePattern(List<DayOfWeek>? daysOfWeek, int minTime, int? maxTime, int? interval)
 		{
-			this.DaysOfWeek = DaysOfWeek;
-			this.MinTime = MinTime;
-			this.MaxTime = MaxTime;
-			this.Interval = Interval;
+			DaysOfWeek = daysOfWeek;
+			MinTime = minTime;
+			MaxTime = maxTime;
+			Interval = interval;
 		}
 
 		/// <summary>
 		/// Calculates the trigger index based on the given time in minutes
 		/// </summary>
-		/// <param name="LastTimeUtc">Time for the last trigger</param>
-		/// <param name="TimeZone">The timezone for running the schedule</param>
+		/// <param name="lastTimeUtc">Time for the last trigger</param>
+		/// <param name="timeZone">The timezone for running the schedule</param>
 		/// <returns>Index of the trigger</returns>
-		public DateTime GetNextTriggerTimeUtc(DateTime LastTimeUtc, TimeZoneInfo TimeZone)
+		public DateTime GetNextTriggerTimeUtc(DateTime lastTimeUtc, TimeZoneInfo timeZone)
 		{
 			// Convert last time into the correct timezone for running the scheule
-			DateTimeOffset LastTime = TimeZoneInfo.ConvertTime((DateTimeOffset)LastTimeUtc, TimeZone);
+			DateTimeOffset lastTime = TimeZoneInfo.ConvertTime((DateTimeOffset)lastTimeUtc, timeZone);
 
 			// Get the base time (ie. the start of this day) for anchoring the schedule
-			DateTimeOffset BaseTime = new DateTimeOffset(LastTime.Year, LastTime.Month, LastTime.Day, 0, 0, 0, LastTime.Offset);
+			DateTimeOffset baseTime = new DateTimeOffset(lastTime.Year, lastTime.Month, lastTime.Day, 0, 0, 0, lastTime.Offset);
 			for (; ; )
 			{
-				if (DaysOfWeek == null || DaysOfWeek.Contains(BaseTime.DayOfWeek))
+				if (DaysOfWeek == null || DaysOfWeek.Contains(baseTime.DayOfWeek))
 				{
 					// Get the last time in minutes from the start of this day
-					int LastTimeMinutes = (int)(LastTime - BaseTime).TotalMinutes;
+					int lastTimeMinutes = (int)(lastTime - baseTime).TotalMinutes;
 
 					// Get the time of the first trigger of this day. If the last time is less than this, this is the next trigger.
-					if (LastTimeMinutes < MinTime)
+					if (lastTimeMinutes < MinTime)
 					{
-						return BaseTime.AddMinutes(MinTime).UtcDateTime;
+						return baseTime.AddMinutes(MinTime).UtcDateTime;
 					}
 
 					// Otherwise, get the time for the last trigger in the day.
 					if (Interval.HasValue && Interval.Value > 0)
 					{
-						int ActualMaxTime = MaxTime ?? ((24 * 60) - 1);
-						if (LastTimeMinutes < ActualMaxTime)
+						int actualMaxTime = MaxTime ?? ((24 * 60) - 1);
+						if (lastTimeMinutes < actualMaxTime)
 						{
-							int LastIndex = (LastTimeMinutes - MinTime) / Interval.Value;
-							int NextIndex = LastIndex + 1;
+							int lastIndex = (lastTimeMinutes - MinTime) / Interval.Value;
+							int nextIndex = lastIndex + 1;
 
-							int NextTimeMinutes = MinTime + (NextIndex * Interval.Value);
-							if (NextTimeMinutes <= ActualMaxTime)
+							int nextTimeMinutes = MinTime + (nextIndex * Interval.Value);
+							if (nextTimeMinutes <= actualMaxTime)
 							{
-								return BaseTime.AddMinutes(NextTimeMinutes).UtcDateTime;
+								return baseTime.AddMinutes(nextTimeMinutes).UtcDateTime;
 							}
 						}
 					}
 				}
-				BaseTime = BaseTime.AddDays(1.0);
+				baseTime = baseTime.AddDays(1.0);
 			}
 		}
 	}
@@ -137,12 +132,12 @@ namespace Horde.Build.Models
 		/// <summary>
 		/// Constructor
 		/// </summary>
-		/// <param name="TemplateRefId">The template containing the dependency</param>
-		/// <param name="Target">Target to wait for</param>
-		public ScheduleGate(TemplateRefId TemplateRefId, string Target)
+		/// <param name="templateRefId">The template containing the dependency</param>
+		/// <param name="target">Target to wait for</param>
+		public ScheduleGate(TemplateRefId templateRefId, string target)
 		{
-			this.TemplateRefId = TemplateRefId;
-			this.Target = Target;
+			TemplateRefId = templateRefId;
+			Target = target;
 		}
 	}
 
@@ -232,39 +227,39 @@ namespace Horde.Build.Models
 		/// <summary>
 		/// Constructor
 		/// </summary>
-		/// <param name="CurrentTimeUtc">The current time. This will be used to seed the last trigger time.</param>
-		/// <param name="Enabled">Whether the schedule is currently enabled</param>
-		/// <param name="MaxActive">Maximum number of builds that may be active at once</param>
-		/// <param name="MaxChanges">Maximum number of changes the schedule can fall behind head revision</param>
-		/// <param name="RequireSubmittedChange">Whether a change has to be submitted for the schedule to trigger</param>
-		/// <param name="Gate">Reference to another job/target which much succeed for this schedule to trigger</param>
-		/// <param name="Filter">Filter for changes to consider</param>
-		/// <param name="Files">Files that should trigger the schedule</param>
-		/// <param name="TemplateParameters">Parameters for the template to run</param>
-		/// <param name="Patterns">List of patterns for the schedule</param>
-		public Schedule(DateTime CurrentTimeUtc, bool Enabled = true, int MaxActive = 0, int MaxChanges = 0, bool RequireSubmittedChange = true, ScheduleGate? Gate = null, List<ChangeContentFlags>? Filter = null, List<string>? Files = null, Dictionary<string, string>? TemplateParameters = null, List<SchedulePattern>? Patterns = null)
+		/// <param name="currentTimeUtc">The current time. This will be used to seed the last trigger time.</param>
+		/// <param name="enabled">Whether the schedule is currently enabled</param>
+		/// <param name="maxActive">Maximum number of builds that may be active at once</param>
+		/// <param name="maxChanges">Maximum number of changes the schedule can fall behind head revision</param>
+		/// <param name="requireSubmittedChange">Whether a change has to be submitted for the schedule to trigger</param>
+		/// <param name="gate">Reference to another job/target which much succeed for this schedule to trigger</param>
+		/// <param name="filter">Filter for changes to consider</param>
+		/// <param name="files">Files that should trigger the schedule</param>
+		/// <param name="templateParameters">Parameters for the template to run</param>
+		/// <param name="patterns">List of patterns for the schedule</param>
+		public Schedule(DateTime currentTimeUtc, bool enabled = true, int maxActive = 0, int maxChanges = 0, bool requireSubmittedChange = true, ScheduleGate? gate = null, List<ChangeContentFlags>? filter = null, List<string>? files = null, Dictionary<string, string>? templateParameters = null, List<SchedulePattern>? patterns = null)
 		{
-			this.Enabled = Enabled;
-			this.MaxActive = MaxActive;
-			this.MaxChanges = MaxChanges;
-			this.RequireSubmittedChange = RequireSubmittedChange;
-			this.Gate = Gate;
-			this.Filter = Filter;
-			this.Files = Files;
-			this.TemplateParameters = TemplateParameters ?? new Dictionary<string, string>();
-			this.Patterns = Patterns ?? new List<SchedulePattern>();
-			this.LastTriggerTimeUtc = CurrentTimeUtc;
+			Enabled = enabled;
+			MaxActive = maxActive;
+			MaxChanges = maxChanges;
+			RequireSubmittedChange = requireSubmittedChange;
+			Gate = gate;
+			Filter = filter;
+			Files = files;
+			TemplateParameters = templateParameters ?? new Dictionary<string, string>();
+			Patterns = patterns ?? new List<SchedulePattern>();
+			LastTriggerTimeUtc = currentTimeUtc;
 		}
 
 		/// <summary>
 		/// Copies the state fields from another schedule object
 		/// </summary>
-		/// <param name="Other">The schedule object to copy from</param>
-		public void CopyState(Schedule Other)
+		/// <param name="other">The schedule object to copy from</param>
+		public void CopyState(Schedule other)
 		{
-			LastTriggerChange = Other.LastTriggerChange;
-			LastTriggerTime = Other.LastTriggerTime;
-			ActiveJobs.AddRange(Other.ActiveJobs);
+			LastTriggerChange = other.LastTriggerChange;
+			LastTriggerTime = other.LastTriggerTime;
+			ActiveJobs.AddRange(other.ActiveJobs);
 		}
 
 		/// <summary>
@@ -278,12 +273,12 @@ namespace Horde.Build.Models
 				return null;
 			}
 
-			ChangeContentFlags Flags = 0;
-			foreach (ChangeContentFlags Flag in Filter)
+			ChangeContentFlags flags = 0;
+			foreach (ChangeContentFlags flag in Filter)
 			{
-				Flags |= Flag;
+				flags |= flag;
 			}
-			return Flags;
+			return flags;
 		}
 
 		/// <summary>
@@ -303,34 +298,34 @@ namespace Horde.Build.Models
 		/// <summary>
 		/// Get the next time that the schedule will trigger
 		/// </summary>
-		/// <param name="TimeZone">Timezone to evaluate the trigger</param>
+		/// <param name="timeZone">Timezone to evaluate the trigger</param>
 		/// <returns>Next time at which the schedule will trigger</returns>
-		public DateTime? GetNextTriggerTimeUtc(TimeZoneInfo TimeZone)
+		public DateTime? GetNextTriggerTimeUtc(TimeZoneInfo timeZone)
 		{
-			return GetNextTriggerTimeUtc(GetLastTriggerTimeUtc(), TimeZone);
+			return GetNextTriggerTimeUtc(GetLastTriggerTimeUtc(), timeZone);
 		}
 
 		/// <summary>
 		/// Get the next time that the schedule will trigger
 		/// </summary>
-		/// <param name="LastTimeUtc">Last time at which the schedule triggered</param>
-		/// <param name="TimeZone">Timezone to evaluate the trigger</param>
+		/// <param name="lastTimeUtc">Last time at which the schedule triggered</param>
+		/// <param name="timeZone">Timezone to evaluate the trigger</param>
 		/// <returns>Next time at which the schedule will trigger</returns>
-		public DateTime? GetNextTriggerTimeUtc(DateTime LastTimeUtc, TimeZoneInfo TimeZone)
+		public DateTime? GetNextTriggerTimeUtc(DateTime lastTimeUtc, TimeZoneInfo timeZone)
 		{
-			DateTime? NextTriggerTimeUtc = null;
+			DateTime? nextTriggerTimeUtc = null;
 			if (Enabled)
 			{
-				foreach (SchedulePattern Pattern in Patterns)
+				foreach (SchedulePattern pattern in Patterns)
 				{
-					DateTime PatternTriggerTime = Pattern.GetNextTriggerTimeUtc(LastTimeUtc, TimeZone);
-					if (NextTriggerTimeUtc == null || PatternTriggerTime < NextTriggerTimeUtc)
+					DateTime patternTriggerTime = pattern.GetNextTriggerTimeUtc(lastTimeUtc, timeZone);
+					if (nextTriggerTimeUtc == null || patternTriggerTime < nextTriggerTimeUtc)
 					{
-						NextTriggerTimeUtc = PatternTriggerTime;
+						nextTriggerTimeUtc = patternTriggerTime;
 					}
 				}
 			}
-			return NextTriggerTimeUtc;
+			return nextTriggerTimeUtc;
 		}
 	}
 }
diff --git a/Engine/Source/Programs/Horde/Horde.Build/Projects/IStream.cs b/Engine/Source/Programs/Horde/Horde.Build/Projects/IStream.cs
--- a/Engine/Source/Programs/Horde/Horde.Build/Projects/IStream.cs
+++ b/Engine/Source/Programs/Horde/Horde.Build/Projects/IStream.cs
@@ -1,24 +1,16 @@
 ﻿// Copyright Epic Games, Inc. All Rights Reserved.
 
-using EpicGames.Core;
-using Google.Protobuf.WellKnownTypes;
-using Horde.Build.Acls;
-using Horde.Build.Api;
-using HordeCommon;
-using Horde.Build.Utilities;
-using MongoDB.Bson;
-using MongoDB.Bson.Serialization.Attributes;
-using MongoDB.Driver;
 using System;
 using System.Collections.Generic;
 using System.Diagnostics.CodeAnalysis;
-using System.IO;
 using System.Linq;
-using System.Security.Cryptography;
-using System.Text;
 using System.Text.RegularExpressions;
-using System.Threading.Tasks;
-using System.Runtime.CompilerServices;
+using EpicGames.Core;
+using Horde.Build.Acls;
+using Horde.Build.Api;
+using Horde.Build.Utilities;
+using HordeCommon;
+using MongoDB.Bson.Serialization.Attributes;
 
 namespace Horde.Build.Models
 {
@@ -38,12 +30,12 @@ namespace Horde.Build.Models
 		}
 
 		/// <inheritdoc/>
-		public InvalidStreamException(string Message) : base(Message)
+		public InvalidStreamException(string message) : base(message)
 		{
 		}
 
 		/// <inheritdoc/>
-		public InvalidStreamException(string Message, Exception InnerEx) : base(Message, InnerEx)
+		public InvalidStreamException(string message, Exception innerEx) : base(message, innerEx)
 		{
 		}
 	}
@@ -85,37 +77,37 @@ namespace Horde.Build.Models
 		/// <summary>
 		/// Constructor
 		/// </summary>
-		/// <param name="Pool">The pool for this agent</param>
-		/// <param name="Workspace">Name of the workspace to use</param>
-		/// <param name="TempStorageDir">Path to the temp storage directory</param>
-		public AgentType(PoolId Pool, string Workspace, string? TempStorageDir)
+		/// <param name="pool">The pool for this agent</param>
+		/// <param name="workspace">Name of the workspace to use</param>
+		/// <param name="tempStorageDir">Path to the temp storage directory</param>
+		public AgentType(PoolId pool, string workspace, string? tempStorageDir)
 		{
-			this.Pool = Pool;
-			this.Workspace = Workspace;
-			this.TempStorageDir = TempStorageDir;
+			Pool = pool;
+			Workspace = workspace;
+			TempStorageDir = tempStorageDir;
 		}
 
 		/// <summary>
 		/// Constructor
 		/// </summary>
-		/// <param name="Request">The object to construct from</param>
-		public AgentType(CreateAgentTypeRequest Request)
+		/// <param name="request">The object to construct from</param>
+		public AgentType(CreateAgentTypeRequest request)
 		{
-			Pool = new PoolId(Request.Pool);
-			Workspace = Request.Workspace;
-			TempStorageDir = Request.TempStorageDir;
-			Environment = Request.Environment;
+			Pool = new PoolId(request.Pool);
+			Workspace = request.Workspace;
+			TempStorageDir = request.TempStorageDir;
+			Environment = request.Environment;
 		}
 
 		/// <summary>
 		/// Constructs an AgentType object from an optional request
 		/// </summary>
-		/// <param name="Request">The request object</param>
+		/// <param name="request">The request object</param>
 		/// <returns>New agent type object</returns>
-		[return: NotNullIfNotNull("Request")]
-		public static AgentType? FromRequest(CreateAgentTypeRequest? Request)
+		[return: NotNullIfNotNull("request")]
+		public static AgentType? FromRequest(CreateAgentTypeRequest? request)
 		{
-			return (Request != null) ? new AgentType(Request) : null;
+			return (request != null) ? new AgentType(request) : null;
 		}
 
 		/// <summary>
@@ -133,16 +125,16 @@ namespace Horde.Build.Models
 		/// <returns>The response object</returns>
 		public HordeCommon.Rpc.GetAgentTypeResponse ToRpcResponse()
 		{
-			HordeCommon.Rpc.GetAgentTypeResponse Response = new HordeCommon.Rpc.GetAgentTypeResponse();
+			HordeCommon.Rpc.GetAgentTypeResponse response = new HordeCommon.Rpc.GetAgentTypeResponse();
 			if (TempStorageDir != null)
 			{
-				Response.TempStorageDir = TempStorageDir;
+				response.TempStorageDir = TempStorageDir;
 			}
 			if (Environment != null)
 			{
-				Response.Environment.Add(Environment);
+				response.Environment.Add(Environment);
 			}
-			return Response;
+			return response;
 		}
 	}
 
@@ -213,29 +205,29 @@ namespace Horde.Build.Models
 		/// <summary>
 		/// Constructor
 		/// </summary>
-		/// <param name="Request">The object to construct from</param>
-		public WorkspaceType(CreateWorkspaceTypeRequest Request)
+		/// <param name="request">The object to construct from</param>
+		public WorkspaceType(CreateWorkspaceTypeRequest request)
 		{
-			Cluster = Request.Cluster;
-			ServerAndPort = Request.ServerAndPort;
-			UserName = Request.UserName;
-			Password = Request.Password;
-			Identifier = Request.Identifier;
-			Stream = Request.Stream;
-			View = Request.View;
-			Incremental = Request.Incremental;
-			UseAutoSdk = Request.UseAutoSdk;
+			Cluster = request.Cluster;
+			ServerAndPort = request.ServerAndPort;
+			UserName = request.UserName;
+			Password = request.Password;
+			Identifier = request.Identifier;
+			Stream = request.Stream;
+			View = request.View;
+			Incremental = request.Incremental;
+			UseAutoSdk = request.UseAutoSdk;
 		}
 
 		/// <summary>
 		/// Constructs an AgentType object from an optional request
 		/// </summary>
-		/// <param name="Request">The request object</param>
+		/// <param name="request">The request object</param>
 		/// <returns>New agent type object</returns>
-		[return: NotNullIfNotNull("Request")]
-		public static WorkspaceType? FromRequest(CreateWorkspaceTypeRequest? Request)
+		[return: NotNullIfNotNull("request")]
+		public static WorkspaceType? FromRequest(CreateWorkspaceTypeRequest? request)
 		{
-			return (Request != null) ? new WorkspaceType(Request) : null;
+			return (request != null) ? new WorkspaceType(request) : null;
 		}
 
 		/// <summary>
@@ -273,20 +265,20 @@ namespace Horde.Build.Models
 		/// <summary>
 		/// Constructor
 		/// </summary>
-		/// <param name="Trigger">Name of the target that needs to complete</param>
-		/// <param name="TemplateRefId">The new template to trigger</param>
-		public ChainedJobTemplate(string Trigger, TemplateRefId TemplateRefId)
+		/// <param name="trigger">Name of the target that needs to complete</param>
+		/// <param name="templateRefId">The new template to trigger</param>
+		public ChainedJobTemplate(string trigger, TemplateRefId templateRefId)
 		{
-			this.Trigger = Trigger;
-			this.TemplateRefId = TemplateRefId;
+			Trigger = trigger;
+			TemplateRefId = templateRefId;
 		}
 
 		/// <summary>
 		/// Constructor
 		/// </summary>
-		/// <param name="Request">Request to construct from</param>
-		public ChainedJobTemplate(CreateChainedJobTemplateRequest Request)
-			: this(Request.Trigger, new TemplateRefId(Request.TemplateId))
+		/// <param name="request">Request to construct from</param>
+		public ChainedJobTemplate(CreateChainedJobTemplateRequest request)
+			: this(request.Trigger, new TemplateRefId(request.TemplateId))
 		{
 		}
 	}
@@ -353,34 +345,34 @@ namespace Horde.Build.Models
 		/// </summary>
 		private TemplateRef()
 		{
-			this.Name = null!;
-			this.Hash = null!;
+			Name = null!;
+			Hash = null!;
 		}
 
 		/// <summary>
 		/// Constructor
 		/// </summary>
-		/// <param name="Template">The template being referenced</param>
-		/// <param name="ShowUgsBadges">Whether to show badges in UGS for this job</param>
-		/// <param name="ShowUgsAlerts">Whether to show alerts in UGS for this job</param>
-		/// <param name="NotificationChannel">Notification channel for this template</param>
-		/// <param name="NotificationChannelFilter">Notification channel filter for this template</param>
-		/// <param name="TriageChannel"></param>
-		/// <param name="Schedule">Schedule for this template</param>
-		/// <param name="Triggers">List of downstream templates to trigger</param>
-		/// <param name="Acl">ACL for this template</param>
-		public TemplateRef(ITemplate Template, bool ShowUgsBadges = false, bool ShowUgsAlerts = false, string? NotificationChannel = null, string? NotificationChannelFilter = null, string? TriageChannel = null, Schedule? Schedule = null, List<ChainedJobTemplate>? Triggers = null, Acl? Acl = null)
+		/// <param name="template">The template being referenced</param>
+		/// <param name="showUgsBadges">Whether to show badges in UGS for this job</param>
+		/// <param name="showUgsAlerts">Whether to show alerts in UGS for this job</param>
+		/// <param name="notificationChannel">Notification channel for this template</param>
+		/// <param name="notificationChannelFilter">Notification channel filter for this template</param>
+		/// <param name="triageChannel"></param>
+		/// <param name="schedule">Schedule for this template</param>
+		/// <param name="triggers">List of downstream templates to trigger</param>
+		/// <param name="acl">ACL for this template</param>
+		public TemplateRef(ITemplate template, bool showUgsBadges = false, bool showUgsAlerts = false, string? notificationChannel = null, string? notificationChannelFilter = null, string? triageChannel = null, Schedule? schedule = null, List<ChainedJobTemplate>? triggers = null, Acl? acl = null)
 		{
-			this.Name = Template.Name;
-			this.Hash = Template.Id;
-			this.ShowUgsBadges = ShowUgsBadges;
-			this.ShowUgsAlerts = ShowUgsAlerts;
-			this.NotificationChannel = NotificationChannel;
-			this.NotificationChannelFilter = NotificationChannelFilter;
-			this.TriageChannel = TriageChannel;
-			this.Schedule = Schedule;
-			this.ChainedJobs = Triggers;
-			this.Acl = Acl;
+			Name = template.Name;
+			Hash = template.Id;
+			ShowUgsBadges = showUgsBadges;
+			ShowUgsAlerts = showUgsAlerts;
+			NotificationChannel = notificationChannel;
+			NotificationChannelFilter = notificationChannelFilter;
+			TriageChannel = triageChannel;
+			Schedule = schedule;
+			ChainedJobs = triggers;
+			Acl = acl;
 		}
 	}
 
@@ -438,12 +430,12 @@ namespace Horde.Build.Models
 		/// <summary>
 		/// Constructor
 		/// </summary>
-		/// <param name="TemplateRefId"></param>
-		/// <param name="Change">The job type to query for the change to use</param>
-		public DefaultPreflight(TemplateRefId? TemplateRefId, ChangeQuery? Change)
+		/// <param name="templateRefId"></param>
+		/// <param name="change">The job type to query for the change to use</param>
+		public DefaultPreflight(TemplateRefId? templateRefId, ChangeQuery? change)
 		{
-			this.TemplateRefId = TemplateRefId;
-			this.Change = Change;
+			TemplateRefId = templateRefId;
+			Change = change;
 		}
 
 		/// <summary>
@@ -453,17 +445,17 @@ namespace Horde.Build.Models
 		public DefaultPreflightRequest ToRequest()
 		{
 #pragma warning disable CS0618 // Type or member is obsolete
-			ChangeQueryRequest? ChangeRequest = null;
+			ChangeQueryRequest? changeRequest = null;
 			if (Change != null)
 			{
-				ChangeRequest = Change.ToRequest();
+				changeRequest = Change.ToRequest();
 			}
 			else if (ChangeTemplateRefId != null)
 			{
-				ChangeRequest = new ChangeQueryRequest { TemplateId = ChangeTemplateRefId.ToString() };
+				changeRequest = new ChangeQueryRequest { TemplateId = ChangeTemplateRefId.ToString() };
 			}
 
-			return new DefaultPreflightRequest { TemplateId = TemplateRefId?.ToString(), Change = ChangeRequest, ChangeTemplateId = ChangeRequest?.TemplateId };
+			return new DefaultPreflightRequest { TemplateId = TemplateRefId?.ToString(), Change = changeRequest, ChangeTemplateId = changeRequest?.TemplateId };
 #pragma warning restore CS0618 // Type or member is obsolete
 		}
 	}
@@ -476,11 +468,11 @@ namespace Horde.Build.Models
 		/// <summary>
 		/// Adds a new template ref to a list
 		/// </summary>
-		/// <param name="TemplateRefs">List of template refs</param>
-		/// <param name="TemplateRef">The template ref to add</param>
-		public static void AddRef(this Dictionary<TemplateRefId, TemplateRef> TemplateRefs, TemplateRef TemplateRef)
+		/// <param name="templateRefs">List of template refs</param>
+		/// <param name="templateRef">The template ref to add</param>
+		public static void AddRef(this Dictionary<TemplateRefId, TemplateRef> templateRefs, TemplateRef templateRef)
 		{
-			TemplateRefs.Add(new TemplateRefId(TemplateRef.Name), TemplateRef);
+			templateRefs.Add(new TemplateRefId(templateRef.Name), templateRef);
 		}
 	}
 
@@ -619,46 +611,46 @@ namespace Horde.Build.Models
 		/// <summary>
 		/// Tries to get an agent workspace definition from the given type name
 		/// </summary>
-		/// <param name="Stream">The stream object</param>
-		/// <param name="AgentType">The agent type</param>
-		/// <param name="Workspace">Receives the agent workspace definition</param>
+		/// <param name="stream">The stream object</param>
+		/// <param name="agentType">The agent type</param>
+		/// <param name="workspace">Receives the agent workspace definition</param>
 		/// <returns>True if the agent type was valid, and an agent workspace could be created</returns>
-		public static bool TryGetAgentWorkspace(this IStream Stream, AgentType AgentType, [NotNullWhen(true)] out (AgentWorkspace, bool)? Workspace)
+		public static bool TryGetAgentWorkspace(this IStream stream, AgentType agentType, [NotNullWhen(true)] out (AgentWorkspace, bool)? workspace)
 		{
 			// Get the workspace settings
-			if (AgentType.Workspace == null)
+			if (agentType.Workspace == null)
 			{
 				// Use the default settings (fast switching workspace, clean 
-				Workspace = (new AgentWorkspace(null, null, GetDefaultWorkspaceIdentifier(Stream), Stream.Name, null, false), true);
+				workspace = (new AgentWorkspace(null, null, GetDefaultWorkspaceIdentifier(stream), stream.Name, null, false), true);
 				return true;
 			}
 			else
 			{
 				// Try to get the matching workspace type
-				WorkspaceType? WorkspaceType;
-				if (!Stream.WorkspaceTypes.TryGetValue(AgentType.Workspace, out WorkspaceType))
+				WorkspaceType? workspaceType;
+				if (!stream.WorkspaceTypes.TryGetValue(agentType.Workspace, out workspaceType))
 				{
-					Workspace = null;
+					workspace = null;
 					return false;
 				}
 
 				// Get the workspace identifier
-				string Identifier;
-				if (WorkspaceType.Identifier != null)
+				string identifier;
+				if (workspaceType.Identifier != null)
 				{
-					Identifier = WorkspaceType.Identifier;
+					identifier = workspaceType.Identifier;
 				}
-				else if (WorkspaceType.Incremental)
+				else if (workspaceType.Incremental)
 				{
-					Identifier = $"{Stream.GetEscapedName()}+{AgentType.Workspace}";
+					identifier = $"{stream.GetEscapedName()}+{agentType.Workspace}";
 				}
 				else
 				{
-					Identifier = GetDefaultWorkspaceIdentifier(Stream);
+					identifier = GetDefaultWorkspaceIdentifier(stream);
 				}
 
 				// Create the new workspace
-				Workspace = (new AgentWorkspace(WorkspaceType.Cluster, WorkspaceType.UserName, Identifier, WorkspaceType.Stream ?? Stream.Name, WorkspaceType.View, WorkspaceType.Incremental), WorkspaceType.UseAutoSdk);
+				workspace = (new AgentWorkspace(workspaceType.Cluster, workspaceType.UserName, identifier, workspaceType.Stream ?? stream.Name, workspaceType.View, workspaceType.Incremental), workspaceType.UseAutoSdk);
 				return true;
 			}
 		}
@@ -666,61 +658,61 @@ namespace Horde.Build.Models
 		/// <summary>
 		/// The escaped name of this stream. Removes all non-identifier characters.
 		/// </summary>
-		/// <param name="Stream">The stream object</param>
+		/// <param name="stream">The stream object</param>
 		/// <returns>Escaped name for the stream</returns>
-		public static string GetEscapedName(this IStream Stream)
+		public static string GetEscapedName(this IStream stream)
 		{
-			return Regex.Replace(Stream.Name, @"[^a-zA-Z0-9_]", "+");
+			return Regex.Replace(stream.Name, @"[^a-zA-Z0-9_]", "+");
 		}
 
 		/// <summary>
 		/// Gets the default identifier for workspaces created for this stream. Just includes an escaped depot name.
 		/// </summary>
-		/// <param name="Stream">The stream object</param>
+		/// <param name="stream">The stream object</param>
 		/// <returns>The default workspace identifier</returns>
-		private static string GetDefaultWorkspaceIdentifier(IStream Stream)
+		private static string GetDefaultWorkspaceIdentifier(IStream stream)
 		{
-			return Regex.Replace(Stream.GetEscapedName(), @"^(\+\+[^+]*).*$", "$1");
+			return Regex.Replace(stream.GetEscapedName(), @"^(\+\+[^+]*).*$", "$1");
 		}
 
 		/// <summary>
 		/// Converts to a public response object
 		/// </summary>
-		/// <param name="Stream">The stream object</param>
+		/// <param name="stream">The stream object</param>
 		/// <param name="bIncludeAcl">Whether to include the ACL in the response object</param>
-		/// <param name="ApiTemplateRefs">The template refs for this stream. Passed separately because they have their own ACL.</param>
+		/// <param name="apiTemplateRefs">The template refs for this stream. Passed separately because they have their own ACL.</param>
 		/// <returns>New response instance</returns>
-		public static Api.GetStreamResponse ToApiResponse(this IStream Stream, bool bIncludeAcl, List<GetTemplateRefResponse> ApiTemplateRefs)
+		public static Api.GetStreamResponse ToApiResponse(this IStream stream, bool bIncludeAcl, List<GetTemplateRefResponse> apiTemplateRefs)
 		{
-			List<GetStreamTabResponse> ApiTabs = Stream.Tabs.ConvertAll(x => x.ToResponse());
-			Dictionary<string, GetAgentTypeResponse> ApiAgentTypes = Stream.AgentTypes.ToDictionary(x => x.Key, x => x.Value.ToApiResponse());
-			Dictionary<string, GetWorkspaceTypeResponse> ApiWorkspaceTypes = Stream.WorkspaceTypes.ToDictionary(x => x.Key, x => x.Value.ToApiResponse());
-			GetAclResponse? ApiAcl = (bIncludeAcl && Stream.Acl != null)? new GetAclResponse(Stream.Acl) : null;
-			return new Api.GetStreamResponse(Stream.Id.ToString(), Stream.ProjectId.ToString(), Stream.Name, Stream.ConfigPath, Stream.ConfigRevision, Stream.Order, Stream.NotificationChannel, Stream.NotificationChannelFilter, Stream.TriageChannel, Stream.DefaultPreflight?.ToRequest(), ApiTabs, ApiAgentTypes, ApiWorkspaceTypes, ApiTemplateRefs, ApiAcl, Stream.PausedUntil, Stream.PauseComment);
+			List<GetStreamTabResponse> apiTabs = stream.Tabs.ConvertAll(x => x.ToResponse());
+			Dictionary<string, GetAgentTypeResponse> apiAgentTypes = stream.AgentTypes.ToDictionary(x => x.Key, x => x.Value.ToApiResponse());
+			Dictionary<string, GetWorkspaceTypeResponse> apiWorkspaceTypes = stream.WorkspaceTypes.ToDictionary(x => x.Key, x => x.Value.ToApiResponse());
+			GetAclResponse? apiAcl = (bIncludeAcl && stream.Acl != null)? new GetAclResponse(stream.Acl) : null;
+			return new Api.GetStreamResponse(stream.Id.ToString(), stream.ProjectId.ToString(), stream.Name, stream.ConfigPath, stream.ConfigRevision, stream.Order, stream.NotificationChannel, stream.NotificationChannelFilter, stream.TriageChannel, stream.DefaultPreflight?.ToRequest(), apiTabs, apiAgentTypes, apiWorkspaceTypes, apiTemplateRefs, apiAcl, stream.PausedUntil, stream.PauseComment);
 		}
 
 		/// <summary>
 		/// Converts to an RPC response object
 		/// </summary>
-		/// <param name="Stream">The stream object</param>
+		/// <param name="stream">The stream object</param>
 		/// <returns>New response instance</returns>
-		public static HordeCommon.Rpc.GetStreamResponse ToRpcResponse(this IStream Stream)
+		public static HordeCommon.Rpc.GetStreamResponse ToRpcResponse(this IStream stream)
 		{
-			HordeCommon.Rpc.GetStreamResponse Response = new HordeCommon.Rpc.GetStreamResponse();
-			Response.Name = Stream.Name;
-			Response.AgentTypes.Add(Stream.AgentTypes.ToDictionary(x => x.Key, x => x.Value.ToRpcResponse()));
-			return Response;
+			HordeCommon.Rpc.GetStreamResponse response = new HordeCommon.Rpc.GetStreamResponse();
+			response.Name = stream.Name;
+			response.AgentTypes.Add(stream.AgentTypes.ToDictionary(x => x.Key, x => x.Value.ToRpcResponse()));
+			return response;
 		}
 
 		/// <summary>
 		/// Check if stream is paused for new builds
 		/// </summary>
-		/// <param name="Stream">The stream object</param>
-		/// <param name="CurrentTime">Current time (allow tests to pass in a fake clock)</param>
+		/// <param name="stream">The stream object</param>
+		/// <param name="currentTime">Current time (allow tests to pass in a fake clock)</param>
 		/// <returns>If stream is paused</returns>
-		public static bool IsPaused(this IStream Stream, DateTime CurrentTime)
+		public static bool IsPaused(this IStream stream, DateTime currentTime)
 		{
-			return Stream.PausedUntil != null && Stream.PausedUntil > CurrentTime;
+			return stream.PausedUntil != null && stream.PausedUntil > currentTime;
 		}
 	}
 
diff --git a/Engine/Source/Programs/Horde/Horde.Build/Projects/IStreamCollection.cs b/Engine/Source/Programs/Horde/Horde.Build/Projects/IStreamCollection.cs
--- a/Engine/Source/Programs/Horde/Horde.Build/Projects/IStreamCollection.cs
+++ b/Engine/Source/Programs/Horde/Horde.Build/Projects/IStreamCollection.cs
@@ -1,17 +1,10 @@
 // Copyright Epic Games, Inc. All Rights Reserved.
 
-using Horde.Build.Api;
-using Horde.Build.Models;
-using Horde.Build.Services;
-using Horde.Build.Utilities;
-using MongoDB.Bson;
-using MongoDB.Bson.Serialization.Attributes;
-using MongoDB.Driver;
 using System;
 using System.Collections.Generic;
-using System.Diagnostics.CodeAnalysis;
-using System.Linq;
 using System.Threading.Tasks;
+using Horde.Build.Models;
+using Horde.Build.Utilities;
 
 namespace Horde.Build.Collections
 {
@@ -28,28 +21,28 @@ namespace Horde.Build.Collections
 		/// <summary>
 		/// Creates or replaces a stream configuration
 		/// </summary>
-		/// <param name="Id">Unique id for the new stream</param>
-		/// <param name="Stream">The current stream value. If not-null, this will attempt to replace the existing instance.</param>
-		/// <param name="ConfigPath">Path to the config file</param>
-		/// <param name="Revision">The config file revision</param>
-		/// <param name="ProjectId">The project id</param>
-		/// <param name="Config">The stream configuration</param>
+		/// <param name="id">Unique id for the new stream</param>
+		/// <param name="stream">The current stream value. If not-null, this will attempt to replace the existing instance.</param>
+		/// <param name="configPath">Path to the config file</param>
+		/// <param name="revision">The config file revision</param>
+		/// <param name="projectId">The project id</param>
+		/// <param name="config">The stream configuration</param>
 		/// <returns></returns>
-		Task<IStream?> TryCreateOrReplaceAsync(StreamId Id, IStream? Stream, string ConfigPath, string Revision, ProjectId ProjectId, StreamConfig Config);
+		Task<IStream?> TryCreateOrReplaceAsync(StreamId id, IStream? stream, string configPath, string revision, ProjectId projectId, StreamConfig config);
 
 		/// <summary>
 		/// Gets a stream by ID
 		/// </summary>
-		/// <param name="StreamId">Unique id of the stream</param>
+		/// <param name="streamId">Unique id of the stream</param>
 		/// <returns>The stream document</returns>
-		Task<IStream?> GetAsync(StreamId StreamId);
+		Task<IStream?> GetAsync(StreamId streamId);
 
 		/// <summary>
 		/// Gets a stream's permissions by ID
 		/// </summary>
-		/// <param name="StreamId">Unique id of the stream</param>
+		/// <param name="streamId">Unique id of the stream</param>
 		/// <returns>The stream document</returns>
-		Task<IStreamPermissions?> GetPermissionsAsync(StreamId StreamId);
+		Task<IStreamPermissions?> GetPermissionsAsync(StreamId streamId);
 
 		/// <summary>
 		/// Enumerates all streams
@@ -60,36 +53,36 @@ namespace Horde.Build.Collections
 		/// <summary>
 		/// Gets all the available streams for a project
 		/// </summary>
-		/// <param name="ProjectIds">Unique id of the projects to query</param>
+		/// <param name="projectIds">Unique id of the projects to query</param>
 		/// <returns>List of stream documents</returns>
-		Task<List<IStream>> FindForProjectsAsync(ProjectId[] ProjectIds);
+		Task<List<IStream>> FindForProjectsAsync(ProjectId[] projectIds);
 
 		/// <summary>
 		/// Updates user-facing properties for an existing stream
 		/// </summary>
-		/// <param name="Stream">The stream to update</param>
-		/// <param name="NewPausedUntil">The new datetime for pausing builds</param>
-		/// <param name="NewPauseComment">The reason for pausing</param>
+		/// <param name="stream">The stream to update</param>
+		/// <param name="newPausedUntil">The new datetime for pausing builds</param>
+		/// <param name="newPauseComment">The reason for pausing</param>
 		/// <returns>The updated stream if successful, null otherwise</returns>
-		Task<IStream?> TryUpdatePauseStateAsync(IStream Stream, DateTime? NewPausedUntil, string? NewPauseComment);
+		Task<IStream?> TryUpdatePauseStateAsync(IStream stream, DateTime? newPausedUntil, string? newPauseComment);
 
 		/// <summary>
 		/// Attempts to update the last trigger time for a schedule
 		/// </summary>
-		/// <param name="Stream">The stream to update</param>
-		/// <param name="TemplateRefId">The template ref id</param>
-		/// <param name="LastTriggerTimeUtc">New last trigger time for the schedule</param>
-		/// <param name="LastTriggerChange">New last trigger changelist for the schedule</param>
-		/// <param name="NewActiveJobs">New list of active jobs</param>
+		/// <param name="stream">The stream to update</param>
+		/// <param name="templateRefId">The template ref id</param>
+		/// <param name="lastTriggerTimeUtc">New last trigger time for the schedule</param>
+		/// <param name="lastTriggerChange">New last trigger changelist for the schedule</param>
+		/// <param name="newActiveJobs">New list of active jobs</param>
 		/// <returns>The updated stream if successful, null otherwise</returns>
-		Task<IStream?> TryUpdateScheduleTriggerAsync(IStream Stream, TemplateRefId TemplateRefId, DateTime? LastTriggerTimeUtc, int? LastTriggerChange, List<JobId> NewActiveJobs);
+		Task<IStream?> TryUpdateScheduleTriggerAsync(IStream stream, TemplateRefId templateRefId, DateTime? lastTriggerTimeUtc, int? lastTriggerChange, List<JobId> newActiveJobs);
 
 		/// <summary>
 		/// Delete a stream
 		/// </summary>
-		/// <param name="StreamId">Unique id of the stream</param>
+		/// <param name="streamId">Unique id of the stream</param>
 		/// <returns>Async task</returns>
-		Task DeleteAsync(StreamId StreamId);
+		Task DeleteAsync(StreamId streamId);
 	}
 
 	static class StreamCollectionExtensions
@@ -97,24 +90,24 @@ namespace Horde.Build.Collections
 		/// <summary>
 		/// Creates or replaces a stream configuration
 		/// </summary>
-		/// <param name="StreamCollection">The stream collection</param>
-		/// <param name="Id">Unique id for the new stream</param>
-		/// <param name="Stream">The current stream value. If not-null, this will attempt to replace the existing instance.</param>
-		/// <param name="ConfigPath">Path to the config file</param>
-		/// <param name="Revision">The config file revision</param>
-		/// <param name="ProjectId">The project id</param>
-		/// <param name="Config">The stream configuration</param>
+		/// <param name="streamCollection">The stream collection</param>
+		/// <param name="id">Unique id for the new stream</param>
+		/// <param name="stream">The current stream value. If not-null, this will attempt to replace the existing instance.</param>
+		/// <param name="configPath">Path to the config file</param>
+		/// <param name="revision">The config file revision</param>
+		/// <param name="projectId">The project id</param>
+		/// <param name="config">The stream configuration</param>
 		/// <returns></returns>
-		public static async Task<IStream> CreateOrReplaceAsync(this IStreamCollection StreamCollection, StreamId Id, IStream? Stream, string ConfigPath, string Revision, ProjectId ProjectId, StreamConfig Config)
+		public static async Task<IStream> CreateOrReplaceAsync(this IStreamCollection streamCollection, StreamId id, IStream? stream, string configPath, string revision, ProjectId projectId, StreamConfig config)
 		{
 			for (; ; )
 			{
-				Stream = await StreamCollection.TryCreateOrReplaceAsync(Id, Stream, ConfigPath, Revision, ProjectId, Config);
-				if (Stream != null)
+				stream = await streamCollection.TryCreateOrReplaceAsync(id, stream, configPath, revision, projectId, config);
+				if (stream != null)
 				{
-					return Stream;
+					return stream;
 				}
-				Stream = await StreamCollection.GetAsync(Id);
+				stream = await streamCollection.GetAsync(id);
 			}
 		}
 	}
