Horde: Fix coding conventions within Horde.Build.

diff --git a/Engine/Source/Programs/Horde/Horde.Build/Projects/Impl/ProjectCollection.cs b/Engine/Source/Programs/Horde/Horde.Build/Projects/Impl/ProjectCollection.cs
--- a/Engine/Source/Programs/Horde/Horde.Build/Projects/Impl/ProjectCollection.cs
+++ b/Engine/Source/Programs/Horde/Horde.Build/Projects/Impl/ProjectCollection.cs
@@ -1,20 +1,16 @@
 ﻿// Copyright Epic Games, Inc. All Rights Reserved.
 
+using System;
+using System.Collections.Generic;
+using System.Diagnostics.CodeAnalysis;
+using System.Linq;
+using System.Threading.Tasks;
 using Horde.Build.Acls;
-using Horde.Build.Api;
 using Horde.Build.Models;
 using Horde.Build.Services;
 using Horde.Build.Utilities;
-using Microsoft.Extensions.Options;
-using MongoDB.Bson;
 using MongoDB.Bson.Serialization.Attributes;
 using MongoDB.Driver;
-using System;
-using System.Collections.Generic;
-using System.Diagnostics.CodeAnalysis;
-using System.Linq;
-using System.Security.Claims;
-using System.Threading.Tasks;
 
 namespace Horde.Build.Collections.Impl
 {
@@ -56,10 +52,10 @@ namespace Horde.Build.Collections.Impl
 				Name = null!;
 			}
 
-			public ProjectDocument(ProjectId Id, string Name)
+			public ProjectDocument(ProjectId id, string name)
 			{
-				this.Id = Id;
-				this.Name = Name;
+				Id = id;
+				Name = name;
 			}
 		}
 
@@ -97,78 +93,78 @@ namespace Horde.Build.Collections.Impl
 		/// <summary>
 		/// Collection of project documents
 		/// </summary>
-		IMongoCollection<ProjectDocument> Projects;
+		readonly IMongoCollection<ProjectDocument> _projects;
 
 		/// <summary>
 		/// Collection of project logo documents
 		/// </summary>
-		IMongoCollection<ProjectLogoDocument> ProjectLogos;
+		readonly IMongoCollection<ProjectLogoDocument> _projectLogos;
 
 		/// <summary>
 		/// Constructor
 		/// </summary>
-		/// <param name="DatabaseService">The database service instance</param>
-		public ProjectCollection(DatabaseService DatabaseService)
+		/// <param name="databaseService">The database service instance</param>
+		public ProjectCollection(DatabaseService databaseService)
 		{
-			Projects = DatabaseService.GetCollection<ProjectDocument>("Projects");
-			ProjectLogos = DatabaseService.GetCollection<ProjectLogoDocument>("ProjectLogos");
+			_projects = databaseService.GetCollection<ProjectDocument>("Projects");
+			_projectLogos = databaseService.GetCollection<ProjectLogoDocument>("ProjectLogos");
 		}
 
 		/// <inheritdoc/>
-		public async Task<IProject?> AddOrUpdateAsync(ProjectId Id, string ConfigPath, string Revision, int Order, ProjectConfig Config)
+		public async Task<IProject?> AddOrUpdateAsync(ProjectId id, string configPath, string revision, int order, ProjectConfig config)
 		{
-			ProjectDocument NewProject = new ProjectDocument(Id, Config.Name);
-			NewProject.ConfigPath = ConfigPath;
-			NewProject.ConfigRevision = Revision;
-			NewProject.Order = Order;
-			NewProject.Categories = Config.Categories.ConvertAll(x => new StreamCategory(x));
-			NewProject.Acl = Acl.Merge(new Acl(), Config.Acl);
-
-			await Projects.FindOneAndReplaceAsync<ProjectDocument>(x => x.Id == Id, NewProject, new FindOneAndReplaceOptions<ProjectDocument> { IsUpsert = true });
-			return NewProject;
+			ProjectDocument newProject = new ProjectDocument(id, config.Name);
+			newProject.ConfigPath = configPath;
+			newProject.ConfigRevision = revision;
+			newProject.Order = order;
+			newProject.Categories = config.Categories.ConvertAll(x => new StreamCategory(x));
+			newProject.Acl = Acl.Merge(new Acl(), config.Acl);
+
+			await _projects.FindOneAndReplaceAsync<ProjectDocument>(x => x.Id == id, newProject, new FindOneAndReplaceOptions<ProjectDocument> { IsUpsert = true });
+			return newProject;
 		}
 
 		/// <inheritdoc/>
 		public async Task<List<IProject>> FindAllAsync()
 		{
-			List<ProjectDocument> Results = await Projects.Find(x => !x.Deleted).ToListAsync();
-			return Results.OrderBy(x => x.Order).ThenBy(x => x.Name).Select<ProjectDocument, IProject>(x => x).ToList();
+			List<ProjectDocument> results = await _projects.Find(x => !x.Deleted).ToListAsync();
+			return results.OrderBy(x => x.Order).ThenBy(x => x.Name).Select<ProjectDocument, IProject>(x => x).ToList();
 		}
 
 		/// <inheritdoc/>
-		public async Task<IProject?> GetAsync(ProjectId ProjectId)
+		public async Task<IProject?> GetAsync(ProjectId projectId)
 		{
-			return await Projects.Find<ProjectDocument>(x => x.Id == ProjectId).FirstOrDefaultAsync();
+			return await _projects.Find<ProjectDocument>(x => x.Id == projectId).FirstOrDefaultAsync();
 		}
 
 		/// <inheritdoc/>
-		public async Task<IProjectLogo?> GetLogoAsync(ProjectId ProjectId)
+		public async Task<IProjectLogo?> GetLogoAsync(ProjectId projectId)
 		{
-			return await ProjectLogos.Find(x => x.Id == ProjectId).FirstOrDefaultAsync();
+			return await _projectLogos.Find(x => x.Id == projectId).FirstOrDefaultAsync();
 		}
 
 		/// <inheritdoc/>
-		public async Task SetLogoAsync(ProjectId ProjectId, string LogoPath, string LogoRevision, string MimeType, byte[] Data)
+		public async Task SetLogoAsync(ProjectId projectId, string logoPath, string logoRevision, string mimeType, byte[] data)
 		{
-			ProjectLogoDocument Logo = new ProjectLogoDocument();
-			Logo.Id = ProjectId;
-			Logo.Path = LogoPath;
-			Logo.Revision = LogoRevision;
-			Logo.MimeType = MimeType;
-			Logo.Data = Data;
-			await ProjectLogos.ReplaceOneAsync(x => x.Id == ProjectId, Logo, new ReplaceOptions { IsUpsert = true });
+			ProjectLogoDocument logo = new ProjectLogoDocument();
+			logo.Id = projectId;
+			logo.Path = logoPath;
+			logo.Revision = logoRevision;
+			logo.MimeType = mimeType;
+			logo.Data = data;
+			await _projectLogos.ReplaceOneAsync(x => x.Id == projectId, logo, new ReplaceOptions { IsUpsert = true });
 		}
 
 		/// <inheritdoc/>
-		public async Task<IProjectPermissions?> GetPermissionsAsync(ProjectId ProjectId)
+		public async Task<IProjectPermissions?> GetPermissionsAsync(ProjectId projectId)
 		{
-			return await Projects.Find<ProjectDocument>(x => x.Id == ProjectId).Project<ProjectPermissions>(ProjectPermissions.Projection).FirstOrDefaultAsync();
+			return await _projects.Find<ProjectDocument>(x => x.Id == projectId).Project<ProjectPermissions>(ProjectPermissions.Projection).FirstOrDefaultAsync();
 		}
 
 		/// <inheritdoc/>
-		public async Task DeleteAsync(ProjectId ProjectId)
+		public async Task DeleteAsync(ProjectId projectId)
 		{
-			await Projects.UpdateOneAsync(x => x.Id == ProjectId, Builders<ProjectDocument>.Update.Set(x => x.Deleted, true));
+			await _projects.UpdateOneAsync(x => x.Id == projectId, Builders<ProjectDocument>.Update.Set(x => x.Deleted, true));
 		}
 	}
 }
diff --git a/Engine/Source/Programs/Horde/Horde.Build/Projects/Impl/ProjectService.cs b/Engine/Source/Programs/Horde/Horde.Build/Projects/Impl/ProjectService.cs
--- a/Engine/Source/Programs/Horde/Horde.Build/Projects/Impl/ProjectService.cs
+++ b/Engine/Source/Programs/Horde/Horde.Build/Projects/Impl/ProjectService.cs
@@ -1,18 +1,12 @@
 ﻿// Copyright Epic Games, Inc. All Rights Reserved.
 
+using System.Collections.Generic;
+using System.Security.Claims;
+using System.Threading.Tasks;
 using Horde.Build.Acls;
-using Horde.Build.Api;
 using Horde.Build.Collections;
 using Horde.Build.Models;
 using Horde.Build.Utilities;
-using Microsoft.Extensions.Options;
-using MongoDB.Bson;
-using MongoDB.Driver;
-using System;
-using System.Collections.Generic;
-using System.Linq;
-using System.Security.Claims;
-using System.Threading.Tasks;
 
 namespace Horde.Build.Services
 {
@@ -37,27 +31,27 @@ namespace Horde.Build.Services
 		/// <summary>
 		/// The ACL service
 		/// </summary>
-		AclService AclService;
+		readonly AclService _aclService;
 
 		/// <summary>
 		/// Collection of project documents
 		/// </summary>
-		IProjectCollection Projects;
+		readonly IProjectCollection _projects;
 
 		/// <summary>
 		/// Accessor for the collection of project documents
 		/// </summary>
-		public IProjectCollection Collection => Projects;
+		public IProjectCollection Collection => _projects;
 
 		/// <summary>
 		/// Constructor
 		/// </summary>
-		/// <param name="AclService">The ACL service</param>
-		/// <param name="Projects">Collection of project documents</param>
-		public ProjectService(AclService AclService, IProjectCollection Projects)
+		/// <param name="aclService">The ACL service</param>
+		/// <param name="projects">Collection of project documents</param>
+		public ProjectService(AclService aclService, IProjectCollection projects)
 		{
-			this.AclService = AclService;
-			this.Projects = Projects;
+			_aclService = aclService;
+			_projects = projects;
 		}
 
 		/// <summary>
@@ -66,93 +60,93 @@ namespace Horde.Build.Services
 		/// <returns>List of project documents</returns>
 		public Task<List<IProject>> GetProjectsAsync()
 		{
-			return Projects.FindAllAsync();
+			return _projects.FindAllAsync();
 		}
 
 		/// <summary>
 		/// Gets a project by ID
 		/// </summary>
-		/// <param name="ProjectId">Unique id of the project</param>
+		/// <param name="projectId">Unique id of the project</param>
 		/// <returns>The project document</returns>
-		public Task<IProject?> GetProjectAsync(ProjectId ProjectId)
+		public Task<IProject?> GetProjectAsync(ProjectId projectId)
 		{
-			return Projects.GetAsync(ProjectId);
+			return _projects.GetAsync(projectId);
 		}
 
 		/// <summary>
 		/// Gets a project's permissions info by ID
 		/// </summary>
-		/// <param name="ProjectId">Unique id of the project</param>
+		/// <param name="projectId">Unique id of the project</param>
 		/// <returns>The project document</returns>
-		public Task<IProjectPermissions?> GetProjectPermissionsAsync(ProjectId ProjectId)
+		public Task<IProjectPermissions?> GetProjectPermissionsAsync(ProjectId projectId)
 		{
-			return Projects.GetPermissionsAsync(ProjectId);
+			return _projects.GetPermissionsAsync(projectId);
 		}
 
 		/// <summary>
 		/// Deletes a project by id
 		/// </summary>
-		/// <param name="ProjectId">Unique id of the project</param>
-		public async Task DeleteProjectAsync(ProjectId ProjectId)
+		/// <param name="projectId">Unique id of the project</param>
+		public async Task DeleteProjectAsync(ProjectId projectId)
 		{
-			await Projects.DeleteAsync(ProjectId);
+			await _projects.DeleteAsync(projectId);
 		}
 
 		/// <summary>
 		/// Determines if the user is authorized to perform an action on a particular project
 		/// </summary>
-		/// <param name="Acl">Acl for the project to check</param>
-		/// <param name="Action">The action being performed</param>
-		/// <param name="User">The principal to authorize</param>
-		/// <param name="Cache">Cache for the scope table</param>
+		/// <param name="acl">Acl for the project to check</param>
+		/// <param name="action">The action being performed</param>
+		/// <param name="user">The principal to authorize</param>
+		/// <param name="cache">Cache for the scope table</param>
 		/// <returns>True if the action is authorized</returns>
-		private Task<bool> AuthorizeAsync(Acl? Acl, AclAction Action, ClaimsPrincipal User, GlobalPermissionsCache? Cache)
+		private Task<bool> AuthorizeAsync(Acl? acl, AclAction action, ClaimsPrincipal user, GlobalPermissionsCache? cache)
 		{
-			bool? Result = Acl?.Authorize(Action, User);
-			if (Result == null)
+			bool? result = acl?.Authorize(action, user);
+			if (result == null)
 			{
-				return AclService.AuthorizeAsync(Action, User, Cache);
+				return _aclService.AuthorizeAsync(action, user, cache);
 			}
 			else
 			{
-				return Task.FromResult(Result.Value);
+				return Task.FromResult(result.Value);
 			}
 		}
 
 		/// <summary>
 		/// Determines if the user is authorized to perform an action on a particular project
 		/// </summary>
-		/// <param name="Project">The project to check</param>
-		/// <param name="Action">The action being performed</param>
-		/// <param name="User">The principal to authorize</param>
-		/// <param name="Cache">Cache for the scope table</param>
+		/// <param name="project">The project to check</param>
+		/// <param name="action">The action being performed</param>
+		/// <param name="user">The principal to authorize</param>
+		/// <param name="cache">Cache for the scope table</param>
 		/// <returns>True if the action is authorized</returns>
-		public Task<bool> AuthorizeAsync(IProject Project, AclAction Action, ClaimsPrincipal User, GlobalPermissionsCache? Cache)
+		public Task<bool> AuthorizeAsync(IProject project, AclAction action, ClaimsPrincipal user, GlobalPermissionsCache? cache)
 		{
-			return AuthorizeAsync(Project.Acl, Action, User, Cache);
+			return AuthorizeAsync(project.Acl, action, user, cache);
 		}
 
 		/// <summary>
 		/// Determines if the user is authorized to perform an action on a particular project
 		/// </summary>
-		/// <param name="ProjectId">The project id to check</param>
-		/// <param name="Action">The action being performed</param>
-		/// <param name="User">The principal to authorize</param>
-		/// <param name="Cache">Cache for project permissions</param>
+		/// <param name="projectId">The project id to check</param>
+		/// <param name="action">The action being performed</param>
+		/// <param name="user">The principal to authorize</param>
+		/// <param name="cache">Cache for project permissions</param>
 		/// <returns>True if the action is authorized</returns>
-		public async Task<bool> AuthorizeAsync(ProjectId ProjectId, AclAction Action, ClaimsPrincipal User, ProjectPermissionsCache? Cache)
+		public async Task<bool> AuthorizeAsync(ProjectId projectId, AclAction action, ClaimsPrincipal user, ProjectPermissionsCache? cache)
 		{
-			IProjectPermissions? Permissions;
-			if (Cache == null)
+			IProjectPermissions? permissions;
+			if (cache == null)
 			{
-				Permissions = await GetProjectPermissionsAsync(ProjectId);
+				permissions = await GetProjectPermissionsAsync(projectId);
 			}
-			else if (!Cache.Projects.TryGetValue(ProjectId, out Permissions))
+			else if (!cache.Projects.TryGetValue(projectId, out permissions))
 			{
-				Permissions = await GetProjectPermissionsAsync(ProjectId);
-				Cache.Projects.Add(ProjectId, Permissions);
+				permissions = await GetProjectPermissionsAsync(projectId);
+				cache.Projects.Add(projectId, permissions);
 			}
-			return await AuthorizeAsync(Permissions?.Acl, Action, User, Cache);
+			return await AuthorizeAsync(permissions?.Acl, action, user, cache);
 		}
 	}
 }
diff --git a/Engine/Source/Programs/Horde/Horde.Build/Projects/Impl/StreamCollection.cs b/Engine/Source/Programs/Horde/Horde.Build/Projects/Impl/StreamCollection.cs
--- a/Engine/Source/Programs/Horde/Horde.Build/Projects/Impl/StreamCollection.cs
+++ b/Engine/Source/Programs/Horde/Horde.Build/Projects/Impl/StreamCollection.cs
@@ -1,20 +1,19 @@
-﻿// Copyright Epic Games, Inc. All Rights Reserved.
+// Copyright Epic Games, Inc. All Rights Reserved.
 
-using HordeCommon;
+using System;
+using System.Collections.Generic;
+using System.Diagnostics.CodeAnalysis;
+using System.IO;
+using System.Linq;
+using System.Threading.Tasks;
 using Horde.Build.Acls;
 using Horde.Build.Api;
 using Horde.Build.Models;
 using Horde.Build.Services;
 using Horde.Build.Utilities;
-using MongoDB.Bson;
+using HordeCommon;
 using MongoDB.Bson.Serialization.Attributes;
 using MongoDB.Driver;
-using System;
-using System.Collections.Generic;
-using System.Diagnostics.CodeAnalysis;
-using System.IO;
-using System.Linq;
-using System.Threading.Tasks;
 
 namespace Horde.Build.Collections.Impl
 {
@@ -83,11 +82,11 @@ namespace Horde.Build.Collections.Impl
 				Name = null!;
 			}
 
-			public StreamDocument(StreamId Id, string Name, ProjectId ProjectId)
+			public StreamDocument(StreamId id, string name, ProjectId projectId)
 			{
-				this.Id = Id;
-				this.Name = Name;
-				this.ProjectId = ProjectId;
+				Id = id;
+				Name = name;
+				ProjectId = projectId;
 			}
 		}
 
@@ -106,156 +105,156 @@ namespace Horde.Build.Collections.Impl
 		/// <summary>
 		/// The stream collection
 		/// </summary>
-		IMongoCollection<StreamDocument> Streams;
+		readonly IMongoCollection<StreamDocument> _streams;
 
 		/// <summary>
 		/// Clock
 		/// </summary>
-		IClock Clock;
+		readonly IClock _clock;
 
 		/// <summary>
 		/// The template collection
 		/// </summary>
-		ITemplateCollection TemplateCollection;
+		readonly ITemplateCollection _templateCollection;
 
 		/// <summary>
 		/// Constructor
 		/// </summary>
-		/// <param name="DatabaseService">The database service instance</param>
-		/// <param name="Clock"></param>
-		/// <param name="TemplateCollection"></param>
-		public StreamCollection(DatabaseService DatabaseService, IClock Clock, ITemplateCollection TemplateCollection)
+		/// <param name="databaseService">The database service instance</param>
+		/// <param name="clock"></param>
+		/// <param name="templateCollection"></param>
+		public StreamCollection(DatabaseService databaseService, IClock clock, ITemplateCollection templateCollection)
 		{
-			this.Streams = DatabaseService.GetCollection<StreamDocument>("Streams");
-			this.Clock = Clock;
-			this.TemplateCollection = TemplateCollection;
+			_streams = databaseService.GetCollection<StreamDocument>("Streams");
+			_clock = clock;
+			_templateCollection = templateCollection;
 		}
 
 		/// <inheritdoc/>
-		public async Task<IStream?> TryCreateOrReplaceAsync(StreamId Id, IStream? Stream, string ConfigPath, string Revision, ProjectId ProjectId, StreamConfig Config)
+		public async Task<IStream?> TryCreateOrReplaceAsync(StreamId id, IStream? stream, string configPath, string revision, ProjectId projectId, StreamConfig config)
 		{
-			List<StreamTab> Tabs = Config.Tabs.ConvertAll(x => StreamTab.FromRequest(x));
-			Dictionary<TemplateRefId, TemplateRef> TemplateRefs = await CreateTemplateRefsAsync(Config.Templates, Stream, TemplateCollection);
+			List<StreamTab> tabs = config.Tabs.ConvertAll(x => StreamTab.FromRequest(x));
+			Dictionary<TemplateRefId, TemplateRef> templateRefs = await CreateTemplateRefsAsync(config.Templates, stream, _templateCollection);
 
-			Dictionary<string, AgentType> AgentTypes = new Dictionary<string, AgentType>();
-			if (Config.AgentTypes != null)
+			Dictionary<string, AgentType> agentTypes = new Dictionary<string, AgentType>();
+			if (config.AgentTypes != null)
 			{
-				AgentTypes = Config.AgentTypes.Where(x => x.Value != null).ToDictionary(x => x.Key, x => new AgentType(x.Value!));
+				agentTypes = config.AgentTypes.Where(x => x.Value != null).ToDictionary(x => x.Key, x => new AgentType(x.Value!));
 			}
 
-			Dictionary<string, WorkspaceType> WorkspaceTypes = new Dictionary<string, WorkspaceType>();
-			if (Config.WorkspaceTypes != null)
+			Dictionary<string, WorkspaceType> workspaceTypes = new Dictionary<string, WorkspaceType>();
+			if (config.WorkspaceTypes != null)
 			{
-				WorkspaceTypes = Config.WorkspaceTypes.Where(x => x.Value != null).ToDictionary(x => x.Key, x => new WorkspaceType(x.Value!));
+				workspaceTypes = config.WorkspaceTypes.Where(x => x.Value != null).ToDictionary(x => x.Key, x => new WorkspaceType(x.Value!));
 			}
 
-			DefaultPreflight? DefaultPreflight = Config.DefaultPreflight?.ToModel();
-			if (DefaultPreflight == null && Config.DefaultPreflightTemplate != null)
+			DefaultPreflight? defaultPreflight = config.DefaultPreflight?.ToModel();
+			if (defaultPreflight == null && config.DefaultPreflightTemplate != null)
 			{
-				DefaultPreflight = new DefaultPreflight(new TemplateRefId(Config.DefaultPreflightTemplate), null);
+				defaultPreflight = new DefaultPreflight(new TemplateRefId(config.DefaultPreflightTemplate), null);
 			}
 
-			Validate(Id, DefaultPreflight, TemplateRefs, Tabs, AgentTypes, WorkspaceTypes);
+			Validate(id, defaultPreflight, templateRefs, tabs, agentTypes, workspaceTypes);
 
-			Acl? Acl = Acl.Merge(new Acl(), Config.Acl);
-			if (Stream == null)
+			Acl? acl = Acl.Merge(new Acl(), config.Acl);
+			if (stream == null)
 			{
-				return await TryCreateAsync(Id, ProjectId, ConfigPath, Revision, Config, DefaultPreflight, Tabs, AgentTypes, WorkspaceTypes, TemplateRefs, Acl);
+				return await TryCreateAsync(id, projectId, configPath, revision, config, defaultPreflight, tabs, agentTypes, workspaceTypes, templateRefs, acl);
 			}
 			else
 			{
-				return await TryReplaceAsync(Stream, ProjectId, ConfigPath, Revision, Config, DefaultPreflight, Tabs, AgentTypes, WorkspaceTypes, TemplateRefs, Acl);
+				return await TryReplaceAsync(stream, projectId, configPath, revision, config, defaultPreflight, tabs, agentTypes, workspaceTypes, templateRefs, acl);
 			}
 		}
 
 		/// <summary>
 		/// Creates a list of template refs from a set of request objects
 		/// </summary>
-		/// <param name="Requests">Request objects</param>
-		/// <param name="Stream">The current stream state</param>
-		/// <param name="TemplateCollection">The template service</param>
+		/// <param name="requests">Request objects</param>
+		/// <param name="stream">The current stream state</param>
+		/// <param name="templateCollection">The template service</param>
 		/// <returns>List of new template references</returns>
-		async Task<Dictionary<TemplateRefId, TemplateRef>> CreateTemplateRefsAsync(List<CreateTemplateRefRequest> Requests, IStream? Stream, ITemplateCollection TemplateCollection)
+		async Task<Dictionary<TemplateRefId, TemplateRef>> CreateTemplateRefsAsync(List<CreateTemplateRefRequest> requests, IStream? stream, ITemplateCollection templateCollection)
 		{
-			Dictionary<TemplateRefId, TemplateRef> NewTemplateRefs = new Dictionary<TemplateRefId, TemplateRef>();
-			foreach (CreateTemplateRefRequest Request in Requests)
+			Dictionary<TemplateRefId, TemplateRef> newTemplateRefs = new Dictionary<TemplateRefId, TemplateRef>();
+			foreach (CreateTemplateRefRequest request in requests)
 			{
 				// Create the template
-				ITemplate NewTemplate = await TemplateCollection.AddAsync(Request.Name, Request.Priority, Request.AllowPreflights, Request.UpdateIssues, Request.PromoteIssuesByDefault, Request.InitialAgentType, Request.SubmitNewChange, Request.SubmitDescription, Request.Arguments, Request.Parameters.ConvertAll(x => x.ToModel()));
+				ITemplate newTemplate = await templateCollection.AddAsync(request.Name, request.Priority, request.AllowPreflights, request.UpdateIssues, request.PromoteIssuesByDefault, request.InitialAgentType, request.SubmitNewChange, request.SubmitDescription, request.Arguments, request.Parameters.ConvertAll(x => x.ToModel()));
 
 				// Get an identifier for the new template ref
-				TemplateRefId NewTemplateRefId;
-				if (Request.Id != null)
+				TemplateRefId newTemplateRefId;
+				if (request.Id != null)
 				{
-					NewTemplateRefId = new TemplateRefId(Request.Id);
+					newTemplateRefId = new TemplateRefId(request.Id);
 				}
 				else
 				{
-					NewTemplateRefId = TemplateRefId.Sanitize(Request.Name);
+					newTemplateRefId = TemplateRefId.Sanitize(request.Name);
 				}
 
 				// Create the schedule object
-				Schedule? Schedule = null;
-				if (Request.Schedule != null)
+				Schedule? schedule = null;
+				if (request.Schedule != null)
 				{
-					Schedule = Request.Schedule.ToModel(Clock.UtcNow);
+					schedule = request.Schedule.ToModel(_clock.UtcNow);
 				}
 
 				// Add it to the list
-				TemplateRef NewTemplateRef = new TemplateRef(NewTemplate, Request.ShowUgsBadges, Request.ShowUgsAlerts, Request.NotificationChannel, Request.NotificationChannelFilter, Request.TriageChannel, Schedule, Request.ChainedJobs?.ConvertAll(x => new ChainedJobTemplate(x)), Acl.Merge(null, Request.Acl));
-				if (Stream != null && Stream.Templates.TryGetValue(NewTemplateRefId, out TemplateRef? OldTemplateRef))
+				TemplateRef newTemplateRef = new TemplateRef(newTemplate, request.ShowUgsBadges, request.ShowUgsAlerts, request.NotificationChannel, request.NotificationChannelFilter, request.TriageChannel, schedule, request.ChainedJobs?.ConvertAll(x => new ChainedJobTemplate(x)), Acl.Merge(null, request.Acl));
+				if (stream != null && stream.Templates.TryGetValue(newTemplateRefId, out TemplateRef? oldTemplateRef))
 				{
-					if (OldTemplateRef.Schedule != null && NewTemplateRef.Schedule != null)
+					if (oldTemplateRef.Schedule != null && newTemplateRef.Schedule != null)
 					{
-						NewTemplateRef.Schedule.CopyState(OldTemplateRef.Schedule);
+						newTemplateRef.Schedule.CopyState(oldTemplateRef.Schedule);
 					}
 				}
-				NewTemplateRefs.Add(NewTemplateRefId, NewTemplateRef);
+				newTemplateRefs.Add(newTemplateRefId, newTemplateRef);
 			}
-			foreach (TemplateRef TemplateRef in NewTemplateRefs.Values)
+			foreach (TemplateRef templateRef in newTemplateRefs.Values)
 			{
-				if (TemplateRef.ChainedJobs != null)
+				if (templateRef.ChainedJobs != null)
 				{
-					foreach (ChainedJobTemplate ChainedJob in TemplateRef.ChainedJobs)
+					foreach (ChainedJobTemplate chainedJob in templateRef.ChainedJobs)
 					{
-						if (!NewTemplateRefs.ContainsKey(ChainedJob.TemplateRefId))
+						if (!newTemplateRefs.ContainsKey(chainedJob.TemplateRefId))
 						{
-							throw new InvalidDataException($"Invalid template ref id '{ChainedJob.TemplateRefId}");
+							throw new InvalidDataException($"Invalid template ref id '{chainedJob.TemplateRefId}");
 						}
 					}
 				}
 			}
-			return NewTemplateRefs;
+			return newTemplateRefs;
 		}
 
 		/// <inheritdoc/>
-		async Task<IStream?> TryCreateAsync(StreamId Id, ProjectId ProjectId, string ConfigPath, string ConfigRevision, StreamConfig Config, DefaultPreflight? DefaultPreflight, List<StreamTab> Tabs, Dictionary<string, AgentType> AgentTypes, Dictionary<string, WorkspaceType> WorkspaceTypes, Dictionary<TemplateRefId, TemplateRef> TemplateRefs, Acl? Acl)
+		async Task<IStream?> TryCreateAsync(StreamId id, ProjectId projectId, string configPath, string configRevision, StreamConfig config, DefaultPreflight? defaultPreflight, List<StreamTab> tabs, Dictionary<string, AgentType> agentTypes, Dictionary<string, WorkspaceType> workspaceTypes, Dictionary<TemplateRefId, TemplateRef> templateRefs, Acl? acl)
 		{
-			StreamDocument NewStream = new StreamDocument(Id, Config.Name, ProjectId);
-			NewStream.ClusterName = Config.ClusterName;
-			NewStream.ConfigPath = ConfigPath;
-			NewStream.ConfigRevision = ConfigRevision;
-			NewStream.Order = Config.Order ?? StreamDocument.DefaultOrder;
-			NewStream.NotificationChannel = Config.NotificationChannel;
-			NewStream.NotificationChannelFilter = Config.NotificationChannelFilter;
-			NewStream.TriageChannel = Config.TriageChannel;
-			NewStream.DefaultPreflight = DefaultPreflight;
-			NewStream.Tabs = Tabs;
-			NewStream.AgentTypes = AgentTypes;
-			NewStream.WorkspaceTypes = WorkspaceTypes;
-			NewStream.Templates = TemplateRefs;
-			NewStream.ReplicationMode = Config.ReplicationMode;
-			NewStream.ReplicationFilter = Config.ReplicationFilter;
-			NewStream.Acl = Acl;
+			StreamDocument newStream = new StreamDocument(id, config.Name, projectId);
+			newStream.ClusterName = config.ClusterName;
+			newStream.ConfigPath = configPath;
+			newStream.ConfigRevision = configRevision;
+			newStream.Order = config.Order ?? StreamDocument.DefaultOrder;
+			newStream.NotificationChannel = config.NotificationChannel;
+			newStream.NotificationChannelFilter = config.NotificationChannelFilter;
+			newStream.TriageChannel = config.TriageChannel;
+			newStream.DefaultPreflight = defaultPreflight;
+			newStream.Tabs = tabs;
+			newStream.AgentTypes = agentTypes;
+			newStream.WorkspaceTypes = workspaceTypes;
+			newStream.Templates = templateRefs;
+			newStream.ReplicationMode = config.ReplicationMode;
+			newStream.ReplicationFilter = config.ReplicationFilter;
+			newStream.Acl = acl;
 
 			try
 			{
-				await Streams.InsertOneAsync(NewStream);
-				return NewStream;
+				await _streams.InsertOneAsync(newStream);
+				return newStream;
 			}
-			catch (MongoWriteException Ex)
+			catch (MongoWriteException ex)
 			{
-				if (Ex.WriteError.Category == ServerErrorCategory.DuplicateKey)
+				if (ex.WriteError.Category == ServerErrorCategory.DuplicateKey)
 				{
 					return null;
 				}
@@ -267,167 +266,167 @@ namespace Horde.Build.Collections.Impl
 		}
 
 		/// <inheritdoc/>
-		async Task<IStream?> TryReplaceAsync(IStream StreamInterface, ProjectId ProjectId, string ConfigPath, string ConfigRevision, StreamConfig Config, DefaultPreflight? DefaultPreflight, List<StreamTab> Tabs, Dictionary<string, AgentType>? AgentTypes, Dictionary<string, WorkspaceType>? WorkspaceTypes, Dictionary<TemplateRefId, TemplateRef>? TemplateRefs, Acl? Acl)
+		async Task<IStream?> TryReplaceAsync(IStream streamInterface, ProjectId projectId, string configPath, string configRevision, StreamConfig config, DefaultPreflight? defaultPreflight, List<StreamTab> tabs, Dictionary<string, AgentType>? agentTypes, Dictionary<string, WorkspaceType>? workspaceTypes, Dictionary<TemplateRefId, TemplateRef>? templateRefs, Acl? acl)
 		{
-			int Order = Config.Order ?? StreamDocument.DefaultOrder;
-
-			StreamDocument Stream = (StreamDocument)StreamInterface;
-
-			UpdateDefinitionBuilder<StreamDocument> UpdateBuilder = Builders<StreamDocument>.Update;
-
-			List<UpdateDefinition<StreamDocument>> Updates = new List<UpdateDefinition<StreamDocument>>();
-			Updates.Add(UpdateBuilder.Set(x => x.Name, Config.Name));
-			Updates.Add(UpdateBuilder.Set(x => x.ProjectId, ProjectId));
-			Updates.Add(UpdateBuilder.Set(x => x.ClusterName, Config.ClusterName));
-			Updates.Add(UpdateBuilder.Set(x => x.ConfigPath, ConfigPath));
-			Updates.Add(UpdateBuilder.Set(x => x.ConfigRevision, ConfigRevision));
-			Updates.Add(UpdateBuilder.Set(x => x.Order, Order));
-			Updates.Add(UpdateBuilder.Set(x => x.NotificationChannel, Config.NotificationChannel));
-			Updates.Add(UpdateBuilder.Set(x => x.NotificationChannelFilter, Config.NotificationChannelFilter));
-			Updates.Add(UpdateBuilder.Set(x => x.TriageChannel, Config.TriageChannel));
-			Updates.Add(UpdateBuilder.Set(x => x.DefaultPreflight, DefaultPreflight));
-			Updates.Add(UpdateBuilder.Set(x => x.Tabs, Tabs ?? new List<StreamTab>()));
-			Updates.Add(UpdateBuilder.Set(x => x.AgentTypes, AgentTypes ?? new Dictionary<string, AgentType>()));
-			Updates.Add(UpdateBuilder.Set(x => x.WorkspaceTypes, WorkspaceTypes ?? new Dictionary<string, WorkspaceType>()));
-			Updates.Add(UpdateBuilder.Set(x => x.Templates, TemplateRefs ?? new Dictionary<TemplateRefId, TemplateRef>()));
-			Updates.Add(UpdateBuilder.Set(x => x.ReplicationMode, Config.ReplicationMode));
-			Updates.Add(UpdateBuilder.SetOrUnsetNullRef(x => x.ReplicationFilter, Config.ReplicationFilter));
-			Updates.Add(UpdateBuilder.SetOrUnsetNullRef(x => x.Acl, Acl));
-			Updates.Add(UpdateBuilder.Unset(x => x.Deleted));
-
-			return await TryUpdateStreamAsync(Stream, UpdateBuilder.Combine(Updates));
+			int order = config.Order ?? StreamDocument.DefaultOrder;
+
+			StreamDocument stream = (StreamDocument)streamInterface;
+
+			UpdateDefinitionBuilder<StreamDocument> updateBuilder = Builders<StreamDocument>.Update;
+
+			List<UpdateDefinition<StreamDocument>> updates = new List<UpdateDefinition<StreamDocument>>();
+			updates.Add(updateBuilder.Set(x => x.Name, config.Name));
+			updates.Add(updateBuilder.Set(x => x.ProjectId, projectId));
+			updates.Add(updateBuilder.Set(x => x.ClusterName, config.ClusterName));
+			updates.Add(updateBuilder.Set(x => x.ConfigPath, configPath));
+			updates.Add(updateBuilder.Set(x => x.ConfigRevision, configRevision));
+			updates.Add(updateBuilder.Set(x => x.Order, order));
+			updates.Add(updateBuilder.Set(x => x.NotificationChannel, config.NotificationChannel));
+			updates.Add(updateBuilder.Set(x => x.NotificationChannelFilter, config.NotificationChannelFilter));
+			updates.Add(updateBuilder.Set(x => x.TriageChannel, config.TriageChannel));
+			updates.Add(updateBuilder.Set(x => x.DefaultPreflight, defaultPreflight));
+			updates.Add(updateBuilder.Set(x => x.Tabs, tabs ?? new List<StreamTab>()));
+			updates.Add(updateBuilder.Set(x => x.AgentTypes, agentTypes ?? new Dictionary<string, AgentType>()));
+			updates.Add(updateBuilder.Set(x => x.WorkspaceTypes, workspaceTypes ?? new Dictionary<string, WorkspaceType>()));
+			updates.Add(updateBuilder.Set(x => x.Templates, templateRefs ?? new Dictionary<TemplateRefId, TemplateRef>()));
+			updates.Add(updateBuilder.Set(x => x.ReplicationMode, config.ReplicationMode));
+			updates.Add(updateBuilder.SetOrUnsetNullRef(x => x.ReplicationFilter, config.ReplicationFilter));
+			updates.Add(updateBuilder.SetOrUnsetNullRef(x => x.Acl, acl));
+			updates.Add(updateBuilder.Unset(x => x.Deleted));
+
+			return await TryUpdateStreamAsync(stream, updateBuilder.Combine(updates));
 		}
 
 		/// <inheritdoc/>
-		public async Task<IStream?> GetAsync(StreamId StreamId)
+		public async Task<IStream?> GetAsync(StreamId streamId)
 		{
-			return await Streams.Find<StreamDocument>(x => x.Id == StreamId).FirstOrDefaultAsync();
+			return await _streams.Find<StreamDocument>(x => x.Id == streamId).FirstOrDefaultAsync();
 		}
 
 		/// <inheritdoc/>
-		public async Task<IStreamPermissions?> GetPermissionsAsync(StreamId StreamId)
+		public async Task<IStreamPermissions?> GetPermissionsAsync(StreamId streamId)
 		{
-			return await Streams.Find<StreamDocument>(x => x.Id == StreamId).Project<StreamPermissions>(StreamPermissions.Projection).FirstOrDefaultAsync();
+			return await _streams.Find<StreamDocument>(x => x.Id == streamId).Project<StreamPermissions>(StreamPermissions.Projection).FirstOrDefaultAsync();
 		}
 
 		/// <inheritdoc/>
 		public async Task<List<IStream>> FindAllAsync()
 		{
-			List<StreamDocument> Results = await Streams.Find(Builders<StreamDocument>.Filter.Ne(x => x.Deleted, true)).ToListAsync();
-			return Results.ConvertAll<IStream>(x => x);
+			List<StreamDocument> results = await _streams.Find(Builders<StreamDocument>.Filter.Ne(x => x.Deleted, true)).ToListAsync();
+			return results.ConvertAll<IStream>(x => x);
 		}
 
 		/// <inheritdoc/>
-		public async Task<List<IStream>> FindForProjectsAsync(ProjectId[] ProjectIds)
+		public async Task<List<IStream>> FindForProjectsAsync(ProjectId[] projectIds)
 		{
-			FilterDefinition<StreamDocument> Filter = Builders<StreamDocument>.Filter.In(x => x.ProjectId, ProjectIds) & Builders<StreamDocument>.Filter.Ne(x => x.Deleted, true);
-			List<StreamDocument> Results = await Streams.Find(Filter).ToListAsync();
-			return Results.ConvertAll<IStream>(x => x);
+			FilterDefinition<StreamDocument> filter = Builders<StreamDocument>.Filter.In(x => x.ProjectId, projectIds) & Builders<StreamDocument>.Filter.Ne(x => x.Deleted, true);
+			List<StreamDocument> results = await _streams.Find(filter).ToListAsync();
+			return results.ConvertAll<IStream>(x => x);
 		}
 
 		/// <inheritdoc/>
-		public async Task<IStream?> TryUpdatePauseStateAsync(IStream StreamInterface, DateTime? NewPausedUntil, string? NewPauseComment)
+		public async Task<IStream?> TryUpdatePauseStateAsync(IStream streamInterface, DateTime? newPausedUntil, string? newPauseComment)
 		{
-			StreamDocument Stream = (StreamDocument)StreamInterface;
+			StreamDocument stream = (StreamDocument)streamInterface;
 
-			UpdateDefinitionBuilder<StreamDocument> UpdateBuilder = Builders<StreamDocument>.Update;
+			UpdateDefinitionBuilder<StreamDocument> updateBuilder = Builders<StreamDocument>.Update;
 
-			List<UpdateDefinition<StreamDocument>> Updates = new List<UpdateDefinition<StreamDocument>>();
-			Stream.PausedUntil = NewPausedUntil;
-			Stream.PauseComment = NewPauseComment;
-			Updates.Add(UpdateBuilder.Set(x => x.PausedUntil, NewPausedUntil));
-			Updates.Add(UpdateBuilder.Set(x => x.PauseComment, NewPauseComment));
+			List<UpdateDefinition<StreamDocument>> updates = new List<UpdateDefinition<StreamDocument>>();
+			stream.PausedUntil = newPausedUntil;
+			stream.PauseComment = newPauseComment;
+			updates.Add(updateBuilder.Set(x => x.PausedUntil, newPausedUntil));
+			updates.Add(updateBuilder.Set(x => x.PauseComment, newPauseComment));
 
-			return await TryUpdateStreamAsync(Stream, UpdateBuilder.Combine(Updates));
+			return await TryUpdateStreamAsync(stream, updateBuilder.Combine(updates));
 		}
 
 		/// <inheritdoc/>
-		public async Task<IStream?> TryUpdateScheduleTriggerAsync(IStream StreamInterface, TemplateRefId TemplateId, DateTime? LastTriggerTimeUtc, int? LastTriggerChange, List<JobId> NewActiveJobs)
+		public async Task<IStream?> TryUpdateScheduleTriggerAsync(IStream streamInterface, TemplateRefId templateId, DateTime? lastTriggerTimeUtc, int? lastTriggerChange, List<JobId> newActiveJobs)
 		{
-			StreamDocument Stream = (StreamDocument)StreamInterface;
-			Schedule Schedule = Stream.Templates[TemplateId].Schedule!;
+			StreamDocument stream = (StreamDocument)streamInterface;
+			TemplateRef template = stream.Templates[templateId];
+			Schedule schedule = template.Schedule!;
 
 			// Build the updates. MongoDB driver cannot parse TemplateRefId in expression tree; need to specify field name explicitly
-			List<UpdateDefinition<StreamDocument>> Updates = new List<UpdateDefinition<StreamDocument>>();
-			if (LastTriggerTimeUtc.HasValue && LastTriggerTimeUtc.Value != Schedule.LastTriggerTime)
+			List<UpdateDefinition<StreamDocument>> updates = new List<UpdateDefinition<StreamDocument>>();
+			if (lastTriggerTimeUtc.HasValue && lastTriggerTimeUtc.Value != schedule.LastTriggerTime)
 			{
-				FieldDefinition<StreamDocument, DateTimeOffset> LastTriggerTimeField = $"{nameof(Stream.Templates)}.{TemplateId}.{nameof(Schedule)}.{nameof(Schedule.LastTriggerTime)}";
-				Updates.Add(Builders<StreamDocument>.Update.Set(LastTriggerTimeField, LastTriggerTimeUtc.Value));
-				Schedule.LastTriggerTimeUtc = LastTriggerTimeUtc.Value;
+				FieldDefinition<StreamDocument, DateTimeOffset> lastTriggerTimeField = $"{nameof(stream.Templates)}.{templateId}.{nameof(template.Schedule)}.{nameof(schedule.LastTriggerTime)}";
+				updates.Add(Builders<StreamDocument>.Update.Set(lastTriggerTimeField, lastTriggerTimeUtc.Value));
+				schedule.LastTriggerTimeUtc = lastTriggerTimeUtc.Value;
 			}
-			if (LastTriggerChange.HasValue && LastTriggerChange.Value > Schedule.LastTriggerChange)
+			if (lastTriggerChange.HasValue && lastTriggerChange.Value > schedule.LastTriggerChange)
 			{
-				FieldDefinition<StreamDocument, int> LastTriggerChangeField = $"{nameof(Stream.Templates)}.{TemplateId}.{nameof(Schedule)}.{nameof(Schedule.LastTriggerChange)}";
-				Updates.Add(Builders<StreamDocument>.Update.Set(LastTriggerChangeField, LastTriggerChange.Value));
-				Schedule.LastTriggerChange = LastTriggerChange.Value;
+				FieldDefinition<StreamDocument, int> lastTriggerChangeField = $"{nameof(stream.Templates)}.{templateId}.{nameof(template.Schedule)}.{nameof(schedule.LastTriggerChange)}";
+				updates.Add(Builders<StreamDocument>.Update.Set(lastTriggerChangeField, lastTriggerChange.Value));
+				schedule.LastTriggerChange = lastTriggerChange.Value;
 			}
-			if (NewActiveJobs != null)
+			if (newActiveJobs != null)
 			{
-				FieldDefinition<StreamDocument, List<JobId>> Field = $"{nameof(Stream.Templates)}.{TemplateId}.{nameof(Schedule)}.{nameof(Schedule.ActiveJobs)}";
-				Updates.Add(Builders<StreamDocument>.Update.Set(Field, NewActiveJobs));
-				Schedule.ActiveJobs = NewActiveJobs;
+				FieldDefinition<StreamDocument, List<JobId>> field = $"{nameof(stream.Templates)}.{templateId}.{nameof(template.Schedule)}.{nameof(schedule.ActiveJobs)}";
+				updates.Add(Builders<StreamDocument>.Update.Set(field, newActiveJobs));
+				schedule.ActiveJobs = newActiveJobs;
 			}
 
-			return (Updates.Count == 0)? StreamInterface : await TryUpdateStreamAsync(Stream, Builders<StreamDocument>.Update.Combine(Updates));
+			return (updates.Count == 0)? streamInterface : await TryUpdateStreamAsync(stream, Builders<StreamDocument>.Update.Combine(updates));
 		}
 
 		/// <summary>
 		/// Update a stream
 		/// </summary>
-		/// <param name="Stream">The stream to update</param>
-		/// <param name="Update">The update definition</param>
+		/// <param name="stream">The stream to update</param>
+		/// <param name="update">The update definition</param>
 		/// <returns>The updated document, or null the update failed</returns>
-		private async Task<StreamDocument?> TryUpdateStreamAsync(StreamDocument Stream, UpdateDefinition<StreamDocument> Update)
+		private async Task<StreamDocument?> TryUpdateStreamAsync(StreamDocument stream, UpdateDefinition<StreamDocument> update)
 		{
-			FilterDefinition<StreamDocument> Filter = Builders<StreamDocument>.Filter.Expr(x => x.Id == Stream.Id && x.UpdateIndex == Stream.UpdateIndex);
-			Update = Update.Set(x => x.UpdateIndex, Stream.UpdateIndex + 1);
+			FilterDefinition<StreamDocument> filter = Builders<StreamDocument>.Filter.Expr(x => x.Id == stream.Id && x.UpdateIndex == stream.UpdateIndex);
+			update = update.Set(x => x.UpdateIndex, stream.UpdateIndex + 1);
 
-			FindOneAndUpdateOptions<StreamDocument> Options = new FindOneAndUpdateOptions<StreamDocument> { ReturnDocument = ReturnDocument.After };
-			return await Streams.FindOneAndUpdateAsync(Filter, Update, Options);
+			FindOneAndUpdateOptions<StreamDocument> options = new FindOneAndUpdateOptions<StreamDocument> { ReturnDocument = ReturnDocument.After };
+			return await _streams.FindOneAndUpdateAsync(filter, update, options);
 		}
 
 		/// <inheritdoc/>
-		public async Task DeleteAsync(StreamId StreamId)
+		public async Task DeleteAsync(StreamId streamId)
 		{
-			await Streams.UpdateOneAsync<StreamDocument>(x => x.Id == StreamId, Builders<StreamDocument>.Update.Set(x => x.Deleted, true).Inc(x => x.UpdateIndex, 1));
+			await _streams.UpdateOneAsync<StreamDocument>(x => x.Id == streamId, Builders<StreamDocument>.Update.Set(x => x.Deleted, true).Inc(x => x.UpdateIndex, 1));
 		}
 
-
 		/// <summary>
 		/// Checks the stream definition for consistency
 		/// </summary>
-		public static void Validate(StreamId StreamId, DefaultPreflight? DefaultPreflight, IReadOnlyDictionary<TemplateRefId, TemplateRef> Templates, IReadOnlyList<StreamTab> Tabs, IReadOnlyDictionary<string, AgentType> AgentTypes, IReadOnlyDictionary<string, WorkspaceType> WorkspaceTypes)
+		public static void Validate(StreamId streamId, DefaultPreflight? defaultPreflight, IReadOnlyDictionary<TemplateRefId, TemplateRef> templates, IReadOnlyList<StreamTab> tabs, IReadOnlyDictionary<string, AgentType> agentTypes, IReadOnlyDictionary<string, WorkspaceType> workspaceTypes)
 		{
 			// Check the default preflight template is valid
-			if (DefaultPreflight != null)
+			if (defaultPreflight != null)
 			{
-				if (DefaultPreflight.TemplateRefId != null && !Templates.ContainsKey(DefaultPreflight.TemplateRefId.Value))
+				if (defaultPreflight.TemplateRefId != null && !templates.ContainsKey(defaultPreflight.TemplateRefId.Value))
 				{
-					throw new InvalidStreamException($"Default preflight template was listed as '{DefaultPreflight.TemplateRefId.Value}', but no template was found by that name");
+					throw new InvalidStreamException($"Default preflight template was listed as '{defaultPreflight.TemplateRefId.Value}', but no template was found by that name");
 				}
 			}
 
 			// Check that all the templates are referenced by a tab
-			HashSet<TemplateRefId> RemainingTemplates = new HashSet<TemplateRefId>(Templates.Keys);
-			foreach (JobsTab JobsTab in Tabs.OfType<JobsTab>())
+			HashSet<TemplateRefId> remainingTemplates = new HashSet<TemplateRefId>(templates.Keys);
+			foreach (JobsTab jobsTab in tabs.OfType<JobsTab>())
 			{
-				if (JobsTab.Templates != null)
+				if (jobsTab.Templates != null)
 				{
-					RemainingTemplates.ExceptWith(JobsTab.Templates);
+					remainingTemplates.ExceptWith(jobsTab.Templates);
 				}
 			}
-			if (RemainingTemplates.Count > 0)
+			if (remainingTemplates.Count > 0)
 			{
-				throw new InvalidStreamException(String.Join("\n", RemainingTemplates.Select(x => $"Template '{x}' is not listed on any tab for {StreamId}")));
+				throw new InvalidStreamException(String.Join("\n", remainingTemplates.Select(x => $"Template '{x}' is not listed on any tab for {streamId}")));
 			}
 
 			// Check that all the agent types reference valid workspace names
-			foreach (KeyValuePair<string, AgentType> Pair in AgentTypes)
+			foreach (KeyValuePair<string, AgentType> pair in agentTypes)
 			{
-				string? WorkspaceTypeName = Pair.Value.Workspace;
-				if (WorkspaceTypeName != null && !WorkspaceTypes.ContainsKey(WorkspaceTypeName))
+				string? workspaceTypeName = pair.Value.Workspace;
+				if (workspaceTypeName != null && !workspaceTypes.ContainsKey(workspaceTypeName))
 				{
-					throw new InvalidStreamException($"Agent type '{Pair.Key}' references undefined workspace type '{Pair.Value.Workspace}' in {StreamId}");
+					throw new InvalidStreamException($"Agent type '{pair.Key}' references undefined workspace type '{pair.Value.Workspace}' in {streamId}");
 				}
 			}
 		}
diff --git a/Engine/Source/Programs/Horde/Horde.Build/Projects/Impl/StreamService.cs b/Engine/Source/Programs/Horde/Horde.Build/Projects/Impl/StreamService.cs
--- a/Engine/Source/Programs/Horde/Horde.Build/Projects/Impl/StreamService.cs
+++ b/Engine/Source/Programs/Horde/Horde.Build/Projects/Impl/StreamService.cs
@@ -1,22 +1,16 @@
 ﻿// Copyright Epic Games, Inc. All Rights Reserved.
 
-using Horde.Build.Acls;
-using Horde.Build.Api;
-using HordeCommon;
-using Horde.Build.Models;
-using Horde.Build.Utilities;
-using Microsoft.Extensions.Options;
-using MongoDB.Bson;
-using MongoDB.Driver;
 using System;
 using System.Collections.Generic;
 using System.Linq;
 using System.Security.Claims;
-using System.Text;
 using System.Threading.Tasks;
-using System.Linq.Expressions;
+using Horde.Build.Acls;
 using Horde.Build.Collections;
+using Horde.Build.Models;
+using Horde.Build.Utilities;
 using Microsoft.Extensions.Caching.Memory;
+using MongoDB.Driver;
 
 namespace Horde.Build.Services
 {
@@ -44,32 +38,32 @@ namespace Horde.Build.Services
 		/// <summary>
 		/// The project service instance
 		/// </summary>
-		ProjectService ProjectService;
+		readonly ProjectService _projectService;
 
 		/// <summary>
 		/// Collection of stream documents
 		/// </summary>
-		IStreamCollection Streams;
+		readonly IStreamCollection _streams;
 
 		/// <summary>
 		/// Cache of stream documents
 		/// </summary>
-		MemoryCache StreamCache = new MemoryCache(new MemoryCacheOptions());
+		readonly MemoryCache _streamCache = new MemoryCache(new MemoryCacheOptions());
 
 		/// <summary>
 		/// Accessor for the stream collection
 		/// </summary>
-		public IStreamCollection StreamCollection => Streams;
+		public IStreamCollection StreamCollection => _streams;
 
 		/// <summary>
 		/// Constructor
 		/// </summary>
-		/// <param name="ProjectService">The project service instance</param>
-		/// <param name="Streams">Collection of stream documents</param>
-		public StreamService(ProjectService ProjectService, IStreamCollection Streams)
+		/// <param name="projectService">The project service instance</param>
+		/// <param name="streams">Collection of stream documents</param>
+		public StreamService(ProjectService projectService, IStreamCollection streams)
 		{
-			this.ProjectService = ProjectService;
-			this.Streams = Streams;
+			_projectService = projectService;
+			_streams = streams;
 		}
 
 		/// <summary>
@@ -77,34 +71,34 @@ namespace Horde.Build.Services
 		/// </summary>
 		public void Dispose()
 		{
-			StreamCache.Dispose();
+			_streamCache.Dispose();
 		}
 
 		/// <summary>
 		/// Deletes an existing stream
 		/// </summary>
-		/// <param name="StreamId">Unique id of the stream</param>
+		/// <param name="streamId">Unique id of the stream</param>
 		/// <returns>Async task object</returns>
-		public async Task DeleteStreamAsync(StreamId StreamId)
+		public async Task DeleteStreamAsync(StreamId streamId)
 		{
-			await Streams.DeleteAsync(StreamId);
+			await _streams.DeleteAsync(streamId);
 		}
 
 		/// <summary>
 		/// Updates an existing stream
 		/// </summary>
-		/// <param name="Stream">The stream to update</param>
-		/// <param name="NewPausedUntil">The new datetime for pausing builds</param>
-		/// <param name="NewPauseComment">The reason for pausing</param>
+		/// <param name="stream">The stream to update</param>
+		/// <param name="newPausedUntil">The new datetime for pausing builds</param>
+		/// <param name="newPauseComment">The reason for pausing</param>
 		/// <returns>Async task object</returns>
-		public async Task<IStream?> UpdatePauseStateAsync(IStream? Stream, DateTime? NewPausedUntil = null, string? NewPauseComment = null)
+		public async Task<IStream?> UpdatePauseStateAsync(IStream? stream, DateTime? newPausedUntil = null, string? newPauseComment = null)
 		{
-			for (; Stream != null; Stream = await GetStreamAsync(Stream.Id))
+			for (; stream != null; stream = await GetStreamAsync(stream.Id))
 			{
-				IStream? NewStream = await Streams.TryUpdatePauseStateAsync(Stream, NewPausedUntil, NewPauseComment);
-				if (NewStream != null)
+				IStream? newStream = await _streams.TryUpdatePauseStateAsync(stream, newPausedUntil, newPauseComment);
+				if (newStream != null)
 				{
-					return NewStream;
+					return newStream;
 				}
 			}
 			return null;
@@ -113,46 +107,46 @@ namespace Horde.Build.Services
 		/// <summary>
 		/// Attempts to update the last trigger time for a schedule
 		/// </summary>
-		/// <param name="Stream">The stream to update</param>
-		/// <param name="TemplateRefId">The template ref id</param>
-		/// <param name="LastTriggerTimeUtc"></param>
-		/// <param name="LastTriggerChange"></param>
-		/// <param name="AddJobs">Jobs to add</param>
-		/// <param name="RemoveJobs">Jobs to remove</param>
+		/// <param name="stream">The stream to update</param>
+		/// <param name="templateRefId">The template ref id</param>
+		/// <param name="lastTriggerTimeUtc"></param>
+		/// <param name="lastTriggerChange"></param>
+		/// <param name="addJobs">Jobs to add</param>
+		/// <param name="removeJobs">Jobs to remove</param>
 		/// <returns>True if the stream was updated</returns>
-		public async Task<IStream?> UpdateScheduleTriggerAsync(IStream Stream, TemplateRefId TemplateRefId, DateTime? LastTriggerTimeUtc = null, int? LastTriggerChange = null, List<JobId>? AddJobs = null, List<JobId>? RemoveJobs = null)
+		public async Task<IStream?> UpdateScheduleTriggerAsync(IStream stream, TemplateRefId templateRefId, DateTime? lastTriggerTimeUtc = null, int? lastTriggerChange = null, List<JobId>? addJobs = null, List<JobId>? removeJobs = null)
 		{
-			IStream? NewStream = Stream;
-			while (NewStream != null)
+			IStream? newStream = stream;
+			while (newStream != null)
 			{
-				TemplateRef? TemplateRef;
-				if (!NewStream.Templates.TryGetValue(TemplateRefId, out TemplateRef))
+				TemplateRef? templateRef;
+				if (!newStream.Templates.TryGetValue(templateRefId, out templateRef))
 				{
 					break;
 				}
-				if (TemplateRef.Schedule == null)
+				if (templateRef.Schedule == null)
 				{
 					break;
 				}
 
-				IEnumerable<JobId> NewActiveJobs = TemplateRef.Schedule.ActiveJobs;
-				if (RemoveJobs != null)
+				IEnumerable<JobId> newActiveJobs = templateRef.Schedule.ActiveJobs;
+				if (removeJobs != null)
 				{
-					NewActiveJobs = NewActiveJobs.Except(RemoveJobs);
+					newActiveJobs = newActiveJobs.Except(removeJobs);
 				}
-				if (AddJobs != null)
+				if (addJobs != null)
 				{
-					NewActiveJobs = NewActiveJobs.Union(AddJobs);
+					newActiveJobs = newActiveJobs.Union(addJobs);
 				}
 
-				NewStream = await Streams.TryUpdateScheduleTriggerAsync(NewStream, TemplateRefId, LastTriggerTimeUtc, LastTriggerChange, NewActiveJobs.ToList());
+				newStream = await _streams.TryUpdateScheduleTriggerAsync(newStream, templateRefId, lastTriggerTimeUtc, lastTriggerChange, newActiveJobs.ToList());
 
-				if (NewStream != null)
+				if (newStream != null)
 				{
-					return NewStream;
+					return newStream;
 				}
 
-				NewStream = await Streams.GetAsync(Stream.Id);
+				newStream = await _streams.GetAsync(stream.Id);
 			}
 			return null;
 		}
@@ -163,164 +157,164 @@ namespace Horde.Build.Services
 		/// <returns>List of stream documents</returns>
 		public Task<List<IStream>> GetStreamsAsync()
 		{
-			return Streams.FindAllAsync();
+			return _streams.FindAllAsync();
 		}
 
 		/// <summary>
 		/// Gets all the available streams for a project
 		/// </summary>
-		/// <param name="ProjectId">Unique id of the project to query</param>
+		/// <param name="projectId">Unique id of the project to query</param>
 		/// <returns>List of stream documents</returns>
-		public Task<List<IStream>> GetStreamsAsync(ProjectId ProjectId)
+		public Task<List<IStream>> GetStreamsAsync(ProjectId projectId)
 		{
-			return Streams.FindForProjectsAsync(new[] { ProjectId });
+			return _streams.FindForProjectsAsync(new[] { projectId });
 		}
 
 		/// <summary>
 		/// Gets all the available streams for a project
 		/// </summary>
-		/// <param name="ProjectIds">Unique id of the project to query</param>
+		/// <param name="projectIds">Unique id of the project to query</param>
 		/// <returns>List of stream documents</returns>
-		public Task<List<IStream>> GetStreamsAsync(ProjectId[] ProjectIds)
+		public Task<List<IStream>> GetStreamsAsync(ProjectId[] projectIds)
 		{
-			if (ProjectIds.Length == 0)
+			if (projectIds.Length == 0)
 			{
-				return Streams.FindAllAsync();
+				return _streams.FindAllAsync();
 			}
 			else
 			{
-				return Streams.FindForProjectsAsync(ProjectIds);
+				return _streams.FindForProjectsAsync(projectIds);
 			}
 		}
 
 		/// <summary>
 		/// Gets a stream by ID
 		/// </summary>
-		/// <param name="StreamId">Unique id of the stream</param>
+		/// <param name="streamId">Unique id of the stream</param>
 		/// <returns>The stream document</returns>
-		public Task<IStream?> GetStreamAsync(StreamId StreamId)
+		public Task<IStream?> GetStreamAsync(StreamId streamId)
 		{
-			return Streams.GetAsync(StreamId);
+			return _streams.GetAsync(streamId);
 		}
 
 		/// <summary>
 		/// Adds a cached stream interface
 		/// </summary>
-		/// <param name="Stream">The stream to add</param>
+		/// <param name="stream">The stream to add</param>
 		/// <returns>The new stream</returns>
-		private void AddCachedStream(IStream Stream)
+		private void AddCachedStream(IStream stream)
 		{
-			MemoryCacheEntryOptions Options = new MemoryCacheEntryOptions().SetAbsoluteExpiration(TimeSpan.FromMinutes(1.0));
-			StreamCache.Set(Stream.Id, Stream, Options);
+			MemoryCacheEntryOptions options = new MemoryCacheEntryOptions().SetAbsoluteExpiration(TimeSpan.FromMinutes(1.0));
+			_streamCache.Set(stream.Id, stream, options);
 		}
 
 		/// <summary>
 		/// Gets a cached stream interface
 		/// </summary>
-		/// <param name="StreamId">Unique id for the stream</param>
+		/// <param name="streamId">Unique id for the stream</param>
 		/// <returns>The new stream</returns>
-		public async Task<IStream?> GetCachedStream(StreamId StreamId)
+		public async Task<IStream?> GetCachedStream(StreamId streamId)
 		{
-			object? Stream;
-			if (!StreamCache.TryGetValue(StreamId, out Stream))
+			object? stream;
+			if (!_streamCache.TryGetValue(streamId, out stream))
 			{
-				Stream = await GetStreamAsync(StreamId);
-				if (Stream != null)
+				stream = await GetStreamAsync(streamId);
+				if (stream != null)
 				{
-					AddCachedStream((IStream)Stream);
+					AddCachedStream((IStream)stream);
 				}
 			}
-			return (IStream?)Stream;
+			return (IStream?)stream;
 		}
 
 		/// <summary>
 		/// Gets a stream's permissions by ID
 		/// </summary>
-		/// <param name="StreamId">Unique id of the stream</param>
+		/// <param name="streamId">Unique id of the stream</param>
 		/// <returns>The stream document</returns>
-		public Task<IStreamPermissions?> GetStreamPermissionsAsync(StreamId StreamId)
+		public Task<IStreamPermissions?> GetStreamPermissionsAsync(StreamId streamId)
 		{
-			return Streams.GetPermissionsAsync(StreamId);
+			return _streams.GetPermissionsAsync(streamId);
 		}
 
 		/// <summary>
 		/// Determines if the user is authorized to perform an action on a particular stream
 		/// </summary>
-		/// <param name="Acl">ACL for the stream to check</param>
-		/// <param name="ProjectId">The parent project id</param>
-		/// <param name="Action">The action being performed</param>
-		/// <param name="User">The principal to authorize</param>
-		/// <param name="Cache">Cache of project permissions</param>
+		/// <param name="acl">ACL for the stream to check</param>
+		/// <param name="projectId">The parent project id</param>
+		/// <param name="action">The action being performed</param>
+		/// <param name="user">The principal to authorize</param>
+		/// <param name="cache">Cache of project permissions</param>
 		/// <returns>True if the action is authorized</returns>
-		private Task<bool> AuthorizeAsync(Acl? Acl, ProjectId ProjectId, AclAction Action, ClaimsPrincipal User, ProjectPermissionsCache? Cache)
+		private Task<bool> AuthorizeAsync(Acl? acl, ProjectId projectId, AclAction action, ClaimsPrincipal user, ProjectPermissionsCache? cache)
 		{
-			bool? Result = Acl?.Authorize(Action, User);
-			if (Result == null)
+			bool? result = acl?.Authorize(action, user);
+			if (result == null)
 			{
-				return ProjectService.AuthorizeAsync(ProjectId, Action, User, Cache);
+				return _projectService.AuthorizeAsync(projectId, action, user, cache);
 			}
 			else
 			{
-				return Task.FromResult(Result.Value);
+				return Task.FromResult(result.Value);
 			}
 		}
 
 		/// <summary>
 		/// Determines if the user is authorized to perform an action on a particular stream
 		/// </summary>
-		/// <param name="Stream">The stream to check</param>
-		/// <param name="Action">The action being performed</param>
-		/// <param name="User">The principal to authorize</param>
-		/// <param name="Cache">Cache of project permissions</param>
+		/// <param name="stream">The stream to check</param>
+		/// <param name="action">The action being performed</param>
+		/// <param name="user">The principal to authorize</param>
+		/// <param name="cache">Cache of project permissions</param>
 		/// <returns>True if the action is authorized</returns>
-		public Task<bool> AuthorizeAsync(IStream Stream, AclAction Action, ClaimsPrincipal User, ProjectPermissionsCache? Cache)
+		public Task<bool> AuthorizeAsync(IStream stream, AclAction action, ClaimsPrincipal user, ProjectPermissionsCache? cache)
 		{
-			return AuthorizeAsync(Stream.Acl, Stream.ProjectId, Action, User, Cache);
+			return AuthorizeAsync(stream.Acl, stream.ProjectId, action, user, cache);
 		}
 
 		/// <summary>
 		/// Determines if the user is authorized to perform an action on a particular stream
 		/// </summary>
-		/// <param name="Stream">The stream to check</param>
-		/// <param name="Template">Template within the stream to check</param>
-		/// <param name="Action">The action being performed</param>
-		/// <param name="User">The principal to authorize</param>
-		/// <param name="Cache">Cache of project permissions</param>
+		/// <param name="stream">The stream to check</param>
+		/// <param name="template">Template within the stream to check</param>
+		/// <param name="action">The action being performed</param>
+		/// <param name="user">The principal to authorize</param>
+		/// <param name="cache">Cache of project permissions</param>
 		/// <returns>True if the action is authorized</returns>
-		public Task<bool> AuthorizeAsync(IStream Stream, TemplateRef Template, AclAction Action, ClaimsPrincipal User, ProjectPermissionsCache? Cache)
+		public Task<bool> AuthorizeAsync(IStream stream, TemplateRef template, AclAction action, ClaimsPrincipal user, ProjectPermissionsCache? cache)
 		{
-			bool? Result = Template.Acl?.Authorize(Action, User);
-			if (Result == null)
+			bool? result = template.Acl?.Authorize(action, user);
+			if (result == null)
 			{
-				return AuthorizeAsync(Stream, Action, User, Cache);
+				return AuthorizeAsync(stream, action, user, cache);
 			}
 			else
 			{
-				return Task.FromResult(Result.Value);
+				return Task.FromResult(result.Value);
 			}
 		}
 
 		/// <summary>
 		/// Determines if the user is authorized to perform an action on a particular project
 		/// </summary>
-		/// <param name="StreamId">The stream id to check</param>
-		/// <param name="Action">The action being performed</param>
-		/// <param name="User">The principal to authorize</param>
-		/// <param name="Cache">Cache of stream permissions</param>
+		/// <param name="streamId">The stream id to check</param>
+		/// <param name="action">The action being performed</param>
+		/// <param name="user">The principal to authorize</param>
+		/// <param name="cache">Cache of stream permissions</param>
 		/// <returns>True if the action is authorized</returns>
-		public async Task<bool> AuthorizeAsync(StreamId StreamId, AclAction Action, ClaimsPrincipal User, StreamPermissionsCache? Cache)
+		public async Task<bool> AuthorizeAsync(StreamId streamId, AclAction action, ClaimsPrincipal user, StreamPermissionsCache? cache)
 		{
-			IStreamPermissions? Permissions;
-			if (Cache == null)
+			IStreamPermissions? permissions;
+			if (cache == null)
 			{
-				Permissions = await GetStreamPermissionsAsync(StreamId);
+				permissions = await GetStreamPermissionsAsync(streamId);
 			}
-			else if (!Cache.Streams.TryGetValue(StreamId, out Permissions))
+			else if (!cache.Streams.TryGetValue(streamId, out permissions))
 			{
-				Permissions = await GetStreamPermissionsAsync(StreamId);
-				Cache.Streams.Add(StreamId, Permissions);
+				permissions = await GetStreamPermissionsAsync(streamId);
+				cache.Streams.Add(streamId, permissions);
 			}
-			return Permissions != null && await AuthorizeAsync(Permissions.Acl, Permissions.ProjectId, Action, User, Cache);
+			return permissions != null && await AuthorizeAsync(permissions.Acl, permissions.ProjectId, action, user, cache);
 		}
 	}
 }
diff --git a/Engine/Source/Programs/Horde/Horde.Build/Projects/ScheduleService.cs b/Engine/Source/Programs/Horde/Horde.Build/Projects/ScheduleService.cs
--- a/Engine/Source/Programs/Horde/Horde.Build/Projects/ScheduleService.cs
+++ b/Engine/Source/Programs/Horde/Horde.Build/Projects/ScheduleService.cs
@@ -1,34 +1,29 @@
 // Copyright Epic Games, Inc. All Rights Reserved.
 
-using Amazon.S3.Model.Internal.MarshallTransformations;
+using System;
+using System.Collections.Generic;
+using System.Diagnostics;
+using System.Linq;
+using System.Text.RegularExpressions;
+using System.Threading;
+using System.Threading.Tasks;
 using EpicGames.Core;
-using HordeCommon;
+using EpicGames.Redis;
+using EpicGames.Redis.Utility;
+using EpicGames.Serialization;
 using Horde.Build.Api;
 using Horde.Build.Collections;
-using Horde.Build.Controllers;
 using Horde.Build.Models;
 using Horde.Build.Utilities;
-using Microsoft.AspNetCore.Server.Kestrel.Core.Features;
+using HordeCommon;
 using Microsoft.Extensions.Hosting;
 using Microsoft.Extensions.Logging;
 using Microsoft.Extensions.Options;
-using MongoDB.Bson;
 using MongoDB.Driver;
 using OpenTracing;
 using OpenTracing.Util;
-using System;
-using System.Collections.Generic;
-using System.Diagnostics;
-using System.Linq;
-using System.Security.Claims;
-using System.Threading;
-using System.Threading.Tasks;
-using TimeZoneConverter;
 using StackExchange.Redis;
-using EpicGames.Redis;
-using EpicGames.Serialization;
-using System.Text.RegularExpressions;
-using EpicGames.Redis.Utility;
+using TimeZoneConverter;
 
 namespace Horde.Build.Services
 {
@@ -54,197 +49,195 @@ namespace Horde.Build.Services
 			{
 			}
 
-			public QueueItem(StreamId StreamId, TemplateRefId TemplateId)
+			public QueueItem(StreamId streamId, TemplateRefId templateId)
 			{
-				this.StreamId = StreamId;
-				this.TemplateId = TemplateId;
+				StreamId = streamId;
+				TemplateId = templateId;
 			}
 
-			public static DateTime GetTimeFromScore(double Score) => DateTime.UnixEpoch + TimeSpan.FromSeconds(Score);
-			public static double GetScoreFromTime(DateTime Time) => (Time.ToUniversalTime() - DateTime.UnixEpoch).TotalSeconds;
+			public static DateTime GetTimeFromScore(double score) => DateTime.UnixEpoch + TimeSpan.FromSeconds(score);
+			public static double GetScoreFromTime(DateTime time) => (time.ToUniversalTime() - DateTime.UnixEpoch).TotalSeconds;
 		}
 
 		class PerforceHistory
 		{
-			IPerforceService Perforce;
-			string ClusterName;
-			string StreamName;
-			int MaxResults;
-			List<ChangeDetails> Changes = new List<ChangeDetails>();
-			int NextIndex;
+			readonly IPerforceService _perforce;
+			readonly string _clusterName;
+			readonly string _streamName;
+			int _maxResults;
+			readonly List<ChangeDetails> _changes = new List<ChangeDetails>();
+			int _nextIndex;
 
-			public PerforceHistory(IPerforceService Perforce, string ClusterName, string StreamName)
+			public PerforceHistory(IPerforceService perforce, string clusterName, string streamName)
 			{
-				this.Perforce = Perforce;
-				this.ClusterName = ClusterName;
-				this.StreamName = StreamName;
-				this.MaxResults = 10;
+				_perforce = perforce;
+				_clusterName = clusterName;
+				_streamName = streamName;
+				_maxResults = 10;
 			}
 
-			async ValueTask<ChangeDetails?> GetChangeAtIndexAsync(int Index)
+			async ValueTask<ChangeDetails?> GetChangeAtIndexAsync(int index)
 			{
-				while (Index >= Changes.Count)
+				while (index >= _changes.Count)
 				{
-					List<ChangeSummary> NewChanges = await Perforce.GetChangesAsync(ClusterName, StreamName, null, null, MaxResults, null);
+					List<ChangeSummary> newChanges = await _perforce.GetChangesAsync(_clusterName, _streamName, null, null, _maxResults, null);
 
-					int NumResults = NewChanges.Count;
-					if (Changes.Count > 0)
+					int numResults = newChanges.Count;
+					if (_changes.Count > 0)
 					{
-						NewChanges.RemoveAll(x => x.Number >= Changes[Changes.Count - 1].Number);
+						newChanges.RemoveAll(x => x.Number >= _changes[^1].Number);
 					}
-					if (NewChanges.Count == 0 && NumResults < MaxResults)
+					if (newChanges.Count == 0 && numResults < _maxResults)
 					{
 						return null;
 					}
-					if(NewChanges.Count > 0)
+					if(newChanges.Count > 0)
 					{
-						Changes.AddRange((await Perforce.GetChangeDetailsAsync(ClusterName, StreamName, NewChanges.ConvertAll(x => x.Number), null)).OrderByDescending(x => x.Number));
+						_changes.AddRange((await _perforce.GetChangeDetailsAsync(_clusterName, _streamName, newChanges.ConvertAll(x => x.Number), null)).OrderByDescending(x => x.Number));
 					}
-					MaxResults += 10;
+					_maxResults += 10;
 				}
-				return Changes[Index];
+				return _changes[index];
 			}
 
 			public async ValueTask<ChangeDetails?> GetNextChangeAsync()
 			{
-				ChangeDetails? Details = await GetChangeAtIndexAsync(NextIndex);
-				if (Details != null)
+				ChangeDetails? details = await GetChangeAtIndexAsync(_nextIndex);
+				if (details != null)
 				{
-					NextIndex++;
+					_nextIndex++;
 				}
-				return Details;
+				return details;
 			}
 
-			public async ValueTask<int> GetCodeChange(int Change)
+			public async ValueTask<int> GetCodeChange(int change)
 			{
-				int Index = Changes.BinarySearch(x => -x.Number, -Change);
-				if (Index < 0)
+				int index = _changes.BinarySearch(x => -x.Number, -change);
+				if (index < 0)
 				{
-					Index = ~Index;
+					index = ~index;
 				}
 
 				for (; ; )
 				{
-					ChangeDetails? Details = await GetChangeAtIndexAsync(Index);
-					if (Details == null)
+					ChangeDetails? details = await GetChangeAtIndexAsync(index);
+					if (details == null)
 					{
 						return 0;
 					}
-					if ((Details.GetContentFlags() & ChangeContentFlags.ContainsCode) != 0)
+					if ((details.GetContentFlags() & ChangeContentFlags.ContainsCode) != 0)
 					{
-						return Details.Number;
+						return details.Number;
 					}
-					Index++;
+					index++;
 				}
 			}
 		}
 
-		IDowntimeService DowntimeService;
-		IGraphCollection Graphs;
-		IPerforceService Perforce;
-		IJobCollection JobCollection;
-		JobService JobService;
-		StreamService StreamService;
-		ITemplateCollection TemplateCollection;
-		TimeZoneInfo TimeZone;
-		IClock Clock;
-		ILogger Logger;
-		RedisService Redis;
-		RedisKey BaseLockKey;
-		RedisKey TickLockKey; // Lock to tick the queue
-		RedisSortedSet<QueueItem> Queue; // Items to tick, ordered by time
-		ITicker Ticker;
+		readonly IGraphCollection _graphs;
+		readonly IPerforceService _perforce;
+		readonly IJobCollection _jobCollection;
+		readonly JobService _jobService;
+		readonly StreamService _streamService;
+		readonly ITemplateCollection _templateCollection;
+		readonly TimeZoneInfo _timeZone;
+		readonly IClock _clock;
+		readonly ILogger _logger;
+		readonly RedisService _redis;
+		readonly RedisKey _baseLockKey;
+		readonly RedisKey _tickLockKey; // Lock to tick the queue
+		readonly RedisSortedSet<QueueItem> _queue; // Items to tick, ordered by time
+		readonly ITicker _ticker;
 
 		/// <summary>
 		/// Constructor
 		/// </summary>
-		public ScheduleService(RedisService Redis, IGraphCollection Graphs, IDowntimeService DowntimeService, IPerforceService Perforce, IJobCollection JobCollection, JobService JobService, StreamService StreamService, ITemplateCollection TemplateCollection, DatabaseService DatabaseService, IClock Clock, IOptionsMonitor<ServerSettings> Settings, ILogger<ScheduleService> Logger)
+		public ScheduleService(RedisService redis, IGraphCollection graphs, IPerforceService perforce, IJobCollection jobCollection, JobService jobService, StreamService streamService, ITemplateCollection templateCollection, DatabaseService databaseService, IClock clock, IOptionsMonitor<ServerSettings> settings, ILogger<ScheduleService> logger)
 		{
-			this.Graphs = Graphs;
-			this.DowntimeService = DowntimeService;
-			this.Perforce = Perforce;
-			this.JobCollection = JobCollection;
-			this.JobService = JobService;
-			this.StreamService = StreamService;
-			this.TemplateCollection = TemplateCollection;
-			this.Clock = Clock;
+			_graphs = graphs;
+			_perforce = perforce;
+			_jobCollection = jobCollection;
+			_jobService = jobService;
+			_streamService = streamService;
+			_templateCollection = templateCollection;
+			_clock = clock;
 
-			string? ScheduleTimeZone = Settings.CurrentValue.ScheduleTimeZone;
-			this.TimeZone = (ScheduleTimeZone == null) ? TimeZoneInfo.Local : TZConvert.GetTimeZoneInfo(ScheduleTimeZone);
+			string? scheduleTimeZone = settings.CurrentValue.ScheduleTimeZone;
+			_timeZone = (scheduleTimeZone == null) ? TimeZoneInfo.Local : TZConvert.GetTimeZoneInfo(scheduleTimeZone);
 
-			this.Logger = Logger;
+			_logger = logger;
 
-			this.Redis = Redis;
-			this.BaseLockKey = "scheduler/locks";
-			this.TickLockKey = BaseLockKey.Append("/tick");
-			this.Queue = new RedisSortedSet<QueueItem>(Redis.Database, "scheduler/queue");
-			if (DatabaseService.ReadOnlyMode)
+			_redis = redis;
+			_baseLockKey = "scheduler/locks";
+			_tickLockKey = _baseLockKey.Append("/tick");
+			_queue = new RedisSortedSet<QueueItem>(redis.Database, "scheduler/queue");
+			if (databaseService.ReadOnlyMode)
 			{
-				Ticker = new NullTicker();
+				_ticker = new NullTicker();
 			}
 			else
 			{
-				Ticker = Clock.AddTicker(TimeSpan.FromMinutes(1.0), TickAsync, Logger);
+				_ticker = clock.AddTicker<ScheduleService>(TimeSpan.FromMinutes(1.0), TickAsync, logger);
 			}
 		}
 
 		/// <inheritdoc/>
-		public async Task StartAsync(CancellationToken CancellationToken)
+		public async Task StartAsync(CancellationToken cancellationToken)
 		{
-			await Ticker.StartAsync();
+			await _ticker.StartAsync();
 		}
 
 		/// <inheritdoc/>
-		public async Task StopAsync(CancellationToken CancellationToken)
+		public async Task StopAsync(CancellationToken cancellationToken)
 		{
-			await Ticker.StopAsync();
+			await _ticker.StopAsync();
 		}
 
 		/// <inheritdoc/>
 		public void Dispose()
 		{
-			Ticker.Dispose();
+			_ticker.Dispose();
 		}
 
-		async ValueTask TickAsync(CancellationToken CancellationToken)
+		async ValueTask TickAsync(CancellationToken cancellationToken)
 		{
-			DateTime UtcNow = Clock.UtcNow;
+			DateTime utcNow = _clock.UtcNow;
 
 			// Update the current queue
-			await using (RedisLock Lock = new RedisLock(Redis.Database, TickLockKey))
+			await using (RedisLock sharedLock = new RedisLock(_redis.Database, _tickLockKey))
 			{
-				if (await Lock.AcquireAsync(TimeSpan.FromMinutes(1.0), false))
+				if (await sharedLock.AcquireAsync(TimeSpan.FromMinutes(1.0), false))
 				{
-					await UpdateQueueAsync(UtcNow);
+					await UpdateQueueAsync(utcNow);
 				}
 			}
 
 			// Keep updating schedules
-			while (!CancellationToken.IsCancellationRequested)
+			while (!cancellationToken.IsCancellationRequested)
 			{
 				// Get the item with the lowest score (ie. the one that hasn't been updated in the longest time)
-				QueueItem? Item = await PopQueueItemAsync();
-				if (Item == null)
+				QueueItem? item = await PopQueueItemAsync();
+				if (item == null)
 				{
 					break;
 				}
 
 				// Acquire the lock for this schedule and update it
-				await using (RedisLock Lock = new RedisLock<QueueItem>(Redis.Database, BaseLockKey, Item))
+				await using (RedisLock sharedLock = new RedisLock<QueueItem>(_redis.Database, _baseLockKey, item))
 				{
-					if (await Lock.AcquireAsync(TimeSpan.FromMinutes(1.0)))
+					if (await sharedLock.AcquireAsync(TimeSpan.FromMinutes(1.0)))
 					{
 						try
 						{
-							await TriggerAsync(Item.StreamId, Item.TemplateId, UtcNow, CancellationToken);
+							await TriggerAsync(item.StreamId, item.TemplateId, utcNow, cancellationToken);
 						}
 						catch (OperationCanceledException)
 						{
 							throw;
 						}
-						catch (Exception Ex)
+						catch (Exception ex)
 						{
-							Logger.LogError(Ex, "Error while updating schedule for {StreamId}/{TemplateId}", Item.StreamId, Item.TemplateId);
+							_logger.LogError(ex, "Error while updating schedule for {StreamId}/{TemplateId}", item.StreamId, item.TemplateId);
 						}
 					}
 				}
@@ -255,365 +248,349 @@ namespace Horde.Build.Services
 		{
 			for (; ; )
 			{
-				QueueItem[] Items = await Queue.RangeByRankAsync(0, 0);
-				if (Items.Length == 0)
+				QueueItem[] items = await _queue.RangeByRankAsync(0, 0);
+				if (items.Length == 0)
 				{
 					return null;
 				}
-				if (await Queue.RemoveAsync(Items[0]))
+				if (await _queue.RemoveAsync(items[0]))
 				{
-					return Items[0];
+					return items[0];
 				}
 			}
 		}
 
 		internal async Task ResetAsync()
 		{
-			await Redis.Database.KeyDeleteAsync(Queue.Key);
-			await Redis.Database.KeyDeleteAsync(TickLockKey);
+			await _redis.Database.KeyDeleteAsync(_queue.Key);
+			await _redis.Database.KeyDeleteAsync(_tickLockKey);
 		}
 
 		internal async Task TickForTestingAsync()
 		{
-			await UpdateQueueAsync(Clock.UtcNow);
+			await UpdateQueueAsync(_clock.UtcNow);
 			await TickAsync(CancellationToken.None);
 		}
 
 		/// <summary>
 		/// Get the current set of streams and ensure there's an entry for each item
 		/// </summary>
-		public async Task UpdateQueueAsync(DateTime UtcNow)
+		public async Task UpdateQueueAsync(DateTime utcNow)
 		{
-			List<SortedSetEntry<QueueItem>> QueueItems = new List<SortedSetEntry<QueueItem>>();
+			List<SortedSetEntry<QueueItem>> queueItems = new List<SortedSetEntry<QueueItem>>();
 
-			List<IStream> Streams = await StreamService.GetStreamsAsync();
-			foreach (IStream Stream in Streams)
+			List<IStream> streams = await _streamService.GetStreamsAsync();
+			foreach (IStream stream in streams)
 			{
-				foreach ((TemplateRefId TemplateId, TemplateRef TemplateRef) in Stream.Templates)
+				foreach ((TemplateRefId templateId, TemplateRef templateRef) in stream.Templates)
 				{
-					if (TemplateRef.Schedule != null)
+					if (templateRef.Schedule != null)
 					{
-						DateTime? NextTriggerTimeUtc = TemplateRef.Schedule.GetNextTriggerTimeUtc(TimeZone);
-						if (NextTriggerTimeUtc != null)
+						DateTime? nextTriggerTimeUtc = templateRef.Schedule.GetNextTriggerTimeUtc(_timeZone);
+						if (nextTriggerTimeUtc != null)
 						{
-							if (UtcNow > NextTriggerTimeUtc.Value)
+							if (utcNow > nextTriggerTimeUtc.Value)
 							{
-								double Score = QueueItem.GetScoreFromTime(NextTriggerTimeUtc.Value);
-								QueueItems.Add(new SortedSetEntry<QueueItem>(new QueueItem(Stream.Id, TemplateId), Score));
+								double score = QueueItem.GetScoreFromTime(nextTriggerTimeUtc.Value);
+								queueItems.Add(new SortedSetEntry<QueueItem>(new QueueItem(stream.Id, templateId), score));
 
-								await StreamService.UpdateScheduleTriggerAsync(Stream, TemplateId, UtcNow);
+								await _streamService.UpdateScheduleTriggerAsync(stream, templateId, utcNow);
 							}
 						}
 					}
 				}
 			}
 
-			await Queue.AddAsync(QueueItems.ToArray());
+			await _queue.AddAsync(queueItems.ToArray());
 		}
 
 		/// <summary>
 		/// Trigger a schedule to run
 		/// </summary>
-		/// <param name="StreamId">Stream for the schedule</param>
-		/// <param name="TemplateId"></param>
-		/// <param name="UtcNow"></param>
-		/// <param name="CancellationToken"></param>
+		/// <param name="streamId">Stream for the schedule</param>
+		/// <param name="templateId"></param>
+		/// <param name="utcNow"></param>
+		/// <param name="cancellationToken"></param>
 		/// <returns>Async task</returns>
-		internal async Task<bool> TriggerAsync(StreamId StreamId, TemplateRefId TemplateId, DateTime UtcNow, CancellationToken CancellationToken)
+		internal async Task<bool> TriggerAsync(StreamId streamId, TemplateRefId templateId, DateTime utcNow, CancellationToken cancellationToken)
 		{
-			IStream? Stream = await StreamService.GetStreamAsync(StreamId);
-			if (Stream == null || !Stream.Templates.TryGetValue(TemplateId, out TemplateRef? TemplateRef))
+			IStream? stream = await _streamService.GetStreamAsync(streamId);
+			if (stream == null || !stream.Templates.TryGetValue(templateId, out TemplateRef? templateRef))
 			{
 				return false;
 			}
 
-			Schedule? Schedule = TemplateRef.Schedule;
-			if (Schedule == null)
+			Schedule? schedule = templateRef.Schedule;
+			if (schedule == null)
 			{
 				return false;
 			}
 
-			using IScope Scope = GlobalTracer.Instance.BuildSpan("ScheduleService.TriggerAsync").StartActive();
-			Scope.Span.SetTag("StreamId", Stream.Id);
-			Scope.Span.SetTag("TemplateId", TemplateId);
+			using IScope scope = GlobalTracer.Instance.BuildSpan("ScheduleService.TriggerAsync").StartActive();
+			scope.Span.SetTag("StreamId", stream.Id);
+			scope.Span.SetTag("TemplateId", templateId);
 
-			Stopwatch Stopwatch = Stopwatch.StartNew();
-			Logger.LogInformation("Updating schedule for {StreamId} template {TemplateId}", Stream.Id, TemplateId);
+			Stopwatch stopwatch = Stopwatch.StartNew();
+			_logger.LogInformation("Updating schedule for {StreamId} template {TemplateId}", stream.Id, templateId);
 
 			// Get a list of jobs that we need to remove
-			List<JobId> RemoveJobIds = new List<JobId>();
-			foreach (JobId ActiveJobId in Schedule.ActiveJobs)
+			List<JobId> removeJobIds = new List<JobId>();
+			foreach (JobId activeJobId in schedule.ActiveJobs)
 			{
-				IJob? Job = await JobService.GetJobAsync(ActiveJobId);
-				if (Job == null || Job.Batches.All(x => x.State == JobStepBatchState.Complete))
+				IJob? job = await _jobService.GetJobAsync(activeJobId);
+				if (job == null || job.Batches.All(x => x.State == JobStepBatchState.Complete))
 				{
-					Logger.LogInformation("Removing active job {JobId}", ActiveJobId);
-					RemoveJobIds.Add(ActiveJobId);
+					_logger.LogInformation("Removing active job {JobId}", activeJobId);
+					removeJobIds.Add(activeJobId);
 				}
 			}
-			await StreamService.UpdateScheduleTriggerAsync(Stream, TemplateId, RemoveJobs: RemoveJobIds);
+			await _streamService.UpdateScheduleTriggerAsync(stream, templateId, removeJobs: removeJobIds);
 
 			// If the stream is paused, bail out
-			if (Stream.IsPaused(UtcNow))
+			if (stream.IsPaused(utcNow))
 			{
-				Logger.LogDebug("Skipping schedule update for stream {StreamId}. It has been paused until {PausedUntil} with comment '{PauseComment}'.", Stream.Id, Stream.PausedUntil, Stream.PauseComment);
+				_logger.LogDebug("Skipping schedule update for stream {StreamId}. It has been paused until {PausedUntil} with comment '{PauseComment}'.", stream.Id, stream.PausedUntil, stream.PauseComment);
 				return false;
 			}
 
 			// Trigger this schedule
 			try
 			{
-				await TriggerAsync(Stream, TemplateId, TemplateRef, Schedule, Schedule.ActiveJobs.Count - RemoveJobIds.Count, UtcNow, CancellationToken);
+				await TriggerAsync(stream, templateId, templateRef, schedule, schedule.ActiveJobs.Count - removeJobIds.Count, utcNow, cancellationToken);
 			}
-			catch (Exception Ex)
+			catch (Exception ex)
 			{
-				Logger.LogError(Ex, "Failed to start schedule {StreamId}/{TemplateId}", Stream.Id, TemplateId);
+				_logger.LogError(ex, "Failed to start schedule {StreamId}/{TemplateId}", stream.Id, templateId);
 			}
 
 			// Print some timing info
-			Stopwatch.Stop();
-			Logger.LogInformation("Updated schedule for {StreamId} template {TemplateId} in {TimeSeconds}ms", Stream.Id, TemplateId, (long)Stopwatch.Elapsed.TotalMilliseconds);
+			stopwatch.Stop();
+			_logger.LogInformation("Updated schedule for {StreamId} template {TemplateId} in {TimeSeconds}ms", stream.Id, templateId, (long)stopwatch.Elapsed.TotalMilliseconds);
 			return true;
 		}
 
 		/// <summary>
 		/// Trigger a schedule to run
 		/// </summary>
-		/// <param name="Stream">Stream for the schedule</param>
-		/// <param name="TemplateId"></param>
-		/// <param name="TemplateRef"></param>
-		/// <param name="Schedule"></param>
-		/// <param name="NumActiveJobs"></param>
-		/// <param name="UtcNow">The current time</param>
-		/// <param name="CancellationToken"></param>
+		/// <param name="stream">Stream for the schedule</param>
+		/// <param name="templateId"></param>
+		/// <param name="templateRef"></param>
+		/// <param name="schedule"></param>
+		/// <param name="numActiveJobs"></param>
+		/// <param name="utcNow">The current time</param>
+		/// <param name="cancellationToken"></param>
 		/// <returns>Async task</returns>
-		private async Task TriggerAsync(IStream Stream, TemplateRefId TemplateId, TemplateRef TemplateRef, Schedule Schedule, int NumActiveJobs, DateTime UtcNow, CancellationToken CancellationToken)
+		private async Task TriggerAsync(IStream stream, TemplateRefId templateId, TemplateRef templateRef, Schedule schedule, int numActiveJobs, DateTime utcNow, CancellationToken cancellationToken)
 		{
 			// Check we're not already at the maximum number of allowed jobs
-			if (Schedule.MaxActive != 0 && NumActiveJobs >= Schedule.MaxActive)
+			if (schedule.MaxActive != 0 && numActiveJobs >= schedule.MaxActive)
 			{
-				Logger.LogInformation("Skipping trigger of {StreamId} template {TemplateId} - already have maximum number of jobs running ({NumJobs})", Stream.Id, TemplateId, Schedule.MaxActive);
-				foreach (JobId JobId in Schedule.ActiveJobs)
+				_logger.LogInformation("Skipping trigger of {StreamId} template {TemplateId} - already have maximum number of jobs running ({NumJobs})", stream.Id, templateId, schedule.MaxActive);
+				foreach (JobId jobId in schedule.ActiveJobs)
 				{
-					Logger.LogInformation("Active job for {StreamId} template {TemplateId}: {JobId}", Stream.Id, TemplateId, JobId);
+					_logger.LogInformation("Active job for {StreamId} template {TemplateId}: {JobId}", stream.Id, templateId, jobId);
 				}
 				return;
 			}
 
 			// Minimum changelist number, inclusive
-			int MinChangeNumber = Schedule.LastTriggerChange;
-			if (MinChangeNumber > 0 && !Schedule.RequireSubmittedChange)
+			int minChangeNumber = schedule.LastTriggerChange;
+			if (minChangeNumber > 0 && !schedule.RequireSubmittedChange)
 			{
-				MinChangeNumber--;
+				minChangeNumber--;
 			}
 
 			// Maximum changelist number, exclusive
-			int? MaxChangeNumber = null;
+			int? maxChangeNumber = null;
 
 			// Get the maximum number of changes to trigger
-			int MaxNewChanges = 1;
-			if (Schedule.MaxChanges != 0)
+			int maxNewChanges = 1;
+			if (schedule.MaxChanges != 0)
 			{
-				MaxNewChanges = Schedule.MaxChanges;
+				maxNewChanges = schedule.MaxChanges;
 			}
-			if (Schedule.MaxActive != 0)
+			if (schedule.MaxActive != 0)
 			{
-				MaxNewChanges = Math.Min(MaxNewChanges, Schedule.MaxActive - NumActiveJobs);
+				maxNewChanges = Math.Min(maxNewChanges, schedule.MaxActive - numActiveJobs);
 			}
 
 			// Create a timer to limit the amount we look back through P4 history
-			Stopwatch Timer = Stopwatch.StartNew();
-			ChangeContentFlags? FilterFlags = Schedule.GetFilterFlags();
+			Stopwatch timer = Stopwatch.StartNew();
+			ChangeContentFlags? filterFlags = schedule.GetFilterFlags();
 
 			// Create a file filter
-			FileFilter? FileFilter = null;
-			if (Schedule.Files != null)
+			FileFilter? fileFilter = null;
+			if (schedule.Files != null)
 			{
-				FileFilter = new FileFilter(Schedule.Files);
+				fileFilter = new FileFilter(schedule.Files);
 			}
 
 			// Cache the Perforce history as we're iterating through changes to improve query performance
-			PerforceHistory History = new PerforceHistory(Perforce, Stream.ClusterName, Stream.Name);
+			PerforceHistory history = new PerforceHistory(_perforce, stream.ClusterName, stream.Name);
 			
 			// Start as many jobs as possible
-			List<(int Change, int CodeChange)> TriggerChanges = new List<(int, int)>();
-			while (TriggerChanges.Count < MaxNewChanges)
+			List<(int Change, int CodeChange)> triggerChanges = new List<(int, int)>();
+			while (triggerChanges.Count < maxNewChanges)
 			{
-				CancellationToken.ThrowIfCancellationRequested();
+				cancellationToken.ThrowIfCancellationRequested();
 
 				// Get the next valid change
-				ChangeDetails? ChangeDetails = null;
-				if (Schedule.Gate != null)
+				ChangeDetails? changeDetails = null;
+				if (schedule.Gate != null)
 				{
-					ChangeDetails = await GetNextChangeForGateAsync(Stream, TemplateId, Schedule.Gate, MinChangeNumber, MaxChangeNumber, CancellationToken);
+					changeDetails = await GetNextChangeForGateAsync(stream, templateId, schedule.Gate, minChangeNumber, maxChangeNumber, cancellationToken);
 				}
 				else
 				{
-					ChangeDetails = await History.GetNextChangeAsync();
+					changeDetails = await history.GetNextChangeAsync();
 				}
 
 				// Quit if we didn't find anything
-				if (ChangeDetails == null)
+				if (changeDetails == null)
 				{
 					break;
 				}
-				if (ChangeDetails.Number < MinChangeNumber)
+				if (changeDetails.Number < minChangeNumber)
 				{
 					break;
 				}
-				if (ChangeDetails.Number == MinChangeNumber && (Schedule.RequireSubmittedChange || TriggerChanges.Count > 0))
+				if (changeDetails.Number == minChangeNumber && (schedule.RequireSubmittedChange || triggerChanges.Count > 0))
 				{
 					break;
 				}
 
 				// Adjust the changelist for the desired filter
-				int Change = ChangeDetails.Number;
-				if (ShouldBuildChange(ChangeDetails, FilterFlags, FileFilter))
+				int change = changeDetails.Number;
+				if (ShouldBuildChange(changeDetails, filterFlags, fileFilter))
 				{
-					int CodeChange = await History.GetCodeChange(Change);
-					if (CodeChange == -1)
+					int codeChange = await history.GetCodeChange(change);
+					if (codeChange == -1)
 					{
-						Logger.LogWarning("Unable to find code change for CL {Change}", Change);
-						CodeChange = Change;
+						_logger.LogWarning("Unable to find code change for CL {Change}", change);
+						codeChange = change;
 					}
-					TriggerChanges.Add((Change, CodeChange));
+					triggerChanges.Add((change, codeChange));
 				}
 
 				// Check we haven't exceeded the time limit
-				if (Timer.Elapsed > TimeSpan.FromMinutes(2.0))
+				if (timer.Elapsed > TimeSpan.FromMinutes(2.0))
 				{
-					Logger.LogError("Querying for changes to trigger for {StreamId} template {TemplateId} has taken {Time}. Aborting.", Stream.Id, TemplateId, Timer.Elapsed);
+					_logger.LogError("Querying for changes to trigger for {StreamId} template {TemplateId} has taken {Time}. Aborting.", stream.Id, templateId, timer.Elapsed);
 					break;
 				}
 				
 				// Update the remaining range of changes to check for
-				MaxChangeNumber = Change - 1;
-				if(MaxChangeNumber < MinChangeNumber)
+				maxChangeNumber = change - 1;
+				if(maxChangeNumber < minChangeNumber)
 				{
 					break;
 				}
 			}
 
 			// Early out if there's nothing to do
-			if (TriggerChanges.Count == 0)
+			if (triggerChanges.Count == 0)
 			{
-				Logger.LogInformation("Skipping trigger of {StreamName} template {TemplateId} - no candidate changes after CL {LastTriggerChange}", Stream.Id, TemplateId, Schedule.LastTriggerChange);
+				_logger.LogInformation("Skipping trigger of {StreamName} template {TemplateId} - no candidate changes after CL {LastTriggerChange}", stream.Id, templateId, schedule.LastTriggerChange);
 				return;
 			}
 
 			// Get the matching template
-			ITemplate? Template = await TemplateCollection.GetAsync(TemplateRef.Hash);
-			if (Template == null)
+			ITemplate? template = await _templateCollection.GetAsync(templateRef.Hash);
+			if (template == null)
 			{
-				Logger.LogWarning("Unable to find template '{TemplateHash}' for '{TemplateId}'", TemplateRef.Hash, TemplateId);
+				_logger.LogWarning("Unable to find template '{TemplateHash}' for '{TemplateId}'", templateRef.Hash, templateId);
 				return;
 			}
 
 			// Register the graph for it
-			IGraph Graph = await Graphs.AddAsync(Template);
+			IGraph graph = await _graphs.AddAsync(template);
 
 			// We may need to submit a new change for any new jobs. This only makes sense if there's one change.
-			if (Template.SubmitNewChange != null)
+			if (template.SubmitNewChange != null)
 			{
-				int NewChange = await Perforce.CreateNewChangeForTemplateAsync(Stream, Template);
-				int NewCodeChange = await Perforce.GetCodeChangeAsync(Stream.ClusterName, Stream.Name, NewChange);
-				TriggerChanges = new List<(int, int)> { (NewChange, NewCodeChange) };
+				int newChange = await _perforce.CreateNewChangeForTemplateAsync(stream, template);
+				int newCodeChange = await _perforce.GetCodeChangeAsync(stream.ClusterName, stream.Name, newChange);
+				triggerChanges = new List<(int, int)> { (newChange, newCodeChange) };
 			}
 
 			// Try to start all the new jobs
-			Logger.LogInformation("Starting {NumJobs} new jobs for {StreamId} template {TemplateId} (active: {NumActive}, max new: {MaxNewJobs})", TriggerChanges.Count, Stream.Id, TemplateId, NumActiveJobs, MaxNewChanges);
-			foreach ((int Change, int CodeChange) in TriggerChanges.OrderBy(x => x.Change))
+			_logger.LogInformation("Starting {NumJobs} new jobs for {StreamId} template {TemplateId} (active: {NumActive}, max new: {MaxNewJobs})", triggerChanges.Count, stream.Id, templateId, numActiveJobs, maxNewChanges);
+			foreach ((int change, int codeChange) in triggerChanges.OrderBy(x => x.Change))
 			{
-				CancellationToken.ThrowIfCancellationRequested();
-				List<string> DefaultArguments = Template.GetDefaultArguments();
-				IJob NewJob = await JobService.CreateJobAsync(null, Stream, TemplateId, Template.Id, Graph, Template.Name, Change, CodeChange, null, null, null, Template.Priority, null, null, Template.PromoteIssuesByDefault, TemplateRef.ChainedJobs, TemplateRef.ShowUgsBadges, TemplateRef.ShowUgsAlerts, TemplateRef.NotificationChannel, TemplateRef.NotificationChannelFilter, DefaultArguments);
-				Logger.LogInformation("Started new job for {StreamId} template {TemplateId} at CL {Change} (Code CL {CodeChange}): {JobId}", Stream.Id, TemplateId, Change, CodeChange, NewJob.Id);
-				await StreamService.UpdateScheduleTriggerAsync(Stream, TemplateId, UtcNow, Change, new List<JobId> { NewJob.Id }, new List<JobId>());
+				cancellationToken.ThrowIfCancellationRequested();
+				List<string> defaultArguments = template.GetDefaultArguments();
+				IJob newJob = await _jobService.CreateJobAsync(null, stream, templateId, template.Id, graph, template.Name, change, codeChange, null, null, null, template.Priority, null, null, template.PromoteIssuesByDefault, templateRef.ChainedJobs, templateRef.ShowUgsBadges, templateRef.ShowUgsAlerts, templateRef.NotificationChannel, templateRef.NotificationChannelFilter, defaultArguments);
+				_logger.LogInformation("Started new job for {StreamId} template {TemplateId} at CL {Change} (Code CL {CodeChange}): {JobId}", stream.Id, templateId, change, codeChange, newJob.Id);
+				await _streamService.UpdateScheduleTriggerAsync(stream, templateId, utcNow, change, new List<JobId> { newJob.Id }, new List<JobId>());
 			}
 		}
 
 		/// <summary>
 		/// Tests whether a schedule should build a particular change, based on its requested change filters
 		/// </summary>
-		/// <param name="Details">The change details</param>
-		/// <param name="FilterFlags"></param>
-		/// <param name="FileFilter">Filter for the files to trigger a build</param>
+		/// <param name="details">The change details</param>
+		/// <param name="filterFlags"></param>
+		/// <param name="fileFilter">Filter for the files to trigger a build</param>
 		/// <returns></returns>
-		private bool ShouldBuildChange(ChangeDetails Details, ChangeContentFlags? FilterFlags, FileFilter? FileFilter)
+		private bool ShouldBuildChange(ChangeDetails details, ChangeContentFlags? filterFlags, FileFilter? fileFilter)
 		{
-			if (Regex.IsMatch(Details.Description, @"^\s*#\s*skipci", RegexOptions.Multiline))
+			if (Regex.IsMatch(details.Description, @"^\s*#\s*skipci", RegexOptions.Multiline))
 			{
 				return false;
 			}
-			if (FilterFlags != null && FilterFlags.Value != 0)
+			if (filterFlags != null && filterFlags.Value != 0)
 			{
-				if ((Details.GetContentFlags() & FilterFlags.Value) == 0)
+				if ((details.GetContentFlags() & filterFlags.Value) == 0)
 				{
-					Logger.LogDebug("Not building change {Change} ({ChangeFlags}) due to filter flags ({FilterFlags})", Details.Number, Details.GetContentFlags().ToString(), FilterFlags.Value.ToString());
+					_logger.LogDebug("Not building change {Change} ({ChangeFlags}) due to filter flags ({FilterFlags})", details.Number, details.GetContentFlags().ToString(), filterFlags.Value.ToString());
 					return false;
 				}
 			}
-			if (FileFilter != null)
+			if (fileFilter != null)
 			{
-				if (!Details.Files.Any(x => FileFilter.Matches(x.Path)))
+				if (!details.Files.Any(x => fileFilter.Matches(x.Path)))
 				{
-					Logger.LogDebug("Not building change {Change} due to file filter", Details.Number);
+					_logger.LogDebug("Not building change {Change} due to file filter", details.Number);
 					return false;
 				}
 			}
 			return true;
 		}
 
-		/// <summary>
-		/// Get the next change to build in this stream, between a given range
-		/// </summary>
-		/// <returns>The new change and code change to build, or 0.</returns>
-		private async Task<(int Change, int CodeChange)> GetNextChangeAsync(string ClusterName, string StreamName, int? MinChange, int? MaxChange)
-		{
-			List<ChangeSummary> Changes = await Perforce.GetChangesAsync(ClusterName, StreamName, MinChange, MaxChange, 1, null);
-			if (Changes.Count > 0)
-			{
-				int Change = Changes[0].Number;
-				int CodeChange = await Perforce.GetCodeChangeAsync(ClusterName, StreamName, Change);
-				return (Change, CodeChange);
-			}
-			return default;
-		}
-
 		/// <summary>
 		/// Gets the next change to build for a schedule on a gate
 		/// </summary>
 		/// <returns></returns>
-		private async Task<ChangeDetails?> GetNextChangeForGateAsync(IStream Stream, TemplateRefId TemplateRefId, ScheduleGate Gate, int? MinChange, int? MaxChange, CancellationToken CancellationToken)
+		private async Task<ChangeDetails?> GetNextChangeForGateAsync(IStream stream, TemplateRefId templateRefId, ScheduleGate gate, int? minChange, int? maxChange, CancellationToken cancellationToken)
 		{
 			for (; ; )
 			{
-				CancellationToken.ThrowIfCancellationRequested();
+				cancellationToken.ThrowIfCancellationRequested();
 
-				List<IJob> Jobs = await JobCollection.FindAsync(StreamId: Stream.Id, Templates: new[] { Gate.TemplateRefId }, MinChange: MinChange, MaxChange: MaxChange, Count: 1);
-				if (Jobs.Count == 0)
+				List<IJob> jobs = await _jobCollection.FindAsync(streamId: stream.Id, templates: new[] { gate.TemplateRefId }, minChange: minChange, maxChange: maxChange, count: 1);
+				if (jobs.Count == 0)
 				{
 					return null;
 				}
 
-				IJob Job = Jobs[0];
+				IJob job = jobs[0];
 
-				IGraph? Graph = await Graphs.GetAsync(Job.GraphHash);
-				if (Graph != null)
+				IGraph? graph = await _graphs.GetAsync(job.GraphHash);
+				if (graph != null)
 				{
-					(JobStepState, JobStepOutcome)? State = Job.GetTargetState(Graph, Gate.Target);
-					if (State != null && State.Value.Item1 == JobStepState.Completed)
+					(JobStepState, JobStepOutcome)? state = job.GetTargetState(graph, gate.Target);
+					if (state != null && state.Value.Item1 == JobStepState.Completed)
 					{
-						JobStepOutcome Outcome = State.Value.Item2;
-						if (Outcome == JobStepOutcome.Success || Outcome == JobStepOutcome.Warnings)
+						JobStepOutcome outcome = state.Value.Item2;
+						if (outcome == JobStepOutcome.Success || outcome == JobStepOutcome.Warnings)
 						{
-							return await Perforce.GetChangeDetailsAsync(Stream.ClusterName, Stream.Name, Job.Change);
+							return await _perforce.GetChangeDetailsAsync(stream.ClusterName, stream.Name, job.Change);
 						}
-						Logger.LogInformation("Skipping trigger of {StreamName} template {TemplateId} - last {OtherTemplateRefId} job ({JobId}) ended with errors", Stream.Id, TemplateRefId, Gate.TemplateRefId, Job.Id);
+						_logger.LogInformation("Skipping trigger of {StreamName} template {TemplateId} - last {OtherTemplateRefId} job ({JobId}) ended with errors", stream.Id, templateRefId, gate.TemplateRefId, job.Id);
 					}
 				}
 
-				MaxChange = Job.Change - 1;
+				maxChange = job.Change - 1;
 			}
 		}
 	}
