Horde: Fix coding conventions within Horde.Build.

diff --git a/Engine/Source/Programs/Horde/Horde.Build/Projects/StreamTab.cs b/Engine/Source/Programs/Horde/Horde.Build/Projects/StreamTab.cs
--- a/Engine/Source/Programs/Horde/Horde.Build/Projects/StreamTab.cs
+++ b/Engine/Source/Programs/Horde/Horde.Build/Projects/StreamTab.cs
@@ -1,12 +1,10 @@
 // Copyright Epic Games, Inc. All Rights Reserved.
 
+using System;
+using System.Collections.Generic;
 using Horde.Build.Api;
 using Horde.Build.Utilities;
 using MongoDB.Bson.Serialization.Attributes;
-using System;
-using System.Collections.Generic;
-using System.Linq;
-using System.Threading.Tasks;
 
 namespace Horde.Build.Models
 {
@@ -35,35 +33,35 @@ namespace Horde.Build.Models
 		/// <summary>
 		/// Constructor
 		/// </summary>
-		/// <param name="Title">Title of this page</param>
-		protected StreamTab(string Title)
+		/// <param name="title">Title of this page</param>
+		protected StreamTab(string title)
 		{
-			this.Title = Title;
+			Title = title;
 		}
 
 		/// <summary>
 		/// Constructor
 		/// </summary>
-		/// <param name="Request">Public request object</param>
-		protected StreamTab(CreateStreamTabRequest Request)
+		/// <param name="request">Public request object</param>
+		protected StreamTab(CreateStreamTabRequest request)
 		{
-			Title = Request.Title;
+			Title = request.Title;
 		}
 
 		/// <summary>
 		/// Creates an instance from a request
 		/// </summary>
-		/// <param name="Request">The request object</param>
+		/// <param name="request">The request object</param>
 		/// <returns>New tab instance</returns>
-		public static StreamTab FromRequest(CreateStreamTabRequest Request)
+		public static StreamTab FromRequest(CreateStreamTabRequest request)
 		{
-			CreateJobsTabRequest? JobsRequest = Request as CreateJobsTabRequest;
-			if (JobsRequest != null)
+			CreateJobsTabRequest? jobsRequest = request as CreateJobsTabRequest;
+			if (jobsRequest != null)
 			{
-				return new JobsTab(JobsRequest);
+				return new JobsTab(jobsRequest);
 			}
 
-			throw new Exception($"Unknown tab request type '{Request.GetType()}'");
+			throw new Exception($"Unknown tab request type '{request.GetType()}'");
 		}
 
 		/// <summary>
@@ -101,22 +99,22 @@ namespace Horde.Build.Models
 		/// <summary>
 		/// Constructor
 		/// </summary>
-		/// <param name="Heading">Heading for this column</param>
-		/// <param name="RelativeWidth">Relative width of this column.</param>
-		protected JobsTabColumn(string Heading, int? RelativeWidth)
+		/// <param name="heading">Heading for this column</param>
+		/// <param name="relativeWidth">Relative width of this column.</param>
+		protected JobsTabColumn(string heading, int? relativeWidth)
 		{
-			this.Heading = Heading;
-			this.RelativeWidth = RelativeWidth;
+			Heading = heading;
+			RelativeWidth = relativeWidth;
 		}
 
 		/// <summary>
 		/// Constructor
 		/// </summary>
-		/// <param name="Request">Public request object</param>
-		protected JobsTabColumn(CreateJobsTabColumnRequest Request)
+		/// <param name="request">Public request object</param>
+		protected JobsTabColumn(CreateJobsTabColumnRequest request)
 		{
-			this.Heading = Request.Heading;
-			this.RelativeWidth = Request.RelativeWidth;
+			Heading = request.Heading;
+			RelativeWidth = request.RelativeWidth;
 		}
 
 		/// <summary>
@@ -146,13 +144,13 @@ namespace Horde.Build.Models
 		/// <summary>
 		/// Constructor
 		/// </summary>
-		/// <param name="Heading">Heading for this column</param>
-		/// <param name="Category">Category of labels to display in this column</param>
-		/// <param name="RelativeWidth">Relative width of this column.</param>
-		public JobsTabLabelColumn(string Heading, string? Category, int? RelativeWidth)
-			: base(Heading, RelativeWidth)
+		/// <param name="heading">Heading for this column</param>
+		/// <param name="category">Category of labels to display in this column</param>
+		/// <param name="relativeWidth">Relative width of this column.</param>
+		public JobsTabLabelColumn(string heading, string? category, int? relativeWidth)
+			: base(heading, relativeWidth)
 		{
-			this.Category = Category;
+			Category = category;
 		}
 
 		/// <summary>
@@ -185,13 +183,13 @@ namespace Horde.Build.Models
 		/// <summary>
 		/// Constructor
 		/// </summary>
-		/// <param name="Heading">Heading for this column</param>
-		/// <param name="Parameter">Category of labels to display in this column</param>
-		/// <param name="RelativeWidth">Relative width of this column.</param>
-		public JobsTabParameterColumn(string Heading, string? Parameter, int? RelativeWidth)
-			: base(Heading, RelativeWidth)
+		/// <param name="heading">Heading for this column</param>
+		/// <param name="parameter">Category of labels to display in this column</param>
+		/// <param name="relativeWidth">Relative width of this column.</param>
+		public JobsTabParameterColumn(string heading, string? parameter, int? relativeWidth)
+			: base(heading, relativeWidth)
 		{
-			this.Parameter = Parameter;
+			Parameter = parameter;
 		}
 
 		/// <summary>
@@ -233,31 +231,31 @@ namespace Horde.Build.Models
 		/// <summary>
 		/// Constructor
 		/// </summary>
-		/// <param name="Title">Title for this page</param>
-		/// <param name="ShowNames">Show names of jobs on the page</param>
-		/// <param name="Templates">Template names to include on this page</param>
-		/// <param name="JobNames">Names of jobs to include on this page</param>
-		/// <param name="Columns">Columns ot display in the table</param>
-		public JobsTab(string Title, bool ShowNames, List<TemplateRefId>? Templates, List<string>? JobNames, List<JobsTabColumn>? Columns)
-			: base(Title)
+		/// <param name="title">Title for this page</param>
+		/// <param name="showNames">Show names of jobs on the page</param>
+		/// <param name="templates">Template names to include on this page</param>
+		/// <param name="jobNames">Names of jobs to include on this page</param>
+		/// <param name="columns">Columns ot display in the table</param>
+		public JobsTab(string title, bool showNames, List<TemplateRefId>? templates, List<string>? jobNames, List<JobsTabColumn>? columns)
+			: base(title)
 		{
-			this.ShowNames = ShowNames;
-			this.Templates = Templates;
-			this.JobNames = JobNames;
-			this.Columns = Columns;
+			ShowNames = showNames;
+			Templates = templates;
+			JobNames = jobNames;
+			Columns = columns;
 		}
 
 		/// <summary>
 		/// Constructor
 		/// </summary>
-		/// <param name="Request">Public request object</param>
-		public JobsTab(CreateJobsTabRequest Request)
-			: base(Request)
+		/// <param name="request">Public request object</param>
+		public JobsTab(CreateJobsTabRequest request)
+			: base(request)
 		{
-			ShowNames = Request.ShowNames;
-			Templates = Request.Templates?.ConvertAll(x => new TemplateRefId(x));
-			JobNames = Request.JobNames;
-			Columns = Request.Columns?.ConvertAll(x => x.ToModel());
+			ShowNames = request.ShowNames;
+			Templates = request.Templates?.ConvertAll(x => new TemplateRefId(x));
+			JobNames = request.JobNames;
+			Columns = request.Columns?.ConvertAll(x => x.ToModel());
 		}
 
 		/// <inheritdoc/>
diff --git a/Engine/Source/Programs/Horde/Horde.Build/Secrets/Credential.cs b/Engine/Source/Programs/Horde/Horde.Build/Secrets/Credential.cs
--- a/Engine/Source/Programs/Horde/Horde.Build/Secrets/Credential.cs
+++ b/Engine/Source/Programs/Horde/Horde.Build/Secrets/Credential.cs
@@ -1,14 +1,9 @@
 ﻿// Copyright Epic Games, Inc. All Rights Reserved.
 
+using System.Collections.Generic;
 using Horde.Build.Acls;
-using Horde.Build.Api;
-using Horde.Build.Utilities;
 using MongoDB.Bson;
 using MongoDB.Bson.Serialization.Attributes;
-using System;
-using System.Collections.Generic;
-using System.Linq;
-using System.Threading.Tasks;
 
 namespace Horde.Build.Models
 {
@@ -51,29 +46,29 @@ namespace Horde.Build.Models
 		[BsonConstructor]
 		private Credential()
 		{
-			this.Name = null!;
-			this.NormalizedName = null!;
+			Name = null!;
+			NormalizedName = null!;
 		}
 
 		/// <summary>
 		/// Constructor
 		/// </summary>
-		/// <param name="Name">Name of this credential</param>
-		public Credential(string Name)
+		/// <param name="name">Name of this credential</param>
+		public Credential(string name)
 		{
-			this.Id = ObjectId.GenerateNewId();
-			this.Name = Name;
-			this.NormalizedName = GetNormalizedName(Name);
+			Id = ObjectId.GenerateNewId();
+			Name = name;
+			NormalizedName = GetNormalizedName(name);
 		}
 
 		/// <summary>
 		/// Gets the normalized form of a name, used for case-insensitive comparisons
 		/// </summary>
-		/// <param name="Name">The name to normalize</param>
+		/// <param name="name">The name to normalize</param>
 		/// <returns>The normalized name</returns>
-		public static string GetNormalizedName(string Name)
+		public static string GetNormalizedName(string name)
 		{
-			return Name.ToUpperInvariant();
+			return name.ToUpperInvariant();
 		}
 	}
 }
diff --git a/Engine/Source/Programs/Horde/Horde.Build/Secrets/CredentialService.cs b/Engine/Source/Programs/Horde/Horde.Build/Secrets/CredentialService.cs
--- a/Engine/Source/Programs/Horde/Horde.Build/Secrets/CredentialService.cs
+++ b/Engine/Source/Programs/Horde/Horde.Build/Secrets/CredentialService.cs
@@ -1,15 +1,13 @@
 ﻿// Copyright Epic Games, Inc. All Rights Reserved.
 
-using Horde.Build.Acls;
-using Horde.Build.Api;
-using Horde.Build.Models;
-using MongoDB.Bson;
-using MongoDB.Driver;
-using System;
 using System.Collections.Generic;
 using System.Linq;
 using System.Security.Claims;
 using System.Threading.Tasks;
+using Horde.Build.Acls;
+using Horde.Build.Models;
+using MongoDB.Bson;
+using MongoDB.Driver;
 
 namespace Horde.Build.Services
 {
@@ -21,88 +19,88 @@ namespace Horde.Build.Services
 		/// <summary>
 		/// The ACL service instance
 		/// </summary>
-		AclService AclService;
+		readonly AclService _aclService;
 
 		/// <summary>
 		/// Collection of credential documents
 		/// </summary>
-		IMongoCollection<Credential> Credentials; 
+		readonly IMongoCollection<Credential> _credentials; 
 
 		/// <summary>
 		/// Constructor
 		/// </summary>
-		/// <param name="AclService">The ACL service instance</param>
-		/// <param name="DatabaseService">The database service instance</param>
-		public CredentialService(AclService AclService, DatabaseService DatabaseService)
+		/// <param name="aclService">The ACL service instance</param>
+		/// <param name="databaseService">The database service instance</param>
+		public CredentialService(AclService aclService, DatabaseService databaseService)
 		{
-			this.AclService = AclService;
+			_aclService = aclService;
 
-			Credentials = DatabaseService.Credentials;
+			_credentials = databaseService.Credentials;
 
-			if (!DatabaseService.ReadOnlyMode)
+			if (!databaseService.ReadOnlyMode)
 			{
-				Credentials.Indexes.CreateOne(new CreateIndexModel<Credential>(Builders<Credential>.IndexKeys.Ascending(x => x.NormalizedName), new CreateIndexOptions { Unique = true }));
+				_credentials.Indexes.CreateOne(new CreateIndexModel<Credential>(Builders<Credential>.IndexKeys.Ascending(x => x.NormalizedName), new CreateIndexOptions { Unique = true }));
 			}
 		}
 
 		/// <summary>
 		/// Creates a new credential
 		/// </summary>
-		/// <param name="Name">Name of the new credential</param>
-		/// <param name="Properties">Properties for the new credential</param>
+		/// <param name="name">Name of the new credential</param>
+		/// <param name="properties">Properties for the new credential</param>
 		/// <returns>The new credential document</returns>
-		public async Task<Credential> CreateCredentialAsync(string Name, Dictionary<string, string>? Properties)
+		public async Task<Credential> CreateCredentialAsync(string name, Dictionary<string, string>? properties)
 		{
-			Credential NewCredential = new Credential(Name);
-			if (Properties != null)
+			Credential newCredential = new Credential(name);
+			if (properties != null)
 			{
-				NewCredential.Properties = new Dictionary<string, string>(Properties, NewCredential.Properties.Comparer);
+				newCredential.Properties = new Dictionary<string, string>(properties, newCredential.Properties.Comparer);
 			}
 
-			await Credentials.InsertOneAsync(NewCredential);
-			return NewCredential;
+			await _credentials.InsertOneAsync(newCredential);
+			return newCredential;
 		}
 
 		/// <summary>
 		/// Updates an existing credential
 		/// </summary>
-		/// <param name="Id">Unique id of the credential</param>
-		/// <param name="NewName">The new name for the credential</param>
-		/// <param name="NewProperties">Properties on the credential to update. Any properties with a value of null will be removed.</param>
-		/// <param name="NewAcl">The new ACL for the credential</param>
+		/// <param name="id">Unique id of the credential</param>
+		/// <param name="newName">The new name for the credential</param>
+		/// <param name="newProperties">Properties on the credential to update. Any properties with a value of null will be removed.</param>
+		/// <param name="newAcl">The new ACL for the credential</param>
 		/// <returns>Async task object</returns>
-		public async Task UpdateCredentialAsync(ObjectId Id, string? NewName, Dictionary<string, string?>? NewProperties, Acl? NewAcl)
+		public async Task UpdateCredentialAsync(ObjectId id, string? newName, Dictionary<string, string?>? newProperties, Acl? newAcl)
 		{
-			UpdateDefinitionBuilder<Credential> UpdateBuilder = Builders<Credential>.Update;
+			UpdateDefinitionBuilder<Credential> updateBuilder = Builders<Credential>.Update;
 
-			List<UpdateDefinition<Credential>> Updates = new List<UpdateDefinition<Credential>>();
-			if (NewName != null)
+			List<UpdateDefinition<Credential>> updates = new List<UpdateDefinition<Credential>>();
+			if (newName != null)
 			{
-				Updates.Add(UpdateBuilder.Set(x => x.Name, NewName));
-				Updates.Add(UpdateBuilder.Set(x => x.NormalizedName, Credential.GetNormalizedName(NewName)));
+				updates.Add(updateBuilder.Set(x => x.Name, newName));
+				updates.Add(updateBuilder.Set(x => x.NormalizedName, Credential.GetNormalizedName(newName)));
 			}
-			if (NewProperties != null)
+			if (newProperties != null)
 			{
-				foreach (KeyValuePair<string, string?> Pair in NewProperties)
+				foreach (KeyValuePair<string, string?> pair in newProperties)
 				{
-					if (Pair.Value == null)
+					if (pair.Value == null)
 					{
-						Updates.Add(UpdateBuilder.Unset(x => x.Properties[Pair.Key]));
+						updates.Add(updateBuilder.Unset(x => x.Properties[pair.Key]));
 					}
 					else
 					{
-						Updates.Add(UpdateBuilder.Set(x => x.Properties[Pair.Key], Pair.Value));
+						updates.Add(updateBuilder.Set(x => x.Properties[pair.Key], pair.Value));
 					}
 				}
 			}
-			if (NewAcl != null)
+			if (newAcl != null)
 			{
-				Updates.Add(Acl.CreateUpdate<Credential>(x => x.Acl!, NewAcl));
+				updates.Add(Acl.CreateUpdate<Credential>(x => x.Acl!, newAcl));
 			}
 
-			if (Updates.Count > 0)
+			if (updates.Count > 0)
 			{
-				await Credentials.FindOneAndUpdateAsync<Credential>(x => x.Id == Id, UpdateBuilder.Combine(Updates));
+				await _credentials.FindOneAndUpdateAsync<Credential>(x => x.Id == id, updateBuilder.Combine(updates));
 			}
 		}
 
@@ -110,69 +108,69 @@ namespace Horde.Build.Services
 		/// Gets all the available credentials
 		/// </summary>
 		/// <returns>List of project documents</returns>
-		public Task<List<Credential>> FindCredentialsAsync(string? Name)
+		public Task<List<Credential>> FindCredentialsAsync(string? name)
 		{
-			FilterDefinitionBuilder<Credential> FilterBuilder = Builders<Credential>.Filter;
+			FilterDefinitionBuilder<Credential> filterBuilder = Builders<Credential>.Filter;
 
-			FilterDefinition<Credential> Filter = FilterDefinition<Credential>.Empty;
-			if (Name != null)
+			FilterDefinition<Credential> filter = FilterDefinition<Credential>.Empty;
+			if (name != null)
 			{
-				Filter &= FilterBuilder.Eq(x => x.NormalizedName, Credential.GetNormalizedName(Name));
+				filter &= filterBuilder.Eq(x => x.NormalizedName, Credential.GetNormalizedName(name));
 			}
 
-			return Credentials.Find(Filter).ToListAsync();
+			return _credentials.Find(filter).ToListAsync();
 		}
 
 		/// <summary>
 		/// Gets a credential by ID
 		/// </summary>
-		/// <param name="Id">Unique id of the credential</param>
+		/// <param name="id">Unique id of the credential</param>
 		/// <returns>The credential document</returns>
-		public async Task<Credential?> GetCredentialAsync(ObjectId Id)
+		public async Task<Credential?> GetCredentialAsync(ObjectId id)
 		{
-			return await Credentials.Find<Credential>(x => x.Id == Id).FirstOrDefaultAsync();
+			return await _credentials.Find<Credential>(x => x.Id == id).FirstOrDefaultAsync();
 		}
 
 		/// <summary>
 		/// Gets a credential by name
 		/// </summary>
-		/// <param name="Name">Name of the credential</param>
+		/// <param name="name">Name of the credential</param>
 		/// <returns>The credential document</returns>
-		public async Task<Credential?> GetCredentialAsync(string Name)
+		public async Task<Credential?> GetCredentialAsync(string name)
 		{
-			string NormalizedName = Credential.GetNormalizedName(Name);
-			return await Credentials.Find<Credential>(x => x.NormalizedName == NormalizedName).FirstOrDefaultAsync();
+			string normalizedName = Credential.GetNormalizedName(name);
+			return await _credentials.Find<Credential>(x => x.NormalizedName == normalizedName).FirstOrDefaultAsync();
 		}
 
 		/// <summary>
 		/// Deletes a credential by id
 		/// </summary>
-		/// <param name="Id">Unique id of the credential</param>
+		/// <param name="id">Unique id of the credential</param>
 		/// <returns>True if the credential was deleted</returns>
-		public async Task<bool> DeleteCredentialAsync(ObjectId Id)
+		public async Task<bool> DeleteCredentialAsync(ObjectId id)
 		{
-			DeleteResult Result = await Credentials.DeleteOneAsync<Credential>(x => x.Id == Id);
-			return Result.DeletedCount > 0;
+			DeleteResult result = await _credentials.DeleteOneAsync<Credential>(x => x.Id == id);
+			return result.DeletedCount > 0;
 		}
 
 		/// <summary>
 		/// Determines if the user is authorized to perform an action on a particular credential
 		/// </summary>
-		/// <param name="Credential">The credential to check</param>
-		/// <param name="Action">The action being performed</param>
-		/// <param name="User">The principal to authorize</param>
-		/// <param name="Cache">Permissions cache</param>
+		/// <param name="credential">The credential to check</param>
+		/// <param name="action">The action being performed</param>
+		/// <param name="user">The principal to authorize</param>
+		/// <param name="cache">Permissions cache</param>
 		/// <returns>True if the action is authorized</returns>
-		public Task<bool> AuthorizeAsync(Credential Credential, AclAction Action, ClaimsPrincipal User, GlobalPermissionsCache? Cache)
+		public Task<bool> AuthorizeAsync(Credential credential, AclAction action, ClaimsPrincipal user, GlobalPermissionsCache? cache)
 		{
-			bool? Result = Credential.Acl?.Authorize(Action, User);
-			if (Result == null)
+			bool? result = credential.Acl?.Authorize(action, user);
+			if (result == null)
 			{
-				return AclService.AuthorizeAsync(Action, User, Cache);
+				return _aclService.AuthorizeAsync(action, user, cache);
 			}
 			else
 			{
-				return Task.FromResult(Result.Value);
+				return Task.FromResult(result.Value);
 			}
 		}
 	}
diff --git a/Engine/Source/Programs/Horde/Horde.Build/Server/Config.cs b/Engine/Source/Programs/Horde/Horde.Build/Server/Config.cs
--- a/Engine/Source/Programs/Horde/Horde.Build/Server/Config.cs
+++ b/Engine/Source/Programs/Horde/Horde.Build/Server/Config.cs
@@ -1,17 +1,13 @@
 // Copyright Epic Games, Inc. All Rights Reserved.
 
+using System;
+using System.Collections.Generic;
+using System.ComponentModel.DataAnnotations;
 using EpicGames.Core;
 using EpicGames.Horde.Common;
-using EpicGames.Horde.Compute;
-using EpicGames.Horde.Storage;
 using EpicGames.Serialization;
 using Horde.Build.Api;
 using Horde.Build.Utilities;
-using System;
-using System.Collections.Generic;
-using System.ComponentModel.DataAnnotations;
-using System.Linq;
-using System.Threading.Tasks;
 
 namespace Horde.Build.Models
 {
diff --git a/Engine/Source/Programs/Horde/Horde.Build/Server/ConfigService.cs b/Engine/Source/Programs/Horde/Horde.Build/Server/ConfigService.cs
--- a/Engine/Source/Programs/Horde/Horde.Build/Server/ConfigService.cs
+++ b/Engine/Source/Programs/Horde/Horde.Build/Server/ConfigService.cs
@@ -1,35 +1,28 @@
-﻿// Copyright Epic Games, Inc. All Rights Reserved.
+// Copyright Epic Games, Inc. All Rights Reserved.
 
+using System;
+using System.Collections.Generic;
+using System.Globalization;
+using System.IO;
+using System.Linq;
+using System.Reflection;
+using System.Text.Json;
+using System.Text.Json.Serialization;
+using System.Threading;
+using System.Threading.Tasks;
 using EpicGames.Core;
 using Horde.Build.Acls;
 using Horde.Build.Api;
 using Horde.Build.Collections;
-using Horde.Build.Controllers;
 using Horde.Build.Models;
 using Horde.Build.Notifications;
 using Horde.Build.Utilities;
+using HordeCommon;
 using Microsoft.AspNetCore.StaticFiles;
 using Microsoft.Extensions.Configuration;
+using Microsoft.Extensions.Hosting;
 using Microsoft.Extensions.Logging;
 using Microsoft.Extensions.Options;
-using MongoDB.Bson;
-using System;
-using System.Collections.Generic;
-using System.ComponentModel.DataAnnotations;
-using System.Linq;
-using System.IO;
-using System.Text;
-using System.Text.Json;
-using System.Text.RegularExpressions;
-using System.Threading;
-using System.Threading.Tasks;
-using System.Reflection;
-
-using System.Globalization;
-using Horde.Build.Storage;
-using System.Text.Json.Serialization;
-using HordeCommon;
-using Microsoft.Extensions.Hosting;
 
 namespace Horde.Build.Services
 {
@@ -49,204 +42,200 @@ namespace Horde.Build.Services
 		/// Config file version number
 		/// </summary>
 		const int Version = 10;
-
-		DatabaseService DatabaseService;
-		ProjectService ProjectService;
-		StreamService StreamService;
-		IPerforceService PerforceService;
-		PerforceLoadBalancer PerforceLoadBalancer;
-		INotificationService NotificationService;
-		AgentService AgentService;
-		PoolService PoolService;
-		IOptionsMonitor<ServerSettings> Settings;
-		ITicker Ticker;
-		ILogger Logger;
+		readonly DatabaseService _databaseService;
+		readonly ProjectService _projectService;
+		readonly StreamService _streamService;
+		readonly IPerforceService _perforceService;
+		readonly INotificationService _notificationService;
+		readonly AgentService _agentService;
+		readonly PoolService _poolService;
+		readonly IOptionsMonitor<ServerSettings> _settings;
+		readonly ITicker _ticker;
+		readonly ILogger _logger;
 
 		/// <summary>
 		/// Constructor
 		/// </summary>
-		public ConfigService(DatabaseService DatabaseService, IPerforceService PerforceService, ProjectService ProjectService, StreamService StreamService, INotificationService NotificationService, PoolService PoolService, AgentService AgentService, PerforceLoadBalancer PerforceLoadBalancer, IClock Clock, IOptionsMonitor<ServerSettings> Settings, ILogger<ConfigService> Logger)
+		public ConfigService(DatabaseService databaseService, IPerforceService perforceService, ProjectService projectService, StreamService streamService, INotificationService notificationService, PoolService poolService, AgentService agentService, IClock clock, IOptionsMonitor<ServerSettings> settings, ILogger<ConfigService> logger)
 		{
-			this.DatabaseService = DatabaseService;
-			this.PerforceService = PerforceService;
-			this.PerforceLoadBalancer = PerforceLoadBalancer;
-			this.ProjectService = ProjectService;
-			this.StreamService = StreamService;
-			this.NotificationService = NotificationService;
-			this.PoolService = PoolService;
-			this.AgentService = AgentService;
-			this.Settings = Settings;
-			if (DatabaseService.ReadOnlyMode)
+			_databaseService = databaseService;
+			_perforceService = perforceService;
+			_projectService = projectService;
+			_streamService = streamService;
+			_notificationService = notificationService;
+			_poolService = poolService;
+			_agentService = agentService;
+			_settings = settings;
+			if (databaseService.ReadOnlyMode)
 			{
-				this.Ticker = new NullTicker();
+				_ticker = new NullTicker();
 			}
 			else
 			{
-				this.Ticker = Clock.AddSharedTicker<ConfigService>(TimeSpan.FromMinutes(1.0), TickLeaderAsync, Logger);
+				_ticker = clock.AddSharedTicker<ConfigService>(TimeSpan.FromMinutes(1.0), TickLeaderAsync, logger);
 			}
-			this.Logger = Logger;
+			_logger = logger;
 
 			// This will trigger if the local Horde.json user configuration is changed
-			this.Settings.OnChange(OnUserConfigUpdated);
+			_settings.OnChange(OnUserConfigUpdated);
 		}
 
 		/// <inheritdoc/>
-		public Task StartAsync(CancellationToken cancellationToken) => Ticker.StartAsync();
+		public Task StartAsync(CancellationToken cancellationToken) => _ticker.StartAsync();
 
 		/// <inheritdoc/>
-		public Task StopAsync(CancellationToken cancellationToken) => Ticker.StopAsync();
+		public Task StopAsync(CancellationToken cancellationToken) => _ticker.StopAsync();
 
 		/// <inheritdoc/>
-		public void Dispose() => Ticker.Dispose();
+		public void Dispose() => _ticker.Dispose();
 
+		GlobalConfig? _cachedGlobalConfig;
+		string? _cachedGlobalConfigRevision;
+		Dictionary<ProjectId, (ProjectConfig Config, string Revision)> _cachedProjectConfigs = new Dictionary<ProjectId, (ProjectConfig, string)>();
+		readonly Dictionary<ProjectId, string?> _cachedLogoRevisions = new Dictionary<ProjectId, string?>();
 
-		GlobalConfig? CachedGlobalConfig;
-		string? CachedGlobalConfigRevision;
-		Dictionary<ProjectId, (ProjectConfig Config, string Revision)> CachedProjectConfigs = new Dictionary<ProjectId, (ProjectConfig, string)>();
-		Dictionary<ProjectId, string?> CachedLogoRevisions = new Dictionary<ProjectId, string?>();
-
-		async Task UpdateConfigAsync(Uri ConfigPath)
+		async Task UpdateConfigAsync(Uri configPath)
 		{
 			// Update the globals singleton
-			GlobalConfig GlobalConfig;
+			GlobalConfig globalConfig;
 			for (; ; )
 			{
-				Dictionary<Uri, string> GlobalRevisions = await FindRevisionsAsync(new[] { ConfigPath });
-				if (GlobalRevisions.Count == 0)
+				Dictionary<Uri, string> globalRevisions = await FindRevisionsAsync(new[] { configPath });
+				if (globalRevisions.Count == 0)
 				{
-					throw new Exception($"Invalid config path: {ConfigPath}");
+					throw new Exception($"Invalid config path: {configPath}");
 				}
 
-				string Revision = GlobalRevisions.First().Value;
-				if (CachedGlobalConfig == null || Revision != CachedGlobalConfigRevision)
+				string revision = globalRevisions.First().Value;
+				if (_cachedGlobalConfig == null || revision != _cachedGlobalConfigRevision)
 				{
-					Logger.LogInformation("Caching global config from {Revision}", Revision);
+					_logger.LogInformation("Caching global config from {Revision}", revision);
 					try
 					{
-						CachedGlobalConfig = await ReadDataAsync<GlobalConfig>(ConfigPath);
-						CachedGlobalConfigRevision = Revision;
+						_cachedGlobalConfig = await ReadDataAsync<GlobalConfig>(configPath);
+						_cachedGlobalConfigRevision = revision;
 					}
-					catch (Exception Ex)
+					catch (Exception ex)
 					{
-						await SendFailureNotificationAsync(Ex, ConfigPath);
+						await SendFailureNotificationAsync(ex, configPath);
 						return;
 					}
 				}
-				GlobalConfig = CachedGlobalConfig;
+				globalConfig = _cachedGlobalConfig;
 
-				Globals Globals = await DatabaseService.GetGlobalsAsync();
-				if (Globals.ConfigRevision == Revision)
+				Globals globals = await _databaseService.GetGlobalsAsync();
+				if (globals.ConfigRevision == revision)
 				{
-					Logger.LogInformation("Updating configuration from {ConfigPath}", Globals.ConfigRevision);
+					_logger.LogInformation("Updating configuration from {ConfigPath}", globals.ConfigRevision);
 					break;
 				}
 
-				Globals.ConfigRevision = Revision;
-				Globals.PerforceClusters = CachedGlobalConfig.PerforceClusters;
-				Globals.ScheduledDowntime = CachedGlobalConfig.Downtime;
-				Globals.MaxConformCount = CachedGlobalConfig.MaxConformCount;
-				Globals.ComputeClusters = CachedGlobalConfig.Compute;
-				Globals.RootAcl = Acl.Merge(null, CachedGlobalConfig.Acl);
+				globals.ConfigRevision = revision;
+				globals.PerforceClusters = _cachedGlobalConfig.PerforceClusters;
+				globals.ScheduledDowntime = _cachedGlobalConfig.Downtime;
+				globals.MaxConformCount = _cachedGlobalConfig.MaxConformCount;
+				globals.ComputeClusters = _cachedGlobalConfig.Compute;
+				globals.RootAcl = Acl.Merge(null, _cachedGlobalConfig.Acl);
 
-				if (await DatabaseService.TryUpdateSingletonAsync(Globals))
+				if (await _databaseService.TryUpdateSingletonAsync(globals))
 				{
 					break;
 				}
 			}
 
 			// Update the agent rate table
-			await AgentService.UpdateRateTableAsync(GlobalConfig.Rates);
+			await _agentService.UpdateRateTableAsync(globalConfig.Rates);
 
 			// Projects to remove
-			List<IProject> Projects = await ProjectService.GetProjectsAsync();
+			List<IProject> projects = await _projectService.GetProjectsAsync();
 
 			// Get the path to all the project configs
-			List<(ProjectConfigRef ProjectRef, Uri Path)> ProjectConfigs = GlobalConfig.Projects.Select(x => (x, CombinePaths(ConfigPath, x.Path))).ToList();
+			List<(ProjectConfigRef ProjectRef, Uri Path)> projectConfigs = globalConfig.Projects.Select(x => (x, CombinePaths(configPath, x.Path))).ToList();
 
-			Dictionary<ProjectId, (ProjectConfig Config, string Revision)> PrevCachedProjectConfigs = CachedProjectConfigs;
-			CachedProjectConfigs = new Dictionary<ProjectId, (ProjectConfig, string)>();
+			Dictionary<ProjectId, (ProjectConfig Config, string Revision)> prevCachedProjectConfigs = _cachedProjectConfigs;
+			_cachedProjectConfigs = new Dictionary<ProjectId, (ProjectConfig, string)>();
 
-			List<(ProjectId ProjectId, Uri Path)> ProjectLogos = new List<(ProjectId ProjectId, Uri Path)>();
+			List<(ProjectId ProjectId, Uri Path)> projectLogos = new List<(ProjectId ProjectId, Uri Path)>();
 
-			List<(ProjectId ProjectId, StreamConfigRef StreamRef, Uri Path)> StreamConfigs = new List<(ProjectId, StreamConfigRef, Uri)>();
+			List<(ProjectId ProjectId, StreamConfigRef StreamRef, Uri Path)> streamConfigs = new List<(ProjectId, StreamConfigRef, Uri)>();
 
 			// List of project ids that were not able to be updated. We will avoid removing any existing project or stream definitions for these.
-			HashSet<ProjectId> SkipProjectIds = new HashSet<ProjectId>();
+			HashSet<ProjectId> skipProjectIds = new HashSet<ProjectId>();
 
 			// Update any existing projects
-			Dictionary<Uri, string> ProjectRevisions = await FindRevisionsAsync(ProjectConfigs.Select(x => x.Path));
-			for (int Idx = 0; Idx < ProjectConfigs.Count; Idx++)
+			Dictionary<Uri, string> projectRevisions = await FindRevisionsAsync(projectConfigs.Select(x => x.Path));
+			for (int idx = 0; idx < projectConfigs.Count; idx++)
 			{
 				// Make sure we were able to fetch metadata for 
-				(ProjectConfigRef ProjectRef, Uri ProjectPath) = ProjectConfigs[Idx];
-				if (!ProjectRevisions.TryGetValue(ProjectPath, out string? Revision))
+				(ProjectConfigRef projectRef, Uri projectPath) = projectConfigs[idx];
+				if (!projectRevisions.TryGetValue(projectPath, out string? revision))
 				{
-					Logger.LogWarning("Unable to update project {ProjectId} due to missing revision information", ProjectRef.Id);
-					SkipProjectIds.Add(ProjectRef.Id);
+					_logger.LogWarning("Unable to update project {ProjectId} due to missing revision information", projectRef.Id);
+					skipProjectIds.Add(projectRef.Id);
 					continue;
 				}
 
-				IProject? Project = Projects.FirstOrDefault(x => x.Id == ProjectRef.Id);
-				bool Update = (Project == null || Project.ConfigPath != ProjectPath.ToString() || Project.ConfigRevision != Revision);
+				IProject? project = projects.FirstOrDefault(x => x.Id == projectRef.Id);
+				bool update = (project == null || project.ConfigPath != projectPath.ToString() || project.ConfigRevision != revision);
 
-				ProjectConfig? ProjectConfig;
-				if (!Update && PrevCachedProjectConfigs.TryGetValue(ProjectRef.Id, out (ProjectConfig Config, string Revision) Result) && Result.Revision == Revision)
+				ProjectConfig? projectConfig;
+				if (!update && prevCachedProjectConfigs.TryGetValue(projectRef.Id, out (ProjectConfig Config, string Revision) result) && result.Revision == revision)
 				{
-					ProjectConfig = Result.Config;
+					projectConfig = result.Config;
 				}
 				else
 				{
-					Logger.LogInformation("Caching configuration for project {ProjectId} ({Revision})", ProjectRef.Id, Revision);
+					_logger.LogInformation("Caching configuration for project {ProjectId} ({Revision})", projectRef.Id, revision);
 					try
 					{
-						ProjectConfig = await ReadDataAsync<ProjectConfig>(ProjectPath);
-						if (Update)
+						projectConfig = await ReadDataAsync<ProjectConfig>(projectPath);
+						if (update)
 						{
-							Logger.LogInformation("Updating configuration for project {ProjectId} ({Revision})", ProjectRef.Id, Revision);
-							await ProjectService.Collection.AddOrUpdateAsync(ProjectRef.Id, ProjectPath.ToString(), Revision, Idx, ProjectConfig);
+							_logger.LogInformation("Updating configuration for project {ProjectId} ({Revision})", projectRef.Id, revision);
+							await _projectService.Collection.AddOrUpdateAsync(projectRef.Id, projectPath.ToString(), revision, idx, projectConfig);
 						}
 					}
-					catch (Exception Ex)
+					catch (Exception ex)
 					{
-						await SendFailureNotificationAsync(Ex, ProjectPath);
-						SkipProjectIds.Add(ProjectRef.Id);
+						await SendFailureNotificationAsync(ex, projectPath);
+						skipProjectIds.Add(projectRef.Id);
 						continue;
 					}
 				}
 
-				if (ProjectConfig.Logo != null)
+				if (projectConfig.Logo != null)
 				{
-					ProjectLogos.Add((ProjectRef.Id, CombinePaths(ProjectPath, ProjectConfig.Logo)));
+					projectLogos.Add((projectRef.Id, CombinePaths(projectPath, projectConfig.Logo)));
 				}
 
-				CachedProjectConfigs[ProjectRef.Id] = (ProjectConfig, Revision);
-				StreamConfigs.AddRange(ProjectConfig.Streams.Select(x => (ProjectRef.Id, x, CombinePaths(ProjectPath, x.Path))));
+				_cachedProjectConfigs[projectRef.Id] = (projectConfig, revision);
+				streamConfigs.AddRange(projectConfig.Streams.Select(x => (projectRef.Id, x, CombinePaths(projectPath, x.Path))));
 			}
 
 			// Get the logo revisions
-			Dictionary<Uri, string> LogoRevisions = await FindRevisionsAsync(ProjectLogos.Select(x => x.Path));
-			for (int Idx = 0; Idx < ProjectLogos.Count; Idx++)
+			Dictionary<Uri, string> logoRevisions = await FindRevisionsAsync(projectLogos.Select(x => x.Path));
+			for (int idx = 0; idx < projectLogos.Count; idx++)
 			{
-				(ProjectId ProjectId, Uri Path) = ProjectLogos[Idx];
-				if (LogoRevisions.TryGetValue(Path, out string? Revision))
+				(ProjectId projectId, Uri path) = projectLogos[idx];
+				if (logoRevisions.TryGetValue(path, out string? revision))
 				{
-					string? CurrentRevision;
-					if (!CachedLogoRevisions.TryGetValue(ProjectId, out CurrentRevision))
+					string? currentRevision;
+					if (!_cachedLogoRevisions.TryGetValue(projectId, out currentRevision))
 					{
-						CurrentRevision = (await ProjectService.Collection.GetLogoAsync(ProjectId))?.Revision;
-						CachedLogoRevisions[ProjectId] = CurrentRevision;
+						currentRevision = (await _projectService.Collection.GetLogoAsync(projectId))?.Revision;
+						_cachedLogoRevisions[projectId] = currentRevision;
 					}
-					if (Revision != CurrentRevision)
+					if (revision != currentRevision)
 					{
-						Logger.LogInformation("Updating logo for project {ProjectId} ({Revision})", ProjectId, Revision);
+						_logger.LogInformation("Updating logo for project {ProjectId} ({Revision})", projectId, revision);
 						try
 						{
-							await ProjectService.Collection.SetLogoAsync(ProjectId, Path.ToString(), Revision, GetMimeTypeFromPath(Path), await ReadDataAsync(Path));
-							CachedLogoRevisions[ProjectId] = Revision;
+							await _projectService.Collection.SetLogoAsync(projectId, path.ToString(), revision, GetMimeTypeFromPath(path), await ReadDataAsync(path));
+							_cachedLogoRevisions[projectId] = revision;
 						}
-						catch (Exception Ex)
+						catch (Exception ex)
 						{
-							await SendFailureNotificationAsync(Ex, Path);
+							await SendFailureNotificationAsync(ex, path);
 							continue;
 						}
 					}
@@ -254,27 +243,27 @@ namespace Horde.Build.Services
 			}
 
 			// Get the current streams
-			List<IStream> Streams = await StreamService.GetStreamsAsync();
+			List<IStream> streams = await _streamService.GetStreamsAsync();
 
 			// Get the revisions for all the stream documents
-			Dictionary<Uri, string> StreamRevisions = await FindRevisionsAsync(StreamConfigs.Select(x => x.Path));
-			for (int Idx = 0; Idx < StreamConfigs.Count; Idx++)
+			Dictionary<Uri, string> streamRevisions = await FindRevisionsAsync(streamConfigs.Select(x => x.Path));
+			for (int idx = 0; idx < streamConfigs.Count; idx++)
 			{
-				(ProjectId ProjectId, StreamConfigRef StreamRef, Uri StreamPath) = StreamConfigs[Idx];
-				if (StreamRevisions.TryGetValue(StreamPath, out string? Revision))
+				(ProjectId projectId, StreamConfigRef streamRef, Uri streamPath) = streamConfigs[idx];
+				if (streamRevisions.TryGetValue(streamPath, out string? revision))
 				{
-					IStream? Stream = Streams.FirstOrDefault(x => x.Id == StreamRef.Id);
-					if (Stream == null || Stream.ConfigPath != StreamPath.ToString() || Stream.ConfigRevision != Revision)
+					IStream? stream = streams.FirstOrDefault(x => x.Id == streamRef.Id);
+					if (stream == null || stream.ConfigPath != streamPath.ToString() || stream.ConfigRevision != revision)
 					{
-						Logger.LogInformation("Updating configuration for stream {StreamRef} ({Revision})", StreamRef.Id, Revision);
+						_logger.LogInformation("Updating configuration for stream {StreamRef} ({Revision})", streamRef.Id, revision);
 						try
 						{
-							StreamConfig StreamConfig = await ReadDataAsync<StreamConfig>(StreamPath);
-							Stream = await StreamService.StreamCollection.CreateOrReplaceAsync(StreamRef.Id, Stream, StreamPath.ToString(), Revision, ProjectId, StreamConfig);
+							StreamConfig streamConfig = await ReadDataAsync<StreamConfig>(streamPath);
+							stream = await _streamService.StreamCollection.CreateOrReplaceAsync(streamRef.Id, stream, streamPath.ToString(), revision, projectId, streamConfig);
 						}
-						catch (Exception Ex)
+						catch (Exception ex)
 						{
-							await SendFailureNotificationAsync(Ex, StreamPath);
+							await SendFailureNotificationAsync(ex, streamPath);
 							continue;
 						}
 					}
@@ -282,171 +271,171 @@ namespace Horde.Build.Services
 			}
 
 			// Remove any projects which are no longer used
-			HashSet<ProjectId> RemoveProjectIds = new HashSet<ProjectId>(Projects.Select(x => x.Id));
-			RemoveProjectIds.ExceptWith(ProjectConfigs.Select(y => y.ProjectRef.Id));
+			HashSet<ProjectId> removeProjectIds = new HashSet<ProjectId>(projects.Select(x => x.Id));
+			removeProjectIds.ExceptWith(projectConfigs.Select(y => y.ProjectRef.Id));
 
-			foreach (ProjectId RemoveProjectId in RemoveProjectIds)
+			foreach (ProjectId removeProjectId in removeProjectIds)
 			{
-				Logger.LogInformation("Removing project {ProjectId}", RemoveProjectId);
-				await ProjectService.DeleteProjectAsync(RemoveProjectId);
+				_logger.LogInformation("Removing project {ProjectId}", removeProjectId);
+				await _projectService.DeleteProjectAsync(removeProjectId);
 			}
 
 			// Remove any streams that are no longer used
-			HashSet<StreamId> RemoveStreamIds = new HashSet<StreamId>(Streams.Where(x => !SkipProjectIds.Contains(x.ProjectId)).Select(x => x.Id));
-			RemoveStreamIds.ExceptWith(StreamConfigs.Select(x => x.StreamRef.Id));
+			HashSet<StreamId> removeStreamIds = new HashSet<StreamId>(streams.Where(x => !skipProjectIds.Contains(x.ProjectId)).Select(x => x.Id));
+			removeStreamIds.ExceptWith(streamConfigs.Select(x => x.StreamRef.Id));
 
-			foreach (StreamId RemoveStreamId in RemoveStreamIds)
+			foreach (StreamId removeStreamId in removeStreamIds)
 			{
-				Logger.LogInformation("Removing stream {StreamId}", RemoveStreamId);
-				await StreamService.DeleteStreamAsync(RemoveStreamId);
+				_logger.LogInformation("Removing stream {StreamId}", removeStreamId);
+				await _streamService.DeleteStreamAsync(removeStreamId);
 			}
 		}
 
-		static FileExtensionContentTypeProvider ContentTypeProvider = new FileExtensionContentTypeProvider();
+		static readonly FileExtensionContentTypeProvider s_contentTypeProvider = new FileExtensionContentTypeProvider();
 
-		static string GetMimeTypeFromPath(Uri Path)
+		static string GetMimeTypeFromPath(Uri path)
 		{
-			string? ContentType;
-			if (!ContentTypeProvider.TryGetContentType(Path.AbsolutePath, out ContentType))
+			string? contentType;
+			if (!s_contentTypeProvider.TryGetContentType(path.AbsolutePath, out contentType))
 			{
-				ContentType = "application/octet-stream";
+				contentType = "application/octet-stream";
 			}
-			return ContentType;
+			return contentType;
 		}
 
-		static Uri CombinePaths(Uri BaseUri, string Path)
+		static Uri CombinePaths(Uri baseUri, string path)
 		{
-			if (Path.StartsWith("//", StringComparison.Ordinal))
+			if (path.StartsWith("//", StringComparison.Ordinal))
 			{
-				if (BaseUri.Scheme == PerforceScheme)
+				if (baseUri.Scheme == PerforceScheme)
 				{
-					return new Uri($"{PerforceScheme}://{BaseUri.Host}{Path}");
+					return new Uri($"{PerforceScheme}://{baseUri.Host}{path}");
 				}
 				else
 				{
-					return new Uri($"{PerforceScheme}://{PerforceCluster.DefaultName}{Path}");
+					return new Uri($"{PerforceScheme}://{PerforceCluster.DefaultName}{path}");
 				}
 			}
-			return new Uri(BaseUri, Path);
+			return new Uri(baseUri, path);
 		}
 
-		async Task<Dictionary<Uri, string>> FindRevisionsAsync(IEnumerable<Uri> Paths)
+		async Task<Dictionary<Uri, string>> FindRevisionsAsync(IEnumerable<Uri> paths)
 		{
-			Dictionary<Uri, string> Revisions = new Dictionary<Uri, string>();
+			Dictionary<Uri, string> revisions = new Dictionary<Uri, string>();
 
 			// Find all the Perforce uris
-			List<Uri> PerforcePaths = new List<Uri>();
-			foreach (Uri Path in Paths)
+			List<Uri> perforcePaths = new List<Uri>();
+			foreach (Uri path in paths)
 			{
-				if (Path.Scheme == FileScheme)
+				if (path.Scheme == FileScheme)
 				{
-					Revisions[Path] = $"ver={Version},md5={ContentHash.MD5(new FileReference(Path.LocalPath))}";
+					revisions[path] = $"ver={Version},md5={ContentHash.MD5(new FileReference(path.LocalPath))}";
 				}
-				else if (Path.Scheme == PerforceScheme)
+				else if (path.Scheme == PerforceScheme)
 				{
-					PerforcePaths.Add(Path);
+					perforcePaths.Add(path);
 				}
 				else
 				{
-					throw new Exception($"Invalid path format: {Path}");
+					throw new Exception($"Invalid path format: {path}");
 				}
 			}
 
 			// Query all the Perforce revisions
-			foreach (IGrouping<string, Uri> PerforcePath in PerforcePaths.GroupBy(x => x.Host, StringComparer.OrdinalIgnoreCase))
+			foreach (IGrouping<string, Uri> perforcePath in perforcePaths.GroupBy(x => x.Host, StringComparer.OrdinalIgnoreCase))
 			{
-				List<FileSummary> Files = await PerforceService.FindFilesAsync(PerforcePath.Key, PerforcePath.Select(x => x.AbsolutePath));
-				foreach (FileSummary File in Files)
+				List<FileSummary> files = await _perforceService.FindFilesAsync(perforcePath.Key, perforcePath.Select(x => x.AbsolutePath));
+				foreach (FileSummary file in files)
 				{
-					Uri FileUri = new Uri($"{PerforceScheme}://{PerforcePath.Key}{File.DepotPath}");
-					if (File.Error == null)
+					Uri fileUri = new Uri($"{PerforceScheme}://{perforcePath.Key}{file.DepotPath}");
+					if (file.Error == null)
 					{
-						Revisions[FileUri] = $"ver={Version},chg={File.Change},path={FileUri}";
+						revisions[fileUri] = $"ver={Version},chg={file.Change},path={fileUri}";
 					}
 					else
 					{
-						NotificationService.NotifyConfigUpdateFailure(File.Error, File.DepotPath);
+						_notificationService.NotifyConfigUpdateFailure(file.Error, file.DepotPath);
 					}
 				}
 			}
 
-			return Revisions;
+			return revisions;
 		}
 
-		async Task<T> ReadDataAsync<T>(Uri ConfigPath) where T : class
+		async Task<T> ReadDataAsync<T>(Uri configPath) where T : class
 		{
-			byte[] Data = await ReadDataAsync(ConfigPath);
+			byte[] data = await ReadDataAsync(configPath);
 
-			JsonSerializerOptions Options = new JsonSerializerOptions();
-			Startup.ConfigureJsonSerializer(Options);
+			JsonSerializerOptions options = new JsonSerializerOptions();
+			Startup.ConfigureJsonSerializer(options);
 
-			return JsonSerializer.Deserialize<T>(Data, Options)!;
+			return JsonSerializer.Deserialize<T>(data, options)!;
 		}
 
-		Task<byte[]> ReadDataAsync(Uri ConfigPath)
+		Task<byte[]> ReadDataAsync(Uri configPath)
 		{
-			switch(ConfigPath.Scheme)
+			switch(configPath.Scheme)
 			{
 				case FileScheme:
-					return File.ReadAllBytesAsync(ConfigPath.LocalPath);
+					return File.ReadAllBytesAsync(configPath.LocalPath);
 				case PerforceScheme:
-					return PerforceService.PrintAsync(ConfigPath.Host, ConfigPath.AbsolutePath);
+					return _perforceService.PrintAsync(configPath.Host, configPath.AbsolutePath);
 				default:
-					throw new Exception($"Invalid config path: {ConfigPath}");
+					throw new Exception($"Invalid config path: {configPath}");
 			}
 		}
 
-		async Task SendFailureNotificationAsync(Exception Ex, Uri ConfigPath)
+		async Task SendFailureNotificationAsync(Exception ex, Uri configPath)
 		{
-			Logger.LogError(Ex, "Unable to read data from {ConfigPath}: {Message}", ConfigPath, Ex.Message);
+			_logger.LogError(ex, "Unable to read data from {ConfigPath}: {Message}", configPath, ex.Message);
 
-			string FileName = ConfigPath.AbsolutePath;
-			int Change = -1;
-			IUser? Author = null;
-			string? Description = null;
+			string fileName = configPath.AbsolutePath;
+			int change = -1;
+			IUser? author = null;
+			string? description = null;
 
-			if (ConfigPath.Scheme == PerforceScheme)
+			if (configPath.Scheme == PerforceScheme)
 			{
 				try
 				{
-					List<FileSummary> Files = await PerforceService.FindFilesAsync(ConfigPath.Host, new[] { FileName });
-					Change = Files[0].Change;
+					List<FileSummary> files = await _perforceService.FindFilesAsync(configPath.Host, new[] { fileName });
+					change = files[0].Change;
 
-					List<ChangeSummary> Changes = await PerforceService.GetChangesAsync(ConfigPath.Host, Change, Change, 1);
-					if (Changes.Count > 0 && Changes[0].Number == Change)
+					List<ChangeSummary> changes = await _perforceService.GetChangesAsync(configPath.Host, change, change, 1);
+					if (changes.Count > 0 && changes[0].Number == change)
 					{
-						(Author, Description) = (Changes[0].Author, Changes[0].Description);
+						(author, description) = (changes[0].Author, changes[0].Description);
 					}
 				}
-				catch (Exception Ex2)
+				catch (Exception ex2)
 				{
-					Logger.LogError(Ex2, "Unable to identify change that last modified {ConfigPath} from Perforce", ConfigPath);
+					_logger.LogError(ex2, "Unable to identify change that last modified {ConfigPath} from Perforce", configPath);
 				}
 			}
 
-			NotificationService.NotifyConfigUpdateFailure(Ex.Message, FileName, Change, Author, Description);
+			_notificationService.NotifyConfigUpdateFailure(ex.Message, fileName, change, author, description);
 		}
 
 		/// <inheritdoc/>
-		async ValueTask TickLeaderAsync(CancellationToken StoppingToken)
+		async ValueTask TickLeaderAsync(CancellationToken stoppingToken)
 		{
 
-			Uri? ConfigUri = null;
+			Uri? configUri = null;
 			
-			if (Path.IsPathRooted(Settings.CurrentValue.ConfigPath) && !Settings.CurrentValue.ConfigPath.StartsWith("//", StringComparison.Ordinal))
+			if (Path.IsPathRooted(_settings.CurrentValue.ConfigPath) && !_settings.CurrentValue.ConfigPath.StartsWith("//", StringComparison.Ordinal))
 			{
 				// absolute path to config
-				ConfigUri = new Uri(Settings.CurrentValue.ConfigPath);					
+				configUri = new Uri(_settings.CurrentValue.ConfigPath);					
 			}
-			else if (Settings.CurrentValue.ConfigPath != null)
+			else if (_settings.CurrentValue.ConfigPath != null)
 			{
 				// relative (development) or perforce path
-				ConfigUri = CombinePaths(new Uri(FileReference.Combine(Program.AppDir, "_").FullName), Settings.CurrentValue.ConfigPath);				
+				configUri = CombinePaths(new Uri(FileReference.Combine(Program.AppDir, "_").FullName), _settings.CurrentValue.ConfigPath);				
 			}
 
-			if (ConfigUri != null)
+			if (configUri != null)
 			{
-				await UpdateConfigAsync(ConfigUri);
+				await UpdateConfigAsync(configUri);
 			}
 		}
 
@@ -457,79 +446,78 @@ namespace Horde.Build.Services
 		/// <summary>
 		/// Update the global configuration
 		/// </summary>
-		/// <param name="Request"></param>
+		/// <param name="request"></param>
 		/// <returns></returns>
-		public async Task<bool> UpdateGlobalConfig(UpdateGlobalConfigRequest Request)
+		public async Task<bool> UpdateGlobalConfig(UpdateGlobalConfigRequest request)
 		{
-			if (Settings.CurrentValue.ConfigPath == null || !Path.IsPathRooted(Settings.CurrentValue.ConfigPath) || Settings.CurrentValue.ConfigPath.StartsWith("//", StringComparison.Ordinal))
+			if (_settings.CurrentValue.ConfigPath == null || !Path.IsPathRooted(_settings.CurrentValue.ConfigPath) || _settings.CurrentValue.ConfigPath.StartsWith("//", StringComparison.Ordinal))
 			{
-				throw new Exception($"Global config path must be rooted for service updates.  (Perforce paths are not currently supported for live updates), {Settings.CurrentValue.ConfigPath}");
+				throw new Exception($"Global config path must be rooted for service updates.  (Perforce paths are not currently supported for live updates), {_settings.CurrentValue.ConfigPath}");
 			}
 			
-			FileReference GlobalConfigFile = new FileReference(Settings.CurrentValue.ConfigPath);
-			DirectoryReference GlobalConfigDirectory = GlobalConfigFile.Directory;
+			FileReference globalConfigFile = new FileReference(_settings.CurrentValue.ConfigPath);
+			DirectoryReference globalConfigDirectory = globalConfigFile.Directory;
 
 			// make sure the directory exists
-			if (!DirectoryReference.Exists(GlobalConfigDirectory))
+			if (!DirectoryReference.Exists(globalConfigDirectory))
 			{
-				DirectoryReference.CreateDirectory(GlobalConfigDirectory);
+				DirectoryReference.CreateDirectory(globalConfigDirectory);
 			}
 
-			bool ConfigDirty = false;
+			bool configDirty = false;
 
 			// projects
-			if (Request.ProjectsJson != null)
+			if (request.ProjectsJson != null)
 			{
-				ConfigDirty = true;
+				configDirty = true;
 
 				// write out the projects
-				foreach (KeyValuePair<string, string> Project in Request.ProjectsJson)
+				foreach (KeyValuePair<string, string> project in request.ProjectsJson)
 				{
-					await FileReference.WriteAllTextAsync(FileReference.Combine(GlobalConfigDirectory, Project.Key), Project.Value);
+					await FileReference.WriteAllTextAsync(FileReference.Combine(globalConfigDirectory, project.Key), project.Value);
 
-					if (Request.ProjectLogo != null)
+					if (request.ProjectLogo != null)
 					{
-						Byte[] bytes = Convert.FromBase64String(Request.ProjectLogo);
-						await FileReference.WriteAllBytesAsync(FileReference.Combine(GlobalConfigDirectory, Project.Key.Replace(".json", ".png", StringComparison.OrdinalIgnoreCase)), bytes);
+						byte[] bytes = Convert.FromBase64String(request.ProjectLogo);
+						await FileReference.WriteAllBytesAsync(FileReference.Combine(globalConfigDirectory, project.Key.Replace(".json", ".png", StringComparison.OrdinalIgnoreCase)), bytes);
 					}
 				}
-
 			}
 
 			// streams
-			if (Request.StreamsJson != null)
+			if (request.StreamsJson != null)
 			{
-				ConfigDirty = true;
+				configDirty = true;
 
 				// write out the streams
-				foreach (KeyValuePair<string, string> Stream in Request.StreamsJson)
+				foreach (KeyValuePair<string, string> stream in request.StreamsJson)
 				{
-					await FileReference.WriteAllTextAsync(FileReference.Combine(GlobalConfigDirectory, Stream.Key), Stream.Value);
+					await FileReference.WriteAllTextAsync(FileReference.Combine(globalConfigDirectory, stream.Key), stream.Value);
 				}
 			}
 
 			// update global config path
-			if (!string.IsNullOrEmpty(Request.GlobalsJson))
+			if (!String.IsNullOrEmpty(request.GlobalsJson))
 			{
-				ConfigDirty = true;
+				configDirty = true;
 
-				await FileReference.WriteAllTextAsync(GlobalConfigFile, Request.GlobalsJson);
+				await FileReference.WriteAllTextAsync(globalConfigFile, request.GlobalsJson);
 			}
 
-			if (ConfigDirty == true)
+			if (configDirty == true)
 			{				
-				await UpdateConfigAsync(new Uri("file://" + GlobalConfigFile.ToString()));				
+				await UpdateConfigAsync(new Uri("file://" + globalConfigFile.ToString()));				
 			}
 
 			// create the default pool 
-			if (Request.DefaultPoolName != null)
+			if (request.DefaultPoolName != null)
 			{
-				PoolId PoolIdValue = PoolId.Sanitize(Request.DefaultPoolName);
+				PoolId poolIdValue = PoolId.Sanitize(request.DefaultPoolName);
 
-				IPool? Pool = await PoolService.GetPoolAsync(PoolIdValue);
-				if (Pool == null)
+				IPool? pool = await _poolService.GetPoolAsync(poolIdValue);
+				if (pool == null)
 				{
-					await PoolService.CreatePoolAsync(Request.DefaultPoolName, Condition: "OSFamily == 'Windows'", Properties: new Dictionary<string, string>() { ["Color"] = "0" });
+					await _poolService.CreatePoolAsync(request.DefaultPoolName, condition: "OSFamily == 'Windows'", properties: new Dictionary<string, string>() { ["Color"] = "0" });
 				}
 			}
 
@@ -540,136 +528,135 @@ namespace Horde.Build.Services
 		/// <summary>
 		/// Update the server settings
 		/// </summary>
-		/// <param name="Request"></param>
+		/// <param name="request"></param>
 		/// <returns></returns>
-		public async Task<ServerUpdateResponse> UpdateServerSettings(UpdateServerSettingsRequest Request)
+		public async Task<ServerUpdateResponse> UpdateServerSettings(UpdateServerSettingsRequest request)
 		{
 
-			ServerUpdateResponse Response = new ServerUpdateResponse();
+			ServerUpdateResponse response = new ServerUpdateResponse();
 
-			Dictionary<string, object> NewUserSettings = new Dictionary<string, object>();
+			Dictionary<string, object> newUserSettings = new Dictionary<string, object>();
 
 			try
 			{
 
-				if (Request.Settings == null || Request.Settings.Count == 0)
+				if (request.Settings == null || request.Settings.Count == 0)
 				{
-					return Response;
+					return response;
 				}
 
-				if (UserConfigUpdated != null)
+				if (_userConfigUpdated != null)
 				{
-					Response.Errors.Add("User configuation already being updated");
-					return Response;
+					response.Errors.Add("User configuation already being updated");
+					return response;
 				}
 
 				// Load the current user configuration
-				IConfiguration Config = new ConfigurationBuilder()
+				IConfiguration config = new ConfigurationBuilder()
 					.AddJsonFile(Program.UserConfigFile.FullName, optional: true)
 					.Build();
 
-				ServerSettings UserSettings = new ServerSettings();
-				Config.GetSection("Horde").Bind(UserSettings);
+				ServerSettings userSettings = new ServerSettings();
+				config.GetSection("Horde").Bind(userSettings);
 
 				// Figure out current and new settings
-				HashSet<string> UserProps = new HashSet<string>(StringComparer.OrdinalIgnoreCase);
+				HashSet<string> userProps = new HashSet<string>(StringComparer.OrdinalIgnoreCase);
 
 				if (FileReference.Exists(Program.UserConfigFile))
 				{
-					byte[] Data = await FileReference.ReadAllBytesAsync(Program.UserConfigFile);
+					byte[] data = await FileReference.ReadAllBytesAsync(Program.UserConfigFile);
 
-					JsonDocument Document = JsonDocument.Parse(Data);
-					JsonElement HordeElement;
+					JsonDocument document = JsonDocument.Parse(data);
+					JsonElement hordeElement;
 
-					if (Document.RootElement.TryGetProperty("Horde", out HordeElement))
+					if (document.RootElement.TryGetProperty("Horde", out hordeElement))
 					{
-						foreach (JsonProperty Property in HordeElement.EnumerateObject())
+						foreach (JsonProperty property in hordeElement.EnumerateObject())
 						{
-							UserProps.Add(Property.Name);
+							userProps.Add(property.Name);
 						}
 					}
 				}
 
-				foreach (string Name in Request.Settings.Keys)
+				foreach (string name in request.Settings.Keys)
 				{
-					UserProps.Add(Name);
+					userProps.Add(name);
 				}
 
 				// Apply the changes from the request
-				foreach (KeyValuePair<string, object> Pair in Request.Settings)
+				foreach (KeyValuePair<string, object> pair in request.Settings)
 				{
-					PropertyInfo? Property = UserSettings.GetType().GetProperty(Pair.Key, BindingFlags.IgnoreCase | BindingFlags.Public | BindingFlags.Instance);
-					if (Property == null)
+					PropertyInfo? property = userSettings.GetType().GetProperty(pair.Key, BindingFlags.IgnoreCase | BindingFlags.Public | BindingFlags.Instance);
+					if (property == null)
 					{
-						Response.Errors.Add($"Horde configuration property {Pair.Key} does not exist when reading server settings");
-						Logger.LogError("Horde configuration property {Key} does not exist when reading server settings", Pair.Key);
+						response.Errors.Add($"Horde configuration property {pair.Key} does not exist when reading server settings");
+						_logger.LogError("Horde configuration property {Key} does not exist when reading server settings", pair.Key);
 						continue;
 					}
 
-					if (Property.SetMethod == null)
+					if (property.SetMethod == null)
 					{
-						Response.Errors.Add($"Horde configuration property {Pair.Key} does not have a set method");
-						Logger.LogError("Horde configuration property {Key} does not have a set method", Pair.Key);
+						response.Errors.Add($"Horde configuration property {pair.Key} does not have a set method");
+						_logger.LogError("Horde configuration property {Key} does not have a set method", pair.Key);
 						continue;
 					}
 
 					// explicit setting removal
-					if (Pair.Value == null)
+					if (pair.Value == null)
 					{
-						UserProps.Remove(Pair.Key);
+						userProps.Remove(pair.Key);
 						continue;
 					}					
 
+					JsonElement element = (JsonElement)pair.Value;
 
-					JsonElement Element = (JsonElement)Pair.Value;
+					object? value = null;
 
-					object? Value = null;
-
-					switch (Element.ValueKind)
+					switch (element.ValueKind)
 					{
 						case JsonValueKind.True:
-							Value = true;
+							value = true;
 							break;
 						case JsonValueKind.False:
-							Value = false;
+							value = false;
 							break;
 						case JsonValueKind.Number:
-							Value = Element.GetDouble();
+							value = element.GetDouble();
 							break;
 						case JsonValueKind.String:
-							Value = Element.GetString();
+							value = element.GetString();
 							break;
 					}
 
 					// unable to map type
-					if (Value == null)
+					if (value == null)
 					{
-						Response.Errors.Add($"Unable to map type for Property {Property.Name}");
+						response.Errors.Add($"Unable to map type for Property {property.Name}");
 						continue;
 					}					
 
 					// handle common conversions
-					if (Property.GetType() != Value.GetType())
+					if (property.GetType() != value.GetType())
 					{
 						try
 						{
-							Value =	Convert.ChangeType(Value, Property.PropertyType, CultureInfo.CurrentCulture);						
+							value =	Convert.ChangeType(value, property.PropertyType, CultureInfo.CurrentCulture);						
 						}
-						catch (Exception Ex)
+						catch (Exception ex)
 						{
-							Response.Errors.Add($"Property {Property.Name} raised exception during conversion, {Ex.Message}");
+							response.Errors.Add($"Property {property.Name} raised exception during conversion, {ex.Message}");
 							continue;							
 						}
 						
-						if (Value == null)
+						if (value == null)
 						{
-							Response.Errors.Add($"Property {Property.Name} had null value on conversion");
+							response.Errors.Add($"Property {property.Name} had null value on conversion");
 							continue;							
 						}
 					}
 					else 
 					{
-						Response.Errors.Add($"Property {Property.Name} is not assignable to {Value.ToString()}");
+						response.Errors.Add($"Property {property.Name} is not assignable to {value}");
 						continue;
 					}
 					
@@ -677,101 +664,99 @@ namespace Horde.Build.Services
 					//  Set the value, providing some validation
 					try
 					{
-						Property.SetMethod.Invoke(UserSettings, new object[] { Value });
+						property.SetMethod.Invoke(userSettings, new object[] { value });
 					}
-					catch (Exception Ex)
+					catch (Exception ex)
 					{
-						Response.Errors.Add($"Exception updating property {Pair.Key}, {Ex.Message}");
+						response.Errors.Add($"Exception updating property {pair.Key}, {ex.Message}");
 					}
-					
 				}
 
 				// Construct the new user settings
-				foreach (string Name in UserProps)
+				foreach (string name in userProps)
 				{
-					PropertyInfo? Property = UserSettings.GetType().GetProperty(Name, BindingFlags.IgnoreCase | BindingFlags.Public | BindingFlags.Instance);
-					if (Property == null)
+					PropertyInfo? property = userSettings.GetType().GetProperty(name, BindingFlags.IgnoreCase | BindingFlags.Public | BindingFlags.Instance);
+					if (property == null)
 					{
-						Response.Errors.Add($"Horde configuration property {Name} does not exist when writing server settings");
-						Logger.LogError("Horde configuration property {Name} does not exist when writing server settings", Name);
+						response.Errors.Add($"Horde configuration property {name} does not exist when writing server settings");
+						_logger.LogError("Horde configuration property {Name} does not exist when writing server settings", name);
 
 						continue;
 					}
 
-					if (Property.GetMethod == null)
+					if (property.GetMethod == null)
 					{
-						Response.Errors.Add($"Horde configuration property {Name} does not have a get method");
-						Logger.LogError("Horde configuration property {Name} does not have a get method", Name);
+						response.Errors.Add($"Horde configuration property {name} does not have a get method");
+						_logger.LogError("Horde configuration property {Name} does not have a get method", name);
 						continue;
 					}
 
-					object? Result = Property.GetMethod.Invoke(UserSettings, null);
+					object? result = property.GetMethod.Invoke(userSettings, null);
 
-					if (Result == null)
+					if (result == null)
 					{
-						Response.Errors.Add($"Horde configuration property {Name} was null while writing and should have already been filtered out");
-						Logger.LogError("Horde configuration property {Name} was null while writing and should have already been filtered out", Name);
+						response.Errors.Add($"Horde configuration property {name} was null while writing and should have already been filtered out");
+						_logger.LogError("Horde configuration property {Name} was null while writing and should have already been filtered out", name);
 						continue;
 					}
 
-					NewUserSettings.Add(Property.Name, Result);
+					newUserSettings.Add(property.Name, result);
 
 				}
 			}
-			catch (Exception Ex)
+			catch (Exception ex)
 			{
-				Response.Errors.Add($"Exception while updating settings: {Ex.Message}");
-				Logger.LogError(Ex, "{Error}", Response.Errors.Last());
+				response.Errors.Add($"Exception while updating settings: {ex.Message}");
+				_logger.LogError(ex, "{Error}", response.Errors.Last());
 			}
 
-			if (Response.Errors.Count != 0)
+			if (response.Errors.Count != 0)
 			{
-				return Response;
+				return response;
 			}
 
-			Dictionary<string, object> NewLocalSettings = new Dictionary<string, object>();
-			NewLocalSettings["Horde"] = NewUserSettings;
+			Dictionary<string, object> newLocalSettings = new Dictionary<string, object>();
+			newLocalSettings["Horde"] = newUserSettings;
 
 			try
 			{
-				UserConfigUpdated = new TaskCompletionSource<bool>();
+				_userConfigUpdated = new TaskCompletionSource<bool>();
 
 				// This will trigger a setting update as the user config json is set to reload on change
 				try
 				{
-					await FileReference.WriteAllBytesAsync(Program.UserConfigFile, JsonSerializer.SerializeToUtf8Bytes(NewLocalSettings, new JsonSerializerOptions { WriteIndented = true, DefaultIgnoreCondition = JsonIgnoreCondition.WhenWritingNull }));
+					await FileReference.WriteAllBytesAsync(Program.UserConfigFile, JsonSerializer.SerializeToUtf8Bytes(newLocalSettings, new JsonSerializerOptions { WriteIndented = true, DefaultIgnoreCondition = JsonIgnoreCondition.WhenWritingNull }));
 				}
-				catch (Exception Ex)
+				catch (Exception ex)
 				{
-					Response.Errors.Add($"Unable to serialize json settings to {Program.UserConfigFile.ToString()}, {Ex.Message}");
-					Logger.LogError(Ex, "Unable to serialize json settings to {ConfigFile}, {Message}", Program.UserConfigFile.ToString(), Ex.Message);
+					response.Errors.Add($"Unable to serialize json settings to {Program.UserConfigFile}, {ex.Message}");
+					_logger.LogError(ex, "Unable to serialize json settings to {ConfigFile}, {Message}", Program.UserConfigFile.ToString(), ex.Message);
 				}
 
-				if (Response.Errors.Count == 0)
+				if (response.Errors.Count == 0)
 				{
-					if (await Task.WhenAny(UserConfigUpdated.Task, Task.Delay(5000)) != UserConfigUpdated.Task)
+					if (await Task.WhenAny(_userConfigUpdated.Task, Task.Delay(5000)) != _userConfigUpdated.Task)
 					{
-						Response.Errors.Add("Server update timed out while awaiting write");
-						Logger.LogError("Server update timed out after writing config");
+						response.Errors.Add("Server update timed out while awaiting write");
+						_logger.LogError("Server update timed out after writing config");
 					}
 				}
-
 			}
 			finally
 			{
-				UserConfigUpdated = null;
+				_userConfigUpdated = null;
 			}
 
-			return Response;
+			return response;
 
 		}
-		void OnUserConfigUpdated(ServerSettings Settings, string Name)
+		void OnUserConfigUpdated(ServerSettings settings, string name)
 		{
 			NumUserConfigUpdates++;
-			UserConfigUpdated?.SetResult(true);
+			_userConfigUpdated?.SetResult(true);
 		}
 
-		private TaskCompletionSource<bool>? UserConfigUpdated;
+		private TaskCompletionSource<bool>? _userConfigUpdated;
 
 		/// <summary>
 		/// The number of server configuration updates that have been made while the server is running
diff --git a/Engine/Source/Programs/Horde/Horde.Build/Server/DatabaseService.cs b/Engine/Source/Programs/Horde/Horde.Build/Server/DatabaseService.cs
--- a/Engine/Source/Programs/Horde/Horde.Build/Server/DatabaseService.cs
+++ b/Engine/Source/Programs/Horde/Horde.Build/Server/DatabaseService.cs
@@ -1,47 +1,33 @@
 // Copyright Epic Games, Inc. All Rights Reserved.
 
-using EpicGames.Core;
-using Horde.Build.Api;
-using HordeCommon;
-using Horde.Build.Models;
-using Horde.Build.Utilities;
-using Microsoft.Extensions.Logging;
-using Microsoft.Extensions.Options;
-using Microsoft.IdentityModel.Tokens;
-using MongoDB.Bson;
-using MongoDB.Bson.Serialization;
-using MongoDB.Driver;
 using System;
-using System.Collections.Concurrent;
 using System.Collections.Generic;
+using System.Diagnostics;
+using System.Globalization;
 using System.IO;
 using System.Linq;
-using System.Linq.Expressions;
 using System.Net;
+using System.Net.NetworkInformation;
 using System.Net.Security;
+using System.Net.Sockets;
 using System.Runtime.InteropServices;
 using System.Security.Cryptography.X509Certificates;
 using System.Text;
-using System.Threading.Tasks;
-using System.Net.Sockets;
-using MongoDB.Driver.Core.Events;
-using System.Threading;
-using Horde.Build.Collections;
-using System.Diagnostics;
 using System.Text.RegularExpressions;
-using System.Net.NetworkInformation;
-using OpenTracing.Util;
-using OpenTracing;
+using System.Threading.Tasks;
+using EpicGames.Core;
+using Horde.Build.Models;
 using Horde.Build.Utiltiies;
-using System.Globalization;
+using Microsoft.Extensions.Logging;
+using Microsoft.Extensions.Options;
+using Microsoft.IdentityModel.Tokens;
+using MongoDB.Bson;
+using MongoDB.Bson.Serialization;
+using MongoDB.Driver;
+using MongoDB.Driver.Core.Events;
 
 namespace Horde.Build.Services
 {
-	using IStream = Horde.Build.Models.IStream;
-	using PoolId = StringId<IPool>;
-	using ProjectId = StringId<IProject>;
-	using TemplateRefId = StringId<TemplateRef>;
-
 	/// <summary>
 	/// Singleton for accessing the database
 	/// </summary>
@@ -80,7 +66,7 @@ namespace Horde.Build.Services
 		/// <summary>
 		/// Logger for this instance
 		/// </summary>
-		ILogger<DatabaseService> Logger;
+		readonly ILogger<DatabaseService> _logger;
 		
 		/// <summary>
 		/// Access the database in a read-only mode (don't create indices or modify content)
@@ -90,22 +76,22 @@ namespace Horde.Build.Services
 		/// <summary>
 		/// The mongo process group
 		/// </summary>
-		ManagedProcessGroup? MongoProcessGroup;
+		ManagedProcessGroup? _mongoProcessGroup;
 
 		/// <summary>
 		/// The mongo process
 		/// </summary>
-		ManagedProcess? MongoProcess;
+		ManagedProcess? _mongoProcess;
 
 		/// <summary>
 		/// Task to read from the mongo process
 		/// </summary>
-		Task? MongoOutputTask;
+		Task? _mongoOutputTask;
 
 		/// <summary>
 		/// Factory for creating logger instances
 		/// </summary>
-		ILoggerFactory LoggerFactory;
+		readonly ILoggerFactory _loggerFactory;
 
 		/// <summary>
 		/// Default port for MongoDB connections
@@ -115,48 +101,48 @@ namespace Horde.Build.Services
 		/// <summary>
 		/// Constructor
 		/// </summary>
-		/// <param name="SettingsSnapshot">The settings instance</param>
-		/// <param name="LoggerFactory">Instance of the logger for this service</param>
-		public DatabaseService(IOptions<ServerSettings> SettingsSnapshot, ILoggerFactory LoggerFactory)
+		/// <param name="settingsSnapshot">The settings instance</param>
+		/// <param name="loggerFactory">Instance of the logger for this service</param>
+		public DatabaseService(IOptions<ServerSettings> settingsSnapshot, ILoggerFactory loggerFactory)
 		{
-			this.Settings = SettingsSnapshot.Value;
-			this.Logger = LoggerFactory.CreateLogger<DatabaseService>();
-			this.LoggerFactory = LoggerFactory;
+			Settings = settingsSnapshot.Value;
+			_logger = loggerFactory.CreateLogger<DatabaseService>();
+			_loggerFactory = loggerFactory;
 
 			try
 			{
 				ReadOnlyMode = Settings.DatabaseReadOnlyMode;
 				if (Settings.DatabasePublicCert != null)
 				{
-					X509Store LocalTrustStore = new X509Store(StoreName.Root);
+					X509Store localTrustStore = new X509Store(StoreName.Root);
 					try
 					{
-						LocalTrustStore.Open(OpenFlags.ReadWrite);
+						localTrustStore.Open(OpenFlags.ReadWrite);
 
-						X509Certificate2Collection Collection = ImportCertificateBundle(Settings.DatabasePublicCert);
-						foreach (X509Certificate2 Certificate in Collection)
+						X509Certificate2Collection collection = ImportCertificateBundle(Settings.DatabasePublicCert);
+						foreach (X509Certificate2 certificate in collection)
 						{
-							Logger.LogInformation("Importing certificate for {Subject}", Certificate.Subject);
+							_logger.LogInformation("Importing certificate for {Subject}", certificate.Subject);
 						}
 
-						LocalTrustStore.AddRange(Collection);
+						localTrustStore.AddRange(collection);
 					}
 					finally
 					{
-						LocalTrustStore.Close();
+						localTrustStore.Close();
 					}
 				}
 
-				string? ConnectionString = Settings.DatabaseConnectionString;
-				if (ConnectionString == null)
+				string? connectionString = Settings.DatabaseConnectionString;
+				if (connectionString == null)
 				{
 					if (IsPortInUse(DefaultMongoPort))
 					{
-						ConnectionString = "mongodb://localhost:27017";
+						connectionString = "mongodb://localhost:27017";
 					}
-					else if (TryStartMongoServer(Logger))
+					else if (TryStartMongoServer(_logger))
 					{
-						ConnectionString = "mongodb://localhost:27017/?readPreference=primary&appname=Horde&ssl=false";
+						connectionString = "mongodb://localhost:27017/?readPreference=primary&appname=Horde&ssl=false";
 					}
 					else
 					{
@@ -164,40 +150,40 @@ namespace Horde.Build.Services
 					}
 				}
 
-				MongoClientSettings MongoSettings = MongoClientSettings.FromConnectionString(ConnectionString);
-				MongoSettings.ClusterConfigurator = ClusterBuilder =>
+				MongoClientSettings mongoSettings = MongoClientSettings.FromConnectionString(connectionString);
+				mongoSettings.ClusterConfigurator = clusterBuilder =>
 				{
-					if (Logger.IsEnabled(LogLevel.Trace))
+					if (_logger.IsEnabled(LogLevel.Trace))
 					{
-						ClusterBuilder.Subscribe<CommandStartedEvent>(Event => TraceMongoCommand(Event.Command));
+						clusterBuilder.Subscribe<CommandStartedEvent>(ev => TraceMongoCommand(ev.Command));
 					}
 				};
 				
-				MongoSettings.SslSettings = new SslSettings();
-				MongoSettings.SslSettings.ServerCertificateValidationCallback = CertificateValidationCallBack;
-				MongoSettings.MaxConnectionPoolSize = 300; // Default is 100
+				mongoSettings.SslSettings = new SslSettings();
+				mongoSettings.SslSettings.ServerCertificateValidationCallback = CertificateValidationCallBack;
+				mongoSettings.MaxConnectionPoolSize = 300; // Default is 100
 
 				//TestSslConnection(MongoSettings.Server.Host, MongoSettings.Server.Port, Logger);
 
-				MongoClient Client = new MongoClient(MongoSettings);
-				Database = Client.GetDatabase(Settings.DatabaseName);
+				MongoClient client = new MongoClient(mongoSettings);
+				Database = client.GetDatabase(Settings.DatabaseName);
 
 				Singletons = GetCollection<BsonDocument>("Singletons");
 
-				Globals Globals = GetGlobalsAsync().Result;
-				while (Globals.JwtSigningKey == null)
+				Globals globals = GetGlobalsAsync().Result;
+				while (globals.JwtSigningKey == null)
 				{
-					Globals.RotateSigningKey();
-					if (!TryUpdateSingletonAsync(Globals).Result)
+					globals.RotateSigningKey();
+					if (!TryUpdateSingletonAsync(globals).Result)
 					{
-						Globals = GetGlobalsAsync().Result;
+						globals = GetGlobalsAsync().Result;
 					}
 				}
 
 				JwtIssuer = Settings.JwtIssuer ?? Dns.GetHostName();
 				if (String.IsNullOrEmpty(Settings.JwtSecret))
 				{
-					JwtSigningKey = new SymmetricSecurityKey(Globals.JwtSigningKey);
+					JwtSigningKey = new SymmetricSecurityKey(globals.JwtSigningKey);
 				}
 				else
 				{
@@ -206,50 +192,50 @@ namespace Horde.Build.Services
 
 				Credentials = GetCollection<Credential>("Credentials");
 			}
-			catch (Exception Ex)
+			catch (Exception ex)
 			{
-				Logger.LogError(Ex, "Exception while initializing DatabaseService");
+				_logger.LogError(ex, "Exception while initializing DatabaseService");
 				throw;
 			}
 		}
 
-		internal const int CTRL_C_EVENT = 0;
+		internal const int CtrlCEvent = 0;
 
 		[DllImport("kernel32.dll")]
-		internal static extern bool GenerateConsoleCtrlEvent(int Event, int ProcessGroupId);
+		internal static extern bool GenerateConsoleCtrlEvent(int eventId, int processGroupId);
 
 		/// <inheritdoc/>
 		public void Dispose()
 		{
-			if (MongoProcess != null)
+			if (_mongoProcess != null)
 			{
-				GenerateConsoleCtrlEvent(CTRL_C_EVENT, MongoProcess.Id);
+				GenerateConsoleCtrlEvent(CtrlCEvent, _mongoProcess.Id);
 
-				MongoOutputTask?.Wait();
-				MongoOutputTask = null;
+				_mongoOutputTask?.Wait();
+				_mongoOutputTask = null;
 
-				MongoProcess.WaitForExit();
-				MongoProcess.Dispose();
-				MongoProcess = null;
+				_mongoProcess.WaitForExit();
+				_mongoProcess.Dispose();
+				_mongoProcess = null;
 			}
-			if(MongoProcessGroup != null)
+			if(_mongoProcessGroup != null)
 			{
-				MongoProcessGroup.Dispose();
-				MongoProcessGroup = null;
+				_mongoProcessGroup.Dispose();
+				_mongoProcessGroup = null;
 			}
 		}
 
 		/// <summary>
 		/// Checks if the given port is in use
 		/// </summary>
-		/// <param name="PortNumber"></param>
+		/// <param name="portNumber"></param>
 		/// <returns></returns>
-		static bool IsPortInUse(int PortNumber)
+		static bool IsPortInUse(int portNumber)
 		{
-			IPGlobalProperties IpGlobalProperties = IPGlobalProperties.GetIPGlobalProperties();
+			IPGlobalProperties ipGlobalProperties = IPGlobalProperties.GetIPGlobalProperties();
 
-			IPEndPoint[] Listeners = IpGlobalProperties.GetActiveTcpListeners();
-			if (Listeners.Any(x => x.Port == PortNumber))
+			IPEndPoint[] listeners = ipGlobalProperties.GetActiveTcpListeners();
+			if (listeners.Any(x => x.Port == portNumber))
 			{
 				return true;
 			}
@@ -260,60 +246,60 @@ namespace Horde.Build.Services
 		/// <summary>
 		/// Attempts to start a local instance of MongoDB
 		/// </summary>
-		/// <param name="Logger"></param>
+		/// <param name="logger"></param>
 		/// <returns></returns>
-		bool TryStartMongoServer(ILogger Logger)
+		bool TryStartMongoServer(ILogger logger)
 		{
 			if (!RuntimeInformation.IsOSPlatform(OSPlatform.Windows))
 			{
 				return false;
 			}
 
-			FileReference MongoExe = FileReference.Combine(Program.AppDir, "ThirdParty", "Mongo", "mongod.exe");
-			if (!FileReference.Exists(MongoExe))
+			FileReference mongoExe = FileReference.Combine(Program.AppDir, "ThirdParty", "Mongo", "mongod.exe");
+			if (!FileReference.Exists(mongoExe))
 			{
-				Logger.LogWarning("Unable to find Mongo executable.");
+				logger.LogWarning("Unable to find Mongo executable.");
 				return false;
 			}
 
-			DirectoryReference MongoDir = DirectoryReference.Combine(Program.DataDir, "Mongo");
+			DirectoryReference mongoDir = DirectoryReference.Combine(Program.DataDir, "Mongo");
 
-			DirectoryReference MongoDataDir = DirectoryReference.Combine(MongoDir, "Data");
-			DirectoryReference.CreateDirectory(MongoDataDir);
+			DirectoryReference mongoDataDir = DirectoryReference.Combine(mongoDir, "Data");
+			DirectoryReference.CreateDirectory(mongoDataDir);
 
-			FileReference MongoLogFile = FileReference.Combine(MongoDir, "mongod.log");
+			FileReference mongoLogFile = FileReference.Combine(mongoDir, "mongod.log");
 
-			FileReference ConfigFile = FileReference.Combine(MongoDir, "mongod.conf");
-			if(!FileReference.Exists(ConfigFile))
+			FileReference configFile = FileReference.Combine(mongoDir, "mongod.conf");
+			if(!FileReference.Exists(configFile))
 			{
-				DirectoryReference.CreateDirectory(ConfigFile.Directory);
-				using (StreamWriter Writer = new StreamWriter(ConfigFile.FullName))
+				DirectoryReference.CreateDirectory(configFile.Directory);
+				using (StreamWriter writer = new StreamWriter(configFile.FullName))
 				{
-					Writer.WriteLine("# mongod.conf");
-					Writer.WriteLine();
-					Writer.WriteLine("# for documentation of all options, see:");
-					Writer.WriteLine("# http://docs.mongodb.org/manual/reference/configuration-options/");
-					Writer.WriteLine();
-					Writer.WriteLine("storage:");
-					Writer.WriteLine("    dbPath: {0}", MongoDataDir.FullName);
-					Writer.WriteLine();
-					Writer.WriteLine("net:");
-					Writer.WriteLine("    port: {0}", DefaultMongoPort);
-					Writer.WriteLine("    bindIp: 127.0.0.1");
+					writer.WriteLine("# mongod.conf");
+					writer.WriteLine();
+					writer.WriteLine("# for documentation of all options, see:");
+					writer.WriteLine("# http://docs.mongodb.org/manual/reference/configuration-options/");
+					writer.WriteLine();
+					writer.WriteLine("storage:");
+					writer.WriteLine("    dbPath: {0}", mongoDataDir.FullName);
+					writer.WriteLine();
+					writer.WriteLine("net:");
+					writer.WriteLine("    port: {0}", DefaultMongoPort);
+					writer.WriteLine("    bindIp: 127.0.0.1");
 				}
 			}
 
-			MongoProcessGroup = new ManagedProcessGroup();
+			_mongoProcessGroup = new ManagedProcessGroup();
 			try
 			{
-				MongoProcess = new ManagedProcess(MongoProcessGroup, MongoExe.FullName, $"--config \"{ConfigFile}\"", null, null, ProcessPriorityClass.Normal);
-				MongoProcess.StdIn.Close();
-				MongoOutputTask = Task.Run(() => RelayMongoOutput());
+				_mongoProcess = new ManagedProcess(_mongoProcessGroup, mongoExe.FullName, $"--config \"{configFile}\"", null, null, ProcessPriorityClass.Normal);
+				_mongoProcess.StdIn.Close();
+				_mongoOutputTask = Task.Run(() => RelayMongoOutput());
 				return true;
 			}
-			catch (Exception Ex)
+			catch (Exception ex)
 			{
-				Logger.LogWarning(Ex, "Unable to start Mongo server process");
+				logger.LogWarning(ex, "Unable to start Mongo server process");
 				return false;
 			}
 		}
@@ -324,45 +310,45 @@ namespace Horde.Build.Services
 		/// <returns></returns>
 		async Task RelayMongoOutput()
 		{
-			ILogger MongoLogger = LoggerFactory.CreateLogger("MongoDB");
+			ILogger mongoLogger = _loggerFactory.CreateLogger("MongoDB");
 
-			Dictionary<string, ILogger> ChannelLoggers = new Dictionary<string, ILogger>();
+			Dictionary<string, ILogger> channelLoggers = new Dictionary<string, ILogger>();
 			for (; ; )
 			{
-				string? Line = await MongoProcess!.ReadLineAsync();
-				if (Line == null)
+				string? line = await _mongoProcess!.ReadLineAsync();
+				if (line == null)
 				{
 					break;
 				}
-				if (Line.Length > 0)
+				if (line.Length > 0)
 				{
-					Match Match = Regex.Match(Line, @"^\s*[^\s]+\s+([A-Z])\s+([^\s]+)\s+(.*)");
-					if (Match.Success)
+					Match match = Regex.Match(line, @"^\s*[^\s]+\s+([A-Z])\s+([^\s]+)\s+(.*)");
+					if (match.Success)
 					{
-						ILogger? ChannelLogger;
-						if (!ChannelLoggers.TryGetValue(Match.Groups[2].Value, out ChannelLogger))
+						ILogger? channelLogger;
+						if (!channelLoggers.TryGetValue(match.Groups[2].Value, out channelLogger))
 						{
-							ChannelLogger = LoggerFactory.CreateLogger($"MongoDB.{Match.Groups[2].Value}");
-							ChannelLoggers.Add(Match.Groups[2].Value, ChannelLogger);
+							channelLogger = _loggerFactory.CreateLogger($"MongoDB.{match.Groups[2].Value}");
+							channelLoggers.Add(match.Groups[2].Value, channelLogger);
 						}
-						ChannelLogger.Log(ParseMongoLogLevel(Match.Groups[1].Value), "{Message}", Match.Groups[3].Value.TrimEnd());
+						channelLogger.Log(ParseMongoLogLevel(match.Groups[1].Value), "{Message}", match.Groups[3].Value.TrimEnd());
 					}
 					else
 					{
-						MongoLogger.Log(LogLevel.Information, "{Message}", Line);
+						mongoLogger.Log(LogLevel.Information, "{Message}", line);
 					}
 				}
 			}
-			MongoLogger.LogInformation("Exit code {ExitCode}", MongoProcess.ExitCode);
+			mongoLogger.LogInformation("Exit code {ExitCode}", _mongoProcess.ExitCode);
 		}
 
-		static LogLevel ParseMongoLogLevel(string Text)
+		static LogLevel ParseMongoLogLevel(string text)
 		{
-			if (Text.Equals("I", StringComparison.Ordinal))
+			if (text.Equals("I", StringComparison.Ordinal))
 			{
 				return LogLevel.Information;
 			}
-			else if (Text.Equals("E", StringComparison.Ordinal))
+			else if (text.Equals("E", StringComparison.Ordinal))
 			{
 				return LogLevel.Error;
 			}
@@ -375,167 +361,165 @@ namespace Horde.Build.Services
 		/// <summary>
 		/// Logs a mongodb command, removing any fields we don't care about
 		/// </summary>
-		/// <param name="Command">The command document</param>
-		void TraceMongoCommand(BsonDocument Command)
+		/// <param name="command">The command document</param>
+		void TraceMongoCommand(BsonDocument command)
 		{
-			List<string> Params = new List<string>();
-			List<string> Values = new List<string>();
+			List<string> names = new List<string>();
+			List<string> values = new List<string>();
 
-			foreach(BsonElement Element in Command)
+			foreach(BsonElement element in command)
 			{
-				if (Element.Value != null && !Element.Name.Equals("$db", StringComparison.Ordinal) && !Element.Name.Equals("lsid", StringComparison.Ordinal))
+				if (element.Value != null && !element.Name.Equals("$db", StringComparison.Ordinal) && !element.Name.Equals("lsid", StringComparison.Ordinal))
 				{
-					Params.Add($"{Element.Name}: {{{Values.Count}}}");
-					Values.Add(Element.Value.ToString()!);
+					names.Add($"{element.Name}: {{{values.Count}}}");
+					values.Add(element.Value.ToString()!);
 				}
 			}
 
 #pragma warning disable CA2254 // Template should be a static expression
-			Logger.LogTrace($"MongoDB: {String.Join(", ", Params)}", Values.ToArray());
+			_logger.LogTrace($"MongoDB: {String.Join(", ", names)}", values.ToArray());
 #pragma warning restore CA2254 // Template should be a static expression
 		}
 
 		/// <summary>
 		/// Imports one or more certificates from a single PEM file
 		/// </summary>
-		/// <param name="FileName">File to import</param>
+		/// <param name="fileName">File to import</param>
 		/// <returns>Collection of certificates</returns>
-		static X509Certificate2Collection ImportCertificateBundle(string FileName)
+		static X509Certificate2Collection ImportCertificateBundle(string fileName)
 		{
-			X509Certificate2Collection Collection = new X509Certificate2Collection();
+			X509Certificate2Collection collection = new X509Certificate2Collection();
 
-			string Text = File.ReadAllText(FileName);
-			for (int Offset = 0; Offset < Text.Length;)
+			string text = File.ReadAllText(fileName);
+			for (int offset = 0; offset < text.Length;)
 			{
-				int NextOffset = Text.IndexOf("-----BEGIN CERTIFICATE-----", Offset + 1, StringComparison.Ordinal);
-				if (NextOffset == -1)
+				int nextOffset = text.IndexOf("-----BEGIN CERTIFICATE-----", offset + 1, StringComparison.Ordinal);
+				if (nextOffset == -1)
 				{
-					NextOffset = Text.Length;
+					nextOffset = text.Length;
 				}
 
-				string CertificateText = Text.Substring(Offset, NextOffset - Offset);
-				Collection.Add(new X509Certificate2(Encoding.UTF8.GetBytes(CertificateText)));
+				string certificateText = text.Substring(offset, nextOffset - offset);
+				collection.Add(new X509Certificate2(Encoding.UTF8.GetBytes(certificateText)));
 
-				Offset = NextOffset;
+				offset = nextOffset;
 			}
 
-			return Collection;
+			return collection;
 		}
 
 		/// <summary>
 		/// Tests connection to the given server
 		/// </summary>
-		/// <param name="Host">Host name</param>
-		/// <param name="Port">Port number to connect on</param>
-		/// <param name="Logger">Logger for diagnostic messages</param>
+		/// <param name="host">Host name</param>
+		/// <param name="port">Port number to connect on</param>
+		/// <param name="logger">Logger for diagnostic messages</param>
 		/// <returns>True if the connection was valid</returns>
-		static bool TestSslConnection(string Host, int Port, ILogger Logger)
+		static bool TestSslConnection(string host, int port, ILogger logger)
 		{
-#pragma warning disable CA1031 // Do not catch general exception types
-			using (TcpClient Client = new TcpClient())
+			using (TcpClient client = new TcpClient())
 			{
 				try
 				{
-					Client.Connect(Host, Port);
-					Logger.LogInformation("Successfully connected to {Host} on port {Port}", Host, Port);
+					client.Connect(host, port);
+					logger.LogInformation("Successfully connected to {Host} on port {Port}", host, port);
 				}
-				catch (Exception Ex)
+				catch (Exception ex)
 				{
-					Logger.LogError(Ex, "Unable to connect to {Host} on port {Port}", Host, Port);
+					logger.LogError(ex, "Unable to connect to {Host} on port {Port}", host, port);
 					return false;
 				}
 
-				using (SslStream Stream = new SslStream(Client.GetStream()))
+				using (SslStream stream = new SslStream(client.GetStream()))
 				{
 					try
 					{
-						Stream.AuthenticateAsClient(Host);
-						Logger.LogInformation("Successfully authenticated host {Host}", Host);
+						stream.AuthenticateAsClient(host);
+						logger.LogInformation("Successfully authenticated host {Host}", host);
 					}
-					catch (Exception Ex)
+					catch (Exception ex)
 					{
-						Logger.LogError(Ex, "Unable to authenticate as client");
+						logger.LogError(ex, "Unable to authenticate as client");
 						return false;
 					}
 				}
 			}
 			return true;
-#pragma warning restore CA1031 // Do not catch general exception types
 		}
 
 		/// <summary>
 		/// Provides additional diagnostic information for SSL certificate validation
 		/// </summary>
-		/// <param name="Sender"></param>
-		/// <param name="Certificate"></param>
-		/// <param name="Chain"></param>
-		/// <param name="SslPolicyErrors"></param>
+		/// <param name="sender"></param>
+		/// <param name="certificate"></param>
+		/// <param name="chain"></param>
+		/// <param name="sslPolicyErrors"></param>
 		/// <returns>True if the certificate is allowed, false otherwise</returns>
-		bool CertificateValidationCallBack(object Sender, X509Certificate? Certificate, X509Chain? Chain, SslPolicyErrors SslPolicyErrors)
+		bool CertificateValidationCallBack(object sender, X509Certificate? certificate, X509Chain? chain, SslPolicyErrors sslPolicyErrors)
 		{
 			// If the certificate is a valid, signed certificate, return true.
-			if (SslPolicyErrors == SslPolicyErrors.None)
+			if (sslPolicyErrors == SslPolicyErrors.None)
 			{
 				return true;
 			}
 
 			// Generate diagnostic information
-			StringBuilder Builder = new StringBuilder();
-			if (Sender != null)
+			StringBuilder builder = new StringBuilder();
+			if (sender != null)
 			{
-				string SenderInfo = StringUtils.Indent(Sender.ToString() ?? String.Empty, "    ");
-				Builder.Append(CultureInfo.InvariantCulture, $"\nSender:\n{SenderInfo}");
+				string senderInfo = StringUtils.Indent(sender.ToString() ?? String.Empty, "    ");
+				builder.Append(CultureInfo.InvariantCulture, $"\nSender:\n{senderInfo}");
 			}
-			if (Certificate != null)
+			if (certificate != null)
 			{
-				Builder.Append(CultureInfo.InvariantCulture, $"\nCertificate: {Certificate.Subject}");
+				builder.Append(CultureInfo.InvariantCulture, $"\nCertificate: {certificate.Subject}");
 			}
-			if (Chain != null)
+			if (chain != null)
 			{
-				if (Chain.ChainStatus != null && Chain.ChainStatus.Length > 0)
+				if (chain.ChainStatus != null && chain.ChainStatus.Length > 0)
 				{
-					Builder.Append("\nChain status:");
-					foreach (X509ChainStatus Status in Chain.ChainStatus)
+					builder.Append("\nChain status:");
+					foreach (X509ChainStatus status in chain.ChainStatus)
 					{
-						Builder.Append(CultureInfo.InvariantCulture, $"\n  {Status.StatusInformation}");
+						builder.Append(CultureInfo.InvariantCulture, $"\n  {status.StatusInformation}");
 					}
 				}
-				if (Chain.ChainElements != null)
+				if (chain.ChainElements != null)
 				{
-					Builder.Append("\nChain elements:");
-					for (int Idx = 0; Idx < Chain.ChainElements.Count; Idx++)
+					builder.Append("\nChain elements:");
+					for (int idx = 0; idx < chain.ChainElements.Count; idx++)
 					{
-						X509ChainElement Element = Chain.ChainElements[Idx];
-						Builder.Append(CultureInfo.InvariantCulture, $"\n  {Idx,4} - Certificate: {Element.Certificate.Subject}");
-						if (Element.ChainElementStatus != null && Element.ChainElementStatus.Length > 0)
+						X509ChainElement element = chain.ChainElements[idx];
+						builder.Append(CultureInfo.InvariantCulture, $"\n  {idx,4} - Certificate: {element.Certificate.Subject}");
+						if (element.ChainElementStatus != null && element.ChainElementStatus.Length > 0)
 						{
-							foreach (X509ChainStatus Status in Element.ChainElementStatus)
+							foreach (X509ChainStatus status in element.ChainElementStatus)
 							{
-								Builder.Append(CultureInfo.InvariantCulture, $"\n         Status: {Status.StatusInformation} ({Status.Status})");
+								builder.Append(CultureInfo.InvariantCulture, $"\n         Status: {status.StatusInformation} ({status.Status})");
 							}
 						}
-						if (!String.IsNullOrEmpty(Element.Information))
+						if (!String.IsNullOrEmpty(element.Information))
 						{
-							Builder.Append(CultureInfo.InvariantCulture, $"\n         Info: {Element.Information}");
+							builder.Append(CultureInfo.InvariantCulture, $"\n         Info: {element.Information}");
 						}
 					}
 				}
 			}
 
 			// Print out additional diagnostic information
-			Logger.LogError("TLS certificate validation failed ({Errors}).{AdditionalInfo}", SslPolicyErrors, StringUtils.Indent(Builder.ToString(), "    "));
+			_logger.LogError("TLS certificate validation failed ({Errors}).{AdditionalInfo}", sslPolicyErrors, StringUtils.Indent(builder.ToString(), "    "));
 			return false;
 		}
 
 		/// <summary>
 		/// Get a MongoDB collection from database
 		/// </summary>
-		/// <param name="Name">Name of collection</param>
+		/// <param name="name">Name of collection</param>
 		/// <typeparam name="T">A MongoDB document</typeparam>
 		/// <returns></returns>
-		public IMongoCollection<T> GetCollection<T>(string Name)
+		public IMongoCollection<T> GetCollection<T>(string name)
 		{
-			return new MongoTracingCollection<T>(Database.GetCollection<T>(Name));
+			return new MongoTracingCollection<T>(Database.GetCollection<T>(name));
 		}
 
 		/// <summary>
@@ -550,37 +534,37 @@ namespace Horde.Build.Services
 		/// <summary>
 		/// Gets a singleton document by id
 		/// </summary>
-		/// <param name="Id">Id of the singleton document</param>
+		/// <param name="id">Id of the singleton document</param>
 		/// <returns>The document</returns>
-		public async Task<T> GetSingletonAsync<T>(ObjectId Id) where T : SingletonBase, new()
+		public async Task<T> GetSingletonAsync<T>(ObjectId id) where T : SingletonBase, new()
 		{
-			T Singleton = await GetSingletonAsync(Id, () => new T());
-			Singleton.PostLoad();
-			return Singleton;
+			T singleton = await GetSingletonAsync(id, () => new T());
+			singleton.PostLoad();
+			return singleton;
 		}
 
 		/// <summary>
 		/// Gets a singleton document by id
 		/// </summary>
-		/// <param name="Id">Id of the singleton document</param>
-		/// <param name="Constructor">Method to use to construct a new object</param>
+		/// <param name="id">Id of the singleton document</param>
+		/// <param name="constructor">Method to use to construct a new object</param>
 		/// <returns>The document</returns>
-		public async Task<T> GetSingletonAsync<T>(ObjectId Id, Func<T> Constructor) where T : SingletonBase, new()
+		public async Task<T> GetSingletonAsync<T>(ObjectId id, Func<T> constructor) where T : SingletonBase, new()
 		{
-			FilterDefinition<BsonDocument> Filter = new BsonDocument(new BsonElement("_id", Id));
+			FilterDefinition<BsonDocument> filter = new BsonDocument(new BsonElement("_id", id));
 			for (; ; )
 			{
-				BsonDocument? Object = await Singletons.Find<BsonDocument>(Filter).FirstOrDefaultAsync();
-				if (Object != null)
+				BsonDocument? document = await Singletons.Find<BsonDocument>(filter).FirstOrDefaultAsync();
+				if (document != null)
 				{
-					T Item = BsonSerializer.Deserialize<T>(Object);
-					Item.PostLoad();
-					return Item;
+					T item = BsonSerializer.Deserialize<T>(document);
+					item.PostLoad();
+					return item;
 				}
 
-				T NewItem = Constructor();
-				NewItem.Id = Id;
-				await Singletons.InsertOneAsync(NewItem.ToBsonDocument());
+				T newItem = constructor();
+				newItem.Id = id;
+				await Singletons.InsertOneAsync(newItem.ToBsonDocument());
 			}
 		}
 
@@ -588,17 +572,17 @@ namespace Horde.Build.Services
 		/// Updates a singleton
 		/// </summary>
 		/// <typeparam name="T"></typeparam>
-		/// <param name="Id"></param>
-		/// <param name="Updater"></param>
+		/// <param name="id"></param>
+		/// <param name="updater"></param>
 		/// <returns></returns>
-		public async Task UpdateSingletonAsync<T>(ObjectId Id, Action<T> Updater) where T : SingletonBase, new()
+		public async Task UpdateSingletonAsync<T>(ObjectId id, Action<T> updater) where T : SingletonBase, new()
 		{
 			for (; ; )
 			{
-				T Object = await GetSingletonAsync(Id, () => new T());
-				Updater(Object);
+				T document = await GetSingletonAsync(id, () => new T());
+				updater(document);
 				
-				if (await TryUpdateSingletonAsync(Object))
+				if (await TryUpdateSingletonAsync(document))
 				{
 					break;
 				}
@@ -608,22 +592,22 @@ namespace Horde.Build.Services
 		/// <summary>
 		/// Attempts to update a singleton object
 		/// </summary>
-		/// <param name="SingletonObject">The singleton object</param>
+		/// <param name="singletonObject">The singleton object</param>
 		/// <returns>True if the singleton document was updated</returns>
-		public async Task<bool> TryUpdateSingletonAsync<T>(T SingletonObject) where T : SingletonBase
+		public async Task<bool> TryUpdateSingletonAsync<T>(T singletonObject) where T : SingletonBase
 		{
-			int PrevRevision = SingletonObject.Revision++;
+			int prevRevision = singletonObject.Revision++;
 
-			BsonDocument Filter = new BsonDocument { new BsonElement("_id", SingletonObject.Id), new BsonElement(nameof(SingletonBase.Revision), PrevRevision) };
+			BsonDocument filter = new BsonDocument { new BsonElement("_id", singletonObject.Id), new BsonElement(nameof(SingletonBase.Revision), prevRevision) };
 			try
 			{
-				ReplaceOneResult Result = await Singletons.ReplaceOneAsync(Filter, SingletonObject.ToBsonDocument(), new ReplaceOptions { IsUpsert = true });
-				return Result.MatchedCount > 0;
+				ReplaceOneResult result = await Singletons.ReplaceOneAsync(filter, singletonObject.ToBsonDocument(), new ReplaceOptions { IsUpsert = true });
+				return result.MatchedCount > 0;
 			}
-			catch (MongoWriteException Ex)
+			catch (MongoWriteException ex)
 			{
 				// Duplicate key error occurs if filter fails to match because revision is not the same.
-				if (Ex.WriteError != null && Ex.WriteError.Category == ServerErrorCategory.DuplicateKey)
+				if (ex.WriteError != null && ex.WriteError.Category == ServerErrorCategory.DuplicateKey)
 				{
 					return false;
 				}
