Horde: Fix coding conventions within Horde.Build.

diff --git a/Engine/Source/Programs/Horde/Horde.Build/Server/DowntimeService.cs b/Engine/Source/Programs/Horde/Horde.Build/Server/DowntimeService.cs
--- a/Engine/Source/Programs/Horde/Horde.Build/Server/DowntimeService.cs
+++ b/Engine/Source/Programs/Horde/Horde.Build/Server/DowntimeService.cs
@@ -1,17 +1,14 @@
 // Copyright Epic Games, Inc. All Rights Reserved.
 
-using HordeCommon;
-using Horde.Build.Models;
-using Horde.Build.Utilities;
-using Microsoft.Extensions.Hosting;
-using Microsoft.Extensions.Logging;
-using Microsoft.Extensions.Options;
 using System;
-using System.Collections.Generic;
 using System.Linq;
-using System.Security.Cryptography.X509Certificates;
 using System.Threading;
 using System.Threading.Tasks;
+using Horde.Build.Models;
+using HordeCommon;
+using Microsoft.Extensions.Hosting;
+using Microsoft.Extensions.Logging;
+using Microsoft.Extensions.Options;
 
 namespace Horde.Build.Services
 {
@@ -43,64 +40,64 @@ namespace Horde.Build.Services
 			private set;
 		}
 
-		DatabaseService DatabaseService;
-		ITicker Ticker;
-		IOptionsMonitor<ServerSettings> Settings;
-		ILogger Logger;
+		readonly DatabaseService _databaseService;
+		readonly ITicker _ticker;
+		readonly IOptionsMonitor<ServerSettings> _settings;
+		readonly ILogger _logger;
 
 		/// <summary>
 		/// Constructor
 		/// </summary>
-		/// <param name="DatabaseService">The database service instance</param>
-		/// <param name="Clock"></param>
-		/// <param name="Settings">The server settings</param>
-		/// <param name="Logger">Logger instance</param>
-		public DowntimeService(DatabaseService DatabaseService, IClock Clock, IOptionsMonitor<ServerSettings> Settings, ILogger<DowntimeService> Logger)
+		/// <param name="databaseService">The database service instance</param>
+		/// <param name="clock"></param>
+		/// <param name="settings">The server settings</param>
+		/// <param name="logger">Logger instance</param>
+		public DowntimeService(DatabaseService databaseService, IClock clock, IOptionsMonitor<ServerSettings> settings, ILogger<DowntimeService> logger)
 		{
-			this.DatabaseService = DatabaseService;
-			this.Settings = Settings;
-			this.Logger = Logger;
+			_databaseService = databaseService;
+			_settings = settings;
+			_logger = logger;
 
 			// Ensure the initial value to be correct
 			TickAsync(CancellationToken.None).AsTask().Wait();
 
-			Ticker = Clock.AddTicker(TimeSpan.FromMinutes(1.0), TickAsync, Logger);
+			_ticker = clock.AddTicker<DowntimeService>(TimeSpan.FromMinutes(1.0), TickAsync, logger);
 		}
 
 		/// <inheritdoc/>
-		public Task StartAsync(CancellationToken CancellationToken) => Ticker.StartAsync();
+		public Task StartAsync(CancellationToken cancellationToken) => _ticker.StartAsync();
 
 		/// <inheritdoc/>
-		public Task StopAsync(CancellationToken CancellationToken) => Ticker.StopAsync();
+		public Task StopAsync(CancellationToken cancellationToken) => _ticker.StopAsync();
 
 		/// <inheritdoc/>
-		public void Dispose() => Ticker.Dispose();
+		public void Dispose() => _ticker.Dispose();
 
 		/// <summary>
 		/// Periodically called tick function
 		/// </summary>
-		/// <param name="StoppingToken">Token indicating that the service should stop</param>
+		/// <param name="stoppingToken">Token indicating that the service should stop</param>
 		/// <returns>Async task</returns>
-		async ValueTask TickAsync(CancellationToken StoppingToken)
+		async ValueTask TickAsync(CancellationToken stoppingToken)
 		{
-			Globals Globals = await DatabaseService.GetGlobalsAsync();
+			Globals globals = await _databaseService.GetGlobalsAsync();
 
-			DateTimeOffset Now = TimeZoneInfo.ConvertTime(DateTimeOffset.Now, Settings.CurrentValue.TimeZoneInfo);
-			bool bIsActive = Globals.ScheduledDowntime.Any(x => x.IsActive(Now));
+			DateTimeOffset now = TimeZoneInfo.ConvertTime(DateTimeOffset.Now, _settings.CurrentValue.TimeZoneInfo);
+			bool bIsActive = globals.ScheduledDowntime.Any(x => x.IsActive(now));
 
-			DateTimeOffset? Next = null;
-			foreach (ScheduledDowntime Schedule in Globals.ScheduledDowntime)
+			DateTimeOffset? next = null;
+			foreach (ScheduledDowntime schedule in globals.ScheduledDowntime)
 			{
-				DateTimeOffset Start = Schedule.GetNext(Now).StartTime;
-				if(Next == null || Start < Next)
+				DateTimeOffset start = schedule.GetNext(now).StartTime;
+				if(next == null || start < next)
 				{
-					Next = Start;
+					next = start;
 				}
 			}
 
-			if (Next != null)
+			if (next != null)
 			{
-				Logger.LogInformation("Server time: {Time}. Downtime: {Downtime}. Next: {Next}.", Now, bIsActive, Next.Value);
+				_logger.LogInformation("Server time: {Time}. Downtime: {Downtime}. Next: {Next}.", now, bIsActive, next.Value);
 			}
 
 			if (bIsActive != IsDowntimeActive)
@@ -108,11 +105,11 @@ namespace Horde.Build.Services
 				IsDowntimeActive = bIsActive;
 				if (IsDowntimeActive)
 				{
-					Logger.LogInformation("Entering downtime");
+					_logger.LogInformation("Entering downtime");
 				}
 				else
 				{
-					Logger.LogInformation("Leaving downtime");
+					_logger.LogInformation("Leaving downtime");
 				}
 			}
 		}
diff --git a/Engine/Source/Programs/Horde/Horde.Build/Server/LifetimeService.cs b/Engine/Source/Programs/Horde/Horde.Build/Server/LifetimeService.cs
--- a/Engine/Source/Programs/Horde/Horde.Build/Server/LifetimeService.cs
+++ b/Engine/Source/Programs/Horde/Horde.Build/Server/LifetimeService.cs
@@ -1,19 +1,11 @@
 // Copyright Epic Games, Inc. All Rights Reserved.
 
-using Microsoft.AspNetCore.Hosting;
-using Microsoft.Extensions.Hosting;
-using Microsoft.Extensions.Logging;
 using System;
-using System.Collections.Generic;
-using System.Diagnostics.CodeAnalysis;
-using System.Linq;
 using System.Threading;
 using System.Threading.Tasks;
-using Grpc.Core;
-using Grpc.Core.Interceptors;
-using Horde.Build.Utilities;
+using Microsoft.Extensions.Hosting;
+using Microsoft.Extensions.Logging;
 using StackExchange.Redis;
-using System.Diagnostics;
 
 namespace Horde.Build.Services
 {
@@ -25,37 +17,32 @@ namespace Horde.Build.Services
 		/// <summary>
 		/// Writer for log output
 		/// </summary>
-		ILogger<LifetimeService> Logger;
+		readonly ILogger<LifetimeService> _logger;
 
 		/// <summary>
 		/// Task source for the server stopping
 		/// </summary>
-		TaskCompletionSource<bool> StoppingTaskCompletionSource;
+		readonly TaskCompletionSource<bool> _stoppingTaskCompletionSource;
 
 		/// <summary>
 		/// Task source for the server stopping
 		/// </summary>
-		TaskCompletionSource<bool> PreStoppingTaskCompletionSource;
+		readonly TaskCompletionSource<bool> _preStoppingTaskCompletionSource;
 
 		/// <summary>
 		/// Registration token for the stopping event
 		/// </summary>
-		CancellationTokenRegistration Registration;
+		readonly CancellationTokenRegistration _registration;
 		
 		/// <summary>
 		/// Singleton instance of the database service
 		/// </summary>
-		DatabaseService DatabaseService;
+		readonly DatabaseService _databaseService;
 		
 		/// <summary>
 		/// The Redis database
 		/// </summary>
-		IDatabase RedisDb;
-		
-		/// <summary>
-		/// Singleton instance of the request tracker service
-		/// </summary>
-		RequestTrackerService RequestTrackerService;
+		readonly IDatabase _redisDb;
 		
 		/*
 		/// <summary>
@@ -73,26 +60,24 @@ namespace Horde.Build.Services
 		/// <summary>
 		/// Constructor
 		/// </summary>
-		/// <param name="Lifetime">Application lifetime interface</param>
-		/// <param name="DatabaseService">Database singleton service</param>
-		/// <param name="RedisDb">Redis singleton service</param>
-		/// <param name="RequestTrackerService">Request tracker singleton service</param>
-		/// <param name="Logger">Logging interface</param>
-		public LifetimeService(IHostApplicationLifetime Lifetime, DatabaseService DatabaseService, IDatabase RedisDb, RequestTrackerService RequestTrackerService, ILogger<LifetimeService> Logger)
+		/// <param name="lifetime">Application lifetime interface</param>
+		/// <param name="databaseService">Database singleton service</param>
+		/// <param name="redisDb">Redis singleton service</param>
+		/// <param name="logger">Logging interface</param>
+		public LifetimeService(IHostApplicationLifetime lifetime, DatabaseService databaseService, IDatabase redisDb, ILogger<LifetimeService> logger)
 		{
-			this.DatabaseService = DatabaseService;
-			this.RedisDb = RedisDb;
-			this.RequestTrackerService = RequestTrackerService;
-			this.Logger = Logger;
-			this.StoppingTaskCompletionSource = new TaskCompletionSource<bool>();
-			this.PreStoppingTaskCompletionSource = new TaskCompletionSource<bool>();
-			this.Registration = Lifetime.ApplicationStopping.Register(ApplicationStopping);
+			_databaseService = databaseService;
+			_redisDb = redisDb;
+			_logger = logger;
+			_stoppingTaskCompletionSource = new TaskCompletionSource<bool>();
+			_preStoppingTaskCompletionSource = new TaskCompletionSource<bool>();
+			_registration = lifetime.ApplicationStopping.Register(ApplicationStopping);
 		}
 
 		/// <inheritdoc/>
 		public void Dispose()
 		{
-			Registration.Dispose();
+			_registration.Dispose();
 		}
 
 		/// <summary>
@@ -100,17 +85,17 @@ namespace Horde.Build.Services
 		/// </summary>
 		void ApplicationStopping()
 		{
-			Logger.LogInformation("SIGTERM signal received");
+			_logger.LogInformation("SIGTERM signal received");
 			IsPreStopping = true;
 			IsStopping = true;
 
-			PreStoppingTaskCompletionSource.TrySetResult(true);
-			StoppingTaskCompletionSource.TrySetResult(true);
+			_preStoppingTaskCompletionSource.TrySetResult(true);
+			_stoppingTaskCompletionSource.TrySetResult(true);
 
-			int ShutdownDelayMs = 30 * 1000;
-			Logger.LogInformation("Delaying shutdown by sleeping {ShutdownDelayMs} ms...", ShutdownDelayMs);
-			Thread.Sleep(ShutdownDelayMs);
-			Logger.LogInformation("Server process now shutting down...");
+			int shutdownDelayMs = 30 * 1000;
+			_logger.LogInformation("Delaying shutdown by sleeping {ShutdownDelayMs} ms...", shutdownDelayMs);
+			Thread.Sleep(shutdownDelayMs);
+			_logger.LogInformation("Server process now shutting down...");
 			
 			/*
 			if (PreStoppingTaskCompletionSource.TrySetResult(true))
@@ -143,15 +128,6 @@ namespace Horde.Build.Services
 			*/
 		}
 
-		async Task ExecStoppingTask()
-		{
-			await Task.Delay(TimeSpan.FromSeconds(15.0));
-
-			Logger.LogInformation("Setting stopping event");
-			IsStopping = true;
-			StoppingTaskCompletionSource.TrySetResult(true);
-		}
-
 		/// <summary>
 		/// Returns true if the server is stopping
 		/// </summary>
@@ -165,121 +141,50 @@ namespace Horde.Build.Services
 		/// <summary>
 		/// Gets an awaitable task for the server stopping
 		/// </summary>
-		public Task StoppingTask
-		{
-			get { return StoppingTaskCompletionSource.Task; }
-		}
-		
+		public Task StoppingTask => _stoppingTaskCompletionSource.Task;
+
 		/// <summary>
 		/// Check if MongoDB can be reached
 		/// </summary>
 		/// <returns>True if communication works</returns>
-		[SuppressMessage("Design", "CA1031:Do not catch general exception types")]
 		public async Task<bool> IsMongoDbConnectionHealthy()
 		{
-			using CancellationTokenSource CancelSource = new CancellationTokenSource(10000);
-			bool IsHealthy = false;
+			using CancellationTokenSource cancelSource = new CancellationTokenSource(10000);
+			bool isHealthy = false;
 			try
 			{
-				await DatabaseService.Database.ListCollectionNamesAsync(null, CancelSource.Token);
-				IsHealthy = true;
+				await _databaseService.Database.ListCollectionNamesAsync(null, cancelSource.Token);
+				isHealthy = true;
 			}
 			catch (Exception e)
 			{
-				Logger.LogError("MongoDB call failed during health check", e);
+				_logger.LogError("MongoDB call failed during health check", e);
 			}
 			
-			return IsHealthy;
+			return isHealthy;
 		}
 		
 		/// <summary>
 		/// Check if Redis can be reached
 		/// </summary>
 		/// <returns>True if communication works</returns>
-		[SuppressMessage("Design", "CA1031:Do not catch general exception types")]
 		public async Task<bool> IsRedisConnectionHealthy()
 		{
-			using CancellationTokenSource CancelSource = new CancellationTokenSource(10000);
-			bool IsHealthy = false;
+			using CancellationTokenSource cancelSource = new CancellationTokenSource(10000);
+			bool isHealthy = false;
 			try
 			{
-				string Key = "HordeLifetimeService-Health-Check";
-				await RedisDb.StringSetAsync(Key, "ok");
-				await RedisDb.StringGetAsync(Key);
-				IsHealthy = true;
+				string key = "HordeLifetimeService-Health-Check";
+				await _redisDb.StringSetAsync(key, "ok");
+				await _redisDb.StringGetAsync(key);
+				isHealthy = true;
 			}
 			catch (Exception e)
 			{
-				Logger.LogError("Redis call failed during health check", e);
+				_logger.LogError("Redis call failed during health check", e);
 			}
 			
-			return IsHealthy;
-		}
-	}
-
-	/// <summary>
-	/// Intercept incoming gRPC calls and raise an abort exception if the server is shutting down
-	///
-	/// This is to avoid clients keeping connections to a server that is about to shut down.
-	/// The client must handle the RPC exception raised and simply retry the call on a different gRPC channel/connection.
-	/// By retrying the call on a new channel, it should hit a new valid server as the load balancer will not direct traffic
-	/// to a server process shutting down.
-	/// </summary>
-	public class LifetimeGrpcInterceptor : Interceptor
-	{
-		private readonly LifetimeService LifetimeService;
-		private readonly ILogger<LifetimeGrpcInterceptor> Logger;
-
-		/// <summary>
-		/// Constructor
-		/// </summary>
-		/// <param name="LifetimeService"></param>
-		/// <param name="Logger"></param>
-		public LifetimeGrpcInterceptor(LifetimeService LifetimeService, ILogger<LifetimeGrpcInterceptor> Logger)
-		{
-			this.LifetimeService = LifetimeService;
-			this.Logger = Logger;
-		}
-
-		[SuppressMessage("Performance", "CA1822:Mark members as static")]
-		[SuppressMessage("Usage", "CA1801:Review unused parameters")]
-		private void CheckIfServerIsBeingStopped(ServerCallContext Context)
-		{
-			/*
-			if (LifetimeService.IsStopping)
-			{
-				Logger.LogDebug("Refusing to serve {Method} as server process is being shut down.", Context.Method);
-				throw new RpcException(new Status(StatusCode.Aborted, "Server is shutting down."));
-			}
-			*/
-		}
-
-		/// <inheritdoc />
-		public override Task<TResponse> UnaryServerHandler<TRequest, TResponse>(TRequest Request, ServerCallContext Context, UnaryServerMethod<TRequest, TResponse> Continuation)
-		{
-			CheckIfServerIsBeingStopped(Context);
-			return base.UnaryServerHandler(Request, Context, Continuation);
-		}
-
-		/// <inheritdoc />
-		public override Task<TResponse> ClientStreamingServerHandler<TRequest, TResponse>(IAsyncStreamReader<TRequest> RequestStream, ServerCallContext Context, ClientStreamingServerMethod<TRequest, TResponse> Continuation)
-		{
-			CheckIfServerIsBeingStopped(Context);
-			return base.ClientStreamingServerHandler(RequestStream, Context, Continuation);
-		}
-
-		/// <inheritdoc />
-		public override Task ServerStreamingServerHandler<TRequest, TResponse>(TRequest Request, IServerStreamWriter<TResponse> ResponseStream, ServerCallContext Context, ServerStreamingServerMethod<TRequest, TResponse> Continuation)
-		{
-			CheckIfServerIsBeingStopped(Context);
-			return base.ServerStreamingServerHandler(Request, ResponseStream, Context, Continuation);
-		}
-
-		/// <inheritdoc />
-		public override Task DuplexStreamingServerHandler<TRequest, TResponse>(IAsyncStreamReader<TRequest> RequestStream, IServerStreamWriter<TResponse> ResponseStream, ServerCallContext Context, DuplexStreamingServerMethod<TRequest, TResponse> Continuation)
-		{
-			CheckIfServerIsBeingStopped(Context);
-			return base.DuplexStreamingServerHandler(RequestStream, ResponseStream, Context, Continuation);
+			return isHealthy;
 		}
 	}
 }
\ No newline at end of file
diff --git a/Engine/Source/Programs/Horde/Horde.Build/Server/MetricService.cs b/Engine/Source/Programs/Horde/Horde.Build/Server/MetricService.cs
--- a/Engine/Source/Programs/Horde/Horde.Build/Server/MetricService.cs
+++ b/Engine/Source/Programs/Horde/Horde.Build/Server/MetricService.cs
@@ -1,13 +1,12 @@
 // Copyright Epic Games, Inc. All Rights Reserved.
 
-using Horde.Build.Utilities;
-using Microsoft.Extensions.Logging;
 using System;
 using System.Threading;
 using System.Threading.Tasks;
-using StatsdClient;
-using Microsoft.Extensions.Hosting;
 using HordeCommon;
+using Microsoft.Extensions.Hosting;
+using Microsoft.Extensions.Logging;
+using StatsdClient;
 
 namespace Horde.Build.Services
 {
@@ -16,33 +15,33 @@ namespace Horde.Build.Services
 	/// </summary>
 	public sealed class MetricService : IHostedService, IDisposable
 	{
-		IDogStatsd DogStatsd;
-		ITicker Ticker;
+		readonly IDogStatsd _dogStatsd;
+		readonly ITicker _ticker;
 
 		/// <summary>
 		/// Constructor
 		/// </summary>
-		public MetricService(IDogStatsd DogStatsd, IClock Clock, ILogger<MetricService> Logger)
+		public MetricService(IDogStatsd dogStatsd, IClock clock, ILogger<MetricService> logger)
 		{
-			this.DogStatsd = DogStatsd;
-			this.Ticker = Clock.AddTicker(TimeSpan.FromSeconds(20.0), TickAsync, Logger);
+			_dogStatsd = dogStatsd;
+			_ticker = clock.AddTicker<MetricService>(TimeSpan.FromSeconds(20.0), TickAsync, logger);
 		}
 
 		/// <inheritdoc/>
-		public Task StartAsync(CancellationToken cancellationToken) => Ticker.StartAsync();
+		public Task StartAsync(CancellationToken cancellationToken) => _ticker.StartAsync();
 
 		/// <inheritdoc/>
-		public Task StopAsync(CancellationToken cancellationToken) => Ticker.StopAsync();
+		public Task StopAsync(CancellationToken cancellationToken) => _ticker.StopAsync();
 
 		/// <inheritdoc/>
-		public void Dispose() => Ticker.Dispose();
+		public void Dispose() => _ticker.Dispose();
 
 		/// <summary>
 		/// Execute the background task
 		/// </summary>
-		/// <param name="StoppingToken">Cancellation token for the async task</param>
+		/// <param name="stoppingToken">Cancellation token for the async task</param>
 		/// <returns>Async task</returns>
-		ValueTask TickAsync(CancellationToken StoppingToken)
+		ValueTask TickAsync(CancellationToken stoppingToken)
 		{
 			ReportGcMetrics();
 			ReportThreadMetrics();
@@ -51,35 +50,35 @@ namespace Horde.Build.Services
 
 		private void ReportThreadMetrics()
 		{
-			ThreadPool.GetMaxThreads(out int MaxWorkerThreads, out int MaxIoThreads);
-			ThreadPool.GetAvailableThreads(out int FreeWorkerThreads, out int FreeIoThreads);
-			ThreadPool.GetMinThreads(out int MinWorkerThreads, out int MinIoThreads);
+			ThreadPool.GetMaxThreads(out int maxWorkerThreads, out int maxIoThreads);
+			ThreadPool.GetAvailableThreads(out int freeWorkerThreads, out int freeIoThreads);
+			ThreadPool.GetMinThreads(out int minWorkerThreads, out int minIoThreads);
 
-			int BusyIoThreads = MaxIoThreads - FreeIoThreads;
-			int BusyWorkerThreads = MaxWorkerThreads - FreeWorkerThreads;
+			int busyIoThreads = maxIoThreads - freeIoThreads;
+			int busyWorkerThreads = maxWorkerThreads - freeWorkerThreads;
 			
-			DogStatsd.Gauge("horde.clr.threadpool.io.max", MaxIoThreads);
-			DogStatsd.Gauge("horde.clr.threadpool.io.min", MinIoThreads);
-			DogStatsd.Gauge("horde.clr.threadpool.io.free", FreeIoThreads);
-			DogStatsd.Gauge("horde.clr.threadpool.io.busy", BusyIoThreads);
+			_dogStatsd.Gauge("horde.clr.threadpool.io.max", maxIoThreads);
+			_dogStatsd.Gauge("horde.clr.threadpool.io.min", minIoThreads);
+			_dogStatsd.Gauge("horde.clr.threadpool.io.free", freeIoThreads);
+			_dogStatsd.Gauge("horde.clr.threadpool.io.busy", busyIoThreads);
 			
-			DogStatsd.Gauge("horde.clr.threadpool.worker.max", MaxWorkerThreads);
-			DogStatsd.Gauge("horde.clr.threadpool.worker.min", MinWorkerThreads);
-			DogStatsd.Gauge("horde.clr.threadpool.worker.free", FreeWorkerThreads);
-			DogStatsd.Gauge("horde.clr.threadpool.worker.busy", BusyWorkerThreads);
+			_dogStatsd.Gauge("horde.clr.threadpool.worker.max", maxWorkerThreads);
+			_dogStatsd.Gauge("horde.clr.threadpool.worker.min", minWorkerThreads);
+			_dogStatsd.Gauge("horde.clr.threadpool.worker.free", freeWorkerThreads);
+			_dogStatsd.Gauge("horde.clr.threadpool.worker.busy", busyWorkerThreads);
 		}
 
 		private void ReportGcMetrics()
 		{
-			GCMemoryInfo GcMemoryInfo = GC.GetGCMemoryInfo();
+			GCMemoryInfo gcMemoryInfo = GC.GetGCMemoryInfo();
 
-			DogStatsd.Gauge("horde.clr.gc.totalMemory", GC.GetTotalMemory(false));
-			DogStatsd.Gauge("horde.clr.gc.totalAllocated", GC.GetTotalAllocatedBytes());
-			DogStatsd.Gauge("horde.clr.gc.heapSize", GcMemoryInfo.HeapSizeBytes);
-			DogStatsd.Gauge("horde.clr.gc.fragmented", GcMemoryInfo.FragmentedBytes);
-			DogStatsd.Gauge("horde.clr.gc.memoryLoad", GcMemoryInfo.MemoryLoadBytes);
-			DogStatsd.Gauge("horde.clr.gc.totalAvailableMemory", GcMemoryInfo.TotalAvailableMemoryBytes);
-			DogStatsd.Gauge("horde.clr.gc.highMemoryLoadThreshold", GcMemoryInfo.HighMemoryLoadThresholdBytes);
+			_dogStatsd.Gauge("horde.clr.gc.totalMemory", GC.GetTotalMemory(false));
+			_dogStatsd.Gauge("horde.clr.gc.totalAllocated", GC.GetTotalAllocatedBytes());
+			_dogStatsd.Gauge("horde.clr.gc.heapSize", gcMemoryInfo.HeapSizeBytes);
+			_dogStatsd.Gauge("horde.clr.gc.fragmented", gcMemoryInfo.FragmentedBytes);
+			_dogStatsd.Gauge("horde.clr.gc.memoryLoad", gcMemoryInfo.MemoryLoadBytes);
+			_dogStatsd.Gauge("horde.clr.gc.totalAvailableMemory", gcMemoryInfo.TotalAvailableMemoryBytes);
+			_dogStatsd.Gauge("horde.clr.gc.highMemoryLoadThreshold", gcMemoryInfo.HighMemoryLoadThresholdBytes);
 		}
 	}
 }
diff --git a/Engine/Source/Programs/Horde/Horde.Build/Server/RedisService.cs b/Engine/Source/Programs/Horde/Horde.Build/Server/RedisService.cs
--- a/Engine/Source/Programs/Horde/Horde.Build/Server/RedisService.cs
+++ b/Engine/Source/Programs/Horde/Horde.Build/Server/RedisService.cs
@@ -1,12 +1,6 @@
 // Copyright Epic Games, Inc. All Rights Reserved.
 
-using EpicGames.Core;
-using Horde.Build.Utilities;
-using Microsoft.Extensions.Logging;
-using Microsoft.Extensions.Options;
-using StackExchange.Redis;
 using System;
-using System.Collections.Generic;
 using System.Diagnostics;
 using System.IO;
 using System.Linq;
@@ -14,6 +8,11 @@ using System.Net;
 using System.Net.NetworkInformation;
 using System.Runtime.InteropServices;
 using System.Threading.Tasks;
+using EpicGames.Core;
+using Horde.Build.Utilities;
+using Microsoft.Extensions.Logging;
+using Microsoft.Extensions.Options;
+using StackExchange.Redis;
 
 namespace Horde.Build.Services
 {
@@ -30,12 +29,12 @@ namespace Horde.Build.Services
 		/// <summary>
 		/// The managed process group containing the Redis server
 		/// </summary>
-		ManagedProcessGroup? RedisProcessGroup;
+		ManagedProcessGroup? _redisProcessGroup;
 
 		/// <summary>
 		/// The server process
 		/// </summary>
-		ManagedProcess? RedisProcess;
+		ManagedProcess? _redisProcess;
 
 		/// <summary>
 		/// Connection multiplexer
@@ -55,73 +54,68 @@ namespace Horde.Build.Services
 		/// <summary>
 		/// Logger factory
 		/// </summary>
-		ILoggerFactory LoggerFactory;
+		readonly ILoggerFactory _loggerFactory;
 
 		/// <summary>
 		/// Logging instance
 		/// </summary>
-		ILogger<RedisService> Logger;
-
-		/// <summary>
-		/// The output task
-		/// </summary>
-		Task? OutputTask;
+		readonly ILogger<RedisService> _logger;
 
 		/// <summary>
 		/// Hack to initialize RedisService early enough to use data protection
 		/// </summary>
-		/// <param name="Settings"></param>
-		public RedisService(ServerSettings Settings)
-			: this(Options.Create<ServerSettings>(Settings), new Serilog.Extensions.Logging.SerilogLoggerFactory(Serilog.Log.Logger))
+		/// <param name="settings"></param>
+		public RedisService(ServerSettings settings)
+			: this(Options.Create<ServerSettings>(settings), new Serilog.Extensions.Logging.SerilogLoggerFactory(Serilog.Log.Logger))
 		{
 		}
 
 		/// <summary>
 		/// Constructor
 		/// </summary>
-		/// <param name="Options"></param>
-		/// <param name="LoggerFactory"></param>
-		public RedisService(IOptions<ServerSettings> Options, ILoggerFactory LoggerFactory)
+		/// <param name="options"></param>
+		/// <param name="loggerFactory"></param>
+		public RedisService(IOptions<ServerSettings> options, ILoggerFactory loggerFactory)
 		{
-			this.LoggerFactory = LoggerFactory;
-			this.Logger = LoggerFactory.CreateLogger<RedisService>();
+			_loggerFactory = loggerFactory;
+			_logger = loggerFactory.CreateLogger<RedisService>();
 
-			ServerSettings Settings = Options.Value;
+			ServerSettings settings = options.Value;
 
-			string? ConnectionString = Settings.RedisConnectionConfig;
-			if (ConnectionString == null)
+			string? connectionString = settings.RedisConnectionConfig;
+			if (connectionString == null)
 			{
 				if (IsRunningOnDefaultPort())
 				{
-					ConnectionString = $"localhost:{RedisPort}";
+					connectionString = $"localhost:{RedisPort}";
 				}
 				else if (TryStartRedisServer())
 				{
-					ConnectionString = $"localhost:{RedisPort}";
+					connectionString = $"localhost:{RedisPort}";
 				}
 				else
 				{
-					throw new Exception($"Unable to connect to Redis. Please set {nameof(Settings.RedisConnectionConfig)} in {Program.UserConfigFile}");
+					throw new Exception($"Unable to connect to Redis. Please set {nameof(settings.RedisConnectionConfig)} in {Program.UserConfigFile}");
 				}
 			}
 
-			Multiplexer = ConnectionMultiplexer.Connect(ConnectionString);
+			Multiplexer = ConnectionMultiplexer.Connect(connectionString);
 			Database = Multiplexer.GetDatabase();
-			ConnectionPool = new RedisConnectionPool(20, ConnectionString);
+			ConnectionPool = new RedisConnectionPool(20, connectionString);
 		}
 
 		/// <inheritdoc/>
 		public void Dispose()
 		{
-			if (RedisProcess != null)
+			if (_redisProcess != null)
 			{
-				RedisProcess.Dispose();
-				RedisProcess = null;
+				_redisProcess.Dispose();
+				_redisProcess = null;
 			}
-			if (RedisProcessGroup != null)
+			if (_redisProcessGroup != null)
 			{
-				RedisProcessGroup.Dispose();
-				RedisProcessGroup = null;
+				_redisProcessGroup.Dispose();
+				_redisProcessGroup = null;
 			}
 		}
 
@@ -131,10 +125,10 @@ namespace Horde.Build.Services
 		/// <returns></returns>
 		static bool IsRunningOnDefaultPort()
 		{
-			IPGlobalProperties IpGlobalProperties = IPGlobalProperties.GetIPGlobalProperties();
+			IPGlobalProperties ipGlobalProperties = IPGlobalProperties.GetIPGlobalProperties();
 
-			IPEndPoint[] Listeners = IpGlobalProperties.GetActiveTcpListeners();
-			if (Listeners.Any(x => x.Port == RedisPort))
+			IPEndPoint[] listeners = ipGlobalProperties.GetActiveTcpListeners();
+			if (listeners.Any(x => x.Port == RedisPort))
 			{
 				return true;
 			}
@@ -153,36 +147,36 @@ namespace Horde.Build.Services
 				return false;
 			}
 
-			FileReference RedisExe = FileReference.Combine(Program.AppDir, "ThirdParty", "Redis", "redis-server.exe");
-			if (!FileReference.Exists(RedisExe))
+			FileReference redisExe = FileReference.Combine(Program.AppDir, "ThirdParty", "Redis", "redis-server.exe");
+			if (!FileReference.Exists(redisExe))
 			{
-				Logger.LogDebug("Redis executable does not exist at {ExePath}", RedisExe);
+				_logger.LogDebug("Redis executable does not exist at {ExePath}", redisExe);
 				return false;
 			}
 
-			DirectoryReference RedisDir = DirectoryReference.Combine(Program.DataDir, "Redis");
-			DirectoryReference.CreateDirectory(RedisDir);
+			DirectoryReference redisDir = DirectoryReference.Combine(Program.DataDir, "Redis");
+			DirectoryReference.CreateDirectory(redisDir);
 
-			FileReference RedisConfigFile = FileReference.Combine(RedisDir, "redis.conf");
-			if(!FileReference.Exists(RedisConfigFile))
+			FileReference redisConfigFile = FileReference.Combine(redisDir, "redis.conf");
+			if(!FileReference.Exists(redisConfigFile))
 			{
-				using (StreamWriter Writer = new StreamWriter(RedisConfigFile.FullName))
+				using (StreamWriter writer = new StreamWriter(redisConfigFile.FullName))
 				{
-					Writer.WriteLine("# redis.conf");
+					writer.WriteLine("# redis.conf");
 				}
 			}
 
-			RedisProcessGroup = new ManagedProcessGroup();
+			_redisProcessGroup = new ManagedProcessGroup();
 			try
 			{
-				RedisProcess = new ManagedProcess(RedisProcessGroup, RedisExe.FullName, "", null, null, ProcessPriorityClass.Normal);
-				RedisProcess.StdIn.Close();
-				OutputTask = Task.Run(() => RelayRedisOutput());
+				_redisProcess = new ManagedProcess(_redisProcessGroup, redisExe.FullName, "", null, null, ProcessPriorityClass.Normal);
+				_redisProcess.StdIn.Close();
+				Task.Run(() => RelayRedisOutput());
 				return true;
 			}
-			catch (Exception Ex)
+			catch (Exception ex)
 			{
-				Logger.LogWarning(Ex, "Unable to start Redis server process");
+				_logger.LogWarning(ex, "Unable to start Redis server process");
 				return false;
 			}
 		}
@@ -193,20 +187,20 @@ namespace Horde.Build.Services
 		/// <returns></returns>
 		async Task RelayRedisOutput()
 		{
-			ILogger RedisLogger = LoggerFactory.CreateLogger("Redis");
+			ILogger redisLogger = _loggerFactory.CreateLogger("Redis");
 			for (; ; )
 			{
-				string? Line = await RedisProcess!.ReadLineAsync();
-				if (Line == null)
+				string? line = await _redisProcess!.ReadLineAsync();
+				if (line == null)
 				{
 					break;
 				}
-				if (Line.Length > 0)
+				if (line.Length > 0)
 				{
-					RedisLogger.Log(LogLevel.Information, "{Output}", Line);
+					redisLogger.Log(LogLevel.Information, "{Output}", line);
 				}
 			}
-			RedisLogger.LogInformation("Exit code {ExitCode}", RedisProcess.ExitCode);
+			redisLogger.LogInformation("Exit code {ExitCode}", _redisProcess.ExitCode);
 		}
 	}
 }
diff --git a/Engine/Source/Programs/Horde/Horde.Build/Server/RequestTrackerService.cs b/Engine/Source/Programs/Horde/Horde.Build/Server/RequestTrackerService.cs
--- a/Engine/Source/Programs/Horde/Horde.Build/Server/RequestTrackerService.cs
+++ b/Engine/Source/Programs/Horde/Horde.Build/Server/RequestTrackerService.cs
@@ -22,35 +22,35 @@ namespace Horde.Build.Utilities
 		/// <summary>
 		/// Writer for log output
 		/// </summary>
-		private readonly ILogger<RequestTrackerService> Logger;
+		private readonly ILogger<RequestTrackerService> _logger;
 
-		readonly ConcurrentDictionary<string, TrackedRequest> RequestsInProgress = new ConcurrentDictionary<string, TrackedRequest>();
+		readonly ConcurrentDictionary<string, TrackedRequest> _requestsInProgress = new ConcurrentDictionary<string, TrackedRequest>();
 
 		/// <summary>
 		/// Constructor
 		/// </summary>
-		/// <param name="Logger">Logger</param>
-		public RequestTrackerService(ILogger<RequestTrackerService> Logger)
+		/// <param name="logger">Logger</param>
+		public RequestTrackerService(ILogger<RequestTrackerService> logger)
 		{
-			this.Logger = Logger;
+			_logger = logger;
 		}
 
 		/// <summary>
 		/// Called by the middleware when a request is started
 		/// </summary>
-		/// <param name="Context">HTTP Context</param>
-		public void RequestStarted(HttpContext Context)
+		/// <param name="context">HTTP Context</param>
+		public void RequestStarted(HttpContext context)
 		{
-			RequestsInProgress[Context.TraceIdentifier] = new TrackedRequest(Context.Request);
+			_requestsInProgress[context.TraceIdentifier] = new TrackedRequest(context.Request);
 		}
 		
 		/// <summary>
 		/// Called by the middleware when a request is finished (no matter if an exception occurred or not)
 		/// </summary>
-		/// <param name="Context">HTTP Context</param>
-		public void RequestFinished(HttpContext Context)
+		/// <param name="context">HTTP Context</param>
+		public void RequestFinished(HttpContext context)
 		{
-			RequestsInProgress.Remove(Context.TraceIdentifier, out _);
+			_requestsInProgress.Remove(context.TraceIdentifier, out _);
 		}
 
 		/// <summary>
@@ -59,22 +59,22 @@ namespace Horde.Build.Utilities
 		/// <returns>The requests in progress</returns>
 		public IReadOnlyDictionary<string, TrackedRequest> GetRequestsInProgress()
 		{
-			return RequestsInProgress;
+			return _requestsInProgress;
 		}
 
 		private string GetRequestsInProgressAsString()
 		{
-			List<KeyValuePair<string, TrackedRequest>> Requests = GetRequestsInProgress().ToList();
-			Requests.Sort((A, B) => A.Value.StartedAt.CompareTo(B.Value.StartedAt));
-			StringBuilder Content = new StringBuilder();
-			foreach (KeyValuePair<string,TrackedRequest> Pair in Requests)
+			List<KeyValuePair<string, TrackedRequest>> requests = GetRequestsInProgress().ToList();
+			requests.Sort((a, b) => a.Value.StartedAt.CompareTo(b.Value.StartedAt));
+			StringBuilder content = new StringBuilder();
+			foreach (KeyValuePair<string,TrackedRequest> pair in requests)
 			{
-				int AgeInMs = Pair.Value.GetTimeSinceStartInMs();
-				string Path = Pair.Value.Request.Path;
-				Content.AppendLine(CultureInfo.InvariantCulture, $"{AgeInMs,9}  {Path}");
+				int ageInMs = pair.Value.GetTimeSinceStartInMs();
+				string path = pair.Value.Request.Path;
+				content.AppendLine(CultureInfo.InvariantCulture, $"{ageInMs,9}  {path}");
 			}
 
-			return Content.ToString();
+			return content.ToString();
 		}
 
 		/// <summary>
@@ -84,11 +84,11 @@ namespace Horde.Build.Utilities
 		{
 			if (GetRequestsInProgress().Count == 0)
 			{
-				Logger.LogInformation("There are no requests in progress!");
+				_logger.LogInformation("There are no requests in progress!");
 			}
 			else
 			{
-				Logger.LogInformation("Current open requests are:\n{RequestsInProgress}", GetRequestsInProgressAsString());
+				_logger.LogInformation("Current open requests are:\n{RequestsInProgress}", GetRequestsInProgressAsString());
 			}
 		}
 	}
@@ -111,11 +111,11 @@ namespace Horde.Build.Utilities
 		/// <summary>
 		/// Constructor
 		/// </summary>
-		/// <param name="Request">HTTP request being tracked</param>
-		public TrackedRequest(HttpRequest Request)
+		/// <param name="request">HTTP request being tracked</param>
+		public TrackedRequest(HttpRequest request)
 		{
 			StartedAt = DateTime.UtcNow;
-			this.Request = Request;
+			Request = request;
 		}
 
 		/// <summary>
@@ -133,41 +133,41 @@ namespace Horde.Build.Utilities
 	/// </summary>
 	public class RequestTrackerMiddleware
 	{
-		private readonly RequestDelegate Next;
+		private readonly RequestDelegate _next;
 	
 		/// <summary>
 		/// Constructor
 		/// </summary>
-		/// <param name="Next">Next middleware to call</param>
-		public RequestTrackerMiddleware(RequestDelegate Next)
+		/// <param name="next">Next middleware to call</param>
+		public RequestTrackerMiddleware(RequestDelegate next)
 		{
-			this.Next = Next;
+			_next = next;
 		}
 	
 		/// <summary>
 		/// Invoked by ASP.NET framework itself
 		/// </summary>
-		/// <param name="Context">HTTP Context</param>
-		/// <param name="Service">The RequestTrackerService singleton</param>
+		/// <param name="context">HTTP Context</param>
+		/// <param name="service">The RequestTrackerService singleton</param>
 		/// <returns></returns>
-		public async Task Invoke(HttpContext Context, RequestTrackerService Service)
+		public async Task Invoke(HttpContext context, RequestTrackerService service)
 		{
-			if (!Context.Request.Path.StartsWithSegments("/health", StringComparison.Ordinal))
+			if (!context.Request.Path.StartsWithSegments("/health", StringComparison.Ordinal))
 			{
 				try
 				{
-					Service.RequestStarted(Context);
-					await Next(Context);
+					service.RequestStarted(context);
+					await _next(context);
 				}
 				finally
 				{
-					Service.RequestFinished(Context);
+					service.RequestFinished(context);
 				}
 			}
 			else
 			{
 				// Ignore requests to /health/*
-				await Next(Context);
+				await _next(context);
 			}
 		}
 	}
diff --git a/Engine/Source/Programs/Horde/Horde.Build/ServerSettings.cs b/Engine/Source/Programs/Horde/Horde.Build/ServerSettings.cs
--- a/Engine/Source/Programs/Horde/Horde.Build/ServerSettings.cs
+++ b/Engine/Source/Programs/Horde/Horde.Build/ServerSettings.cs
@@ -1,21 +1,11 @@
 // Copyright Epic Games, Inc. All Rights Reserved.
 
-using EpicGames.Core;
-using Horde.Build.Logs.Readers;
-using Horde.Build.Models;
-using Horde.Build.Services;
-using Horde.Build.Storage;
-using Horde.Build.Storage.Backends;
-using Horde.Build.Utilities;
 using System;
-using System.Collections.Generic;
-using System.ComponentModel.DataAnnotations;
 using System.Linq;
-using System.Runtime.InteropServices;
-using System.Threading.Tasks;
+using EpicGames.Horde.Storage;
 using Horde.Build.Fleet.Autoscale;
+using Horde.Build.Utilities;
 using TimeZoneConverter;
-using EpicGames.Horde.Storage;
 
 namespace Horde.Build
 {
@@ -529,15 +519,15 @@ namespace Horde.Build
 		{
 			get
 			{
-				if (CachedTimeZoneInfo == null)
+				if (_cachedTimeZoneInfo == null)
 				{
-					CachedTimeZoneInfo = (ScheduleTimeZone == null) ? TimeZoneInfo.Local : TZConvert.GetTimeZoneInfo(ScheduleTimeZone);
+					_cachedTimeZoneInfo = (ScheduleTimeZone == null) ? TimeZoneInfo.Local : TZConvert.GetTimeZoneInfo(ScheduleTimeZone);
 				}
-				return CachedTimeZoneInfo;
+				return _cachedTimeZoneInfo;
 			}
 		}
 
-		private TimeZoneInfo? CachedTimeZoneInfo;
+		private TimeZoneInfo? _cachedTimeZoneInfo;
 
 		/// <summary>
 		/// Whether to open a browser on startup
@@ -550,12 +540,15 @@ namespace Horde.Build
 		/// <summary>
 		/// Helper method to check if this process has activated the given mode
 		/// </summary>
-		/// <param name="Mode">Run mode</param>
+		/// <param name="mode">Run mode</param>
 		/// <returns>True if mode is active</returns>
-		public bool IsRunModeActive(RunMode Mode)
+		public bool IsRunModeActive(RunMode mode)
 		{
-			if (RunModes == null) return true;
-			return RunModes.Contains(Mode);
+			if (RunModes == null)
+			{
+				return true;
+			}
+			return RunModes.Contains(mode);
 		}
 
 		/// <summary>
diff --git a/Engine/Source/Programs/Horde/Horde.Build/Services/ConsistencyService.cs b/Engine/Source/Programs/Horde/Horde.Build/Services/ConsistencyService.cs
--- a/Engine/Source/Programs/Horde/Horde.Build/Services/ConsistencyService.cs
+++ b/Engine/Source/Programs/Horde/Horde.Build/Services/ConsistencyService.cs
@@ -1,17 +1,16 @@
 // Copyright Epic Games, Inc. All Rights Reserved.
 
-using HordeCommon;
-using Horde.Build.Collections;
-using Horde.Build.Models;
-using Horde.Build.Utilities;
-using Microsoft.Extensions.Hosting;
-using Microsoft.Extensions.Logging;
-using MongoDB.Bson;
 using System;
 using System.Collections.Generic;
 using System.Linq;
 using System.Threading;
 using System.Threading.Tasks;
+using Horde.Build.Collections;
+using Horde.Build.Models;
+using Horde.Build.Utilities;
+using HordeCommon;
+using Microsoft.Extensions.Hosting;
+using Microsoft.Extensions.Logging;
 
 namespace Horde.Build.Services
 {
@@ -22,51 +21,51 @@ namespace Horde.Build.Services
 	/// </summary>
 	class ConsistencyService : IHostedService, IDisposable
 	{
-		ISessionCollection SessionCollection;
-		ILeaseCollection LeaseCollection;
-		ITicker Ticker;
-		ILogger<ConsistencyService> Logger;
+		readonly ISessionCollection _sessionCollection;
+		readonly ILeaseCollection _leaseCollection;
+		readonly ITicker _ticker;
+		readonly ILogger<ConsistencyService> _logger;
 
 		/// <summary>
 		/// Constructor
 		/// </summary>
-		public ConsistencyService(ISessionCollection SessionCollection, ILeaseCollection LeaseCollection, IClock Clock, ILogger<ConsistencyService> Logger)
+		public ConsistencyService(ISessionCollection sessionCollection, ILeaseCollection leaseCollection, IClock clock, ILogger<ConsistencyService> logger)
 		{
-			this.SessionCollection = SessionCollection;
-			this.LeaseCollection = LeaseCollection;
-			this.Ticker = Clock.AddSharedTicker<ConsistencyService>(TimeSpan.FromMinutes(20.0), TickLeaderAsync, Logger);
-			this.Logger = Logger;
+			_sessionCollection = sessionCollection;
+			_leaseCollection = leaseCollection;
+			_ticker = clock.AddSharedTicker<ConsistencyService>(TimeSpan.FromMinutes(20.0), TickLeaderAsync, logger);
+			_logger = logger;
 		}
 
 		/// <inheritdoc/>
-		public Task StartAsync(CancellationToken cancellationToken) => Ticker.StartAsync();
+		public Task StartAsync(CancellationToken cancellationToken) => _ticker.StartAsync();
 
 		/// <inheritdoc/>
-		public Task StopAsync(CancellationToken cancellationToken) => Ticker.StopAsync();
+		public Task StopAsync(CancellationToken cancellationToken) => _ticker.StopAsync();
 
 		/// <inheritdoc/>
-		public void Dispose() => Ticker.Dispose();
+		public void Dispose() => _ticker.Dispose();
 
 		/// <summary>
 		/// Poll for inconsistencies in the database
 		/// </summary>
-		/// <param name="StoppingToken">Stopping token</param>
+		/// <param name="stoppingToken">Stopping token</param>
 		/// <returns>Async task</returns>
-		async ValueTask TickLeaderAsync(CancellationToken StoppingToken)
+		async ValueTask TickLeaderAsync(CancellationToken stoppingToken)
 		{
-			List<ISession> Sessions = await SessionCollection.FindActiveSessionsAsync();
-			Dictionary<SessionId, ISession> SessionIdToInstance = Sessions.ToDictionary(x => x.Id, x => x);
+			List<ISession> sessions = await _sessionCollection.FindActiveSessionsAsync();
+			Dictionary<SessionId, ISession> sessionIdToInstance = sessions.ToDictionary(x => x.Id, x => x);
 
 			// Find any leases that are still running when their session has terminated
-			List<ILease> Leases = await LeaseCollection.FindActiveLeasesAsync();
-			foreach (ILease Lease in Leases)
+			List<ILease> leases = await _leaseCollection.FindActiveLeasesAsync();
+			foreach (ILease lease in leases)
 			{
-				if (!SessionIdToInstance.ContainsKey(Lease.SessionId))
+				if (!sessionIdToInstance.ContainsKey(lease.SessionId))
 				{
-					ISession? Session = await SessionCollection.GetAsync(Lease.SessionId);
-					DateTime FinishTime = Session?.FinishTime ?? DateTime.UtcNow;
-					Logger.LogWarning("Setting finish time for lease {LeaseId} to {FinishTime}", Lease.Id, FinishTime);
-					await LeaseCollection.TrySetOutcomeAsync(Lease.Id, FinishTime, LeaseOutcome.Cancelled, null);
+					ISession? session = await _sessionCollection.GetAsync(lease.SessionId);
+					DateTime finishTime = session?.FinishTime ?? DateTime.UtcNow;
+					_logger.LogWarning("Setting finish time for lease {LeaseId} to {FinishTime}", lease.Id, finishTime);
+					await _leaseCollection.TrySetOutcomeAsync(lease.Id, finishTime, LeaseOutcome.Cancelled, null);
 				}
 			}
 		}
diff --git a/Engine/Source/Programs/Horde/Horde.Build/Services/HealthService.cs b/Engine/Source/Programs/Horde/Horde.Build/Services/HealthService.cs
--- a/Engine/Source/Programs/Horde/Horde.Build/Services/HealthService.cs
+++ b/Engine/Source/Programs/Horde/Horde.Build/Services/HealthService.cs
@@ -1,8 +1,7 @@
 // Copyright Epic Games, Inc. All Rights Reserved.
 
-using Grpc.Core;
-using Microsoft.Extensions.Logging;
 using System.Threading.Tasks;
+using Grpc.Core;
 using Grpc.Health.V1;
 using static Grpc.Health.V1.HealthCheckResponse.Types;
 
@@ -14,25 +13,15 @@ namespace Horde.Build.Services
 	/// </summary>
 	public class HealthService : Health.HealthBase
 	{
-		/// <summary>
-		/// The application lifetime interface
-		/// </summary>
-		LifetimeService LifetimeService;
-
-		/// <summary>
-		/// Writer for log output
-		/// </summary>
-		ILogger<HealthService> Logger;
+		readonly LifetimeService _lifetimeService;
 
 		/// <summary>
 		/// Constructor
 		/// </summary>
-		/// <param name="LifetimeService">The application lifetime</param>
-		/// <param name="Logger">Log writer</param>
-		public HealthService(LifetimeService LifetimeService, ILogger<HealthService> Logger)
+		/// <param name="lifetimeService">The application lifetime</param>
+		public HealthService(LifetimeService lifetimeService)
 		{
-			this.LifetimeService = LifetimeService;
-			this.Logger = Logger;
+			_lifetimeService = lifetimeService;
 		}
 
 		/// <summary>
@@ -41,32 +30,32 @@ namespace Horde.Build.Services
 		///
 		/// For example, the gRPC health check for AWS ALB will pick up and react to the gRPC status code returned.
 		/// </summary>
-		/// <param name="Request">Empty placeholder request (for now)</param>
-		/// <param name="Context">Context for the call</param>
+		/// <param name="request">Empty placeholder request (for now)</param>
+		/// <param name="context">Context for the call</param>
 		/// <returns>Return status code 'unavailable' if stopping</returns>
-		public override Task<HealthCheckResponse> Check(HealthCheckRequest Request, ServerCallContext Context)
+		public override Task<HealthCheckResponse> Check(HealthCheckRequest request, ServerCallContext context)
 		{
-			ServingStatus Status = ServingStatus.Serving;
+			ServingStatus status = ServingStatus.Serving;
 			
-			bool IsStopping = LifetimeService.IsPreStopping || LifetimeService.IsStopping;
-			if (IsStopping)
+			bool isStopping = _lifetimeService.IsPreStopping || _lifetimeService.IsStopping;
+			if (isStopping)
 			{
-				Context.Status = new Status(StatusCode.Unavailable, "Server is stopping");
-				Status = ServingStatus.NotServing;
+				context.Status = new Status(StatusCode.Unavailable, "Server is stopping");
+				status = ServingStatus.NotServing;
 			}
 			 
-			return Task.FromResult(new HealthCheckResponse {Status = Status});
+			return Task.FromResult(new HealthCheckResponse {Status = status});
 		}
 
 		
 		/// <summary>
 		/// Stream the server health status (not implemented)
 		/// </summary>
-		/// <param name="Request"></param>
-		/// <param name="ResponseStream"></param>
-		/// <param name="Context"></param>
+		/// <param name="request"></param>
+		/// <param name="responseStream"></param>
+		/// <param name="context"></param>
 		/// <returns></returns>
-		public override Task Watch(HealthCheckRequest Request, IServerStreamWriter<HealthCheckResponse> ResponseStream, ServerCallContext Context)
+		public override Task Watch(HealthCheckRequest request, IServerStreamWriter<HealthCheckResponse> responseStream, ServerCallContext context)
 		{
 			return Task.FromException(new RpcException(new Status(StatusCode.Unimplemented, "Watch() not implemented")));
 		}
