Horde: Fix coding conventions within Horde.Build.

diff --git a/Engine/Source/Programs/Horde/Horde.Build/Services/RpcService.cs b/Engine/Source/Programs/Horde/Horde.Build/Services/RpcService.cs
--- a/Engine/Source/Programs/Horde/Horde.Build/Services/RpcService.cs
+++ b/Engine/Source/Programs/Horde/Horde.Build/Services/RpcService.cs
@@ -1,42 +1,40 @@
-﻿// Copyright Epic Games, Inc. All Rights Reserved.
+// Copyright Epic Games, Inc. All Rights Reserved.
 
+using System;
+using System.Collections.Generic;
+using System.Globalization;
+using System.Linq;
+using System.Net;
+using System.Security.Claims;
+using System.Text;
+using System.Threading;
+using System.Threading.Tasks;
 using EpicGames.Core;
 using Google.Protobuf.WellKnownTypes;
 using Grpc.Core;
-using HordeCommon;
-using HordeCommon.Rpc;
 using Horde.Build.Acls;
 using Horde.Build.Collections;
+using Horde.Build.Jobs;
 using Horde.Build.Models;
+using Horde.Build.Tasks.Impl;
 using Horde.Build.Utilities;
+using HordeCommon;
+using HordeCommon.Rpc;
 using Microsoft.AspNetCore.Authorization;
 using Microsoft.Extensions.Logging;
 using MongoDB.Bson;
 using MongoDB.Bson.Serialization;
-using System;
-using System.Collections.Generic;
-using System.Globalization;
-using System.Linq;
-using System.Net;
-using System.Runtime.CompilerServices;
-using System.Security.Claims;
-using System.Text;
-using System.Threading.Tasks;
-using Horde.Build.Tasks.Impl;
-using System.Threading;
-using Horde.Build.Jobs;
 
 namespace Horde.Build.Services
 {
 	using AgentSoftwareChannelName = StringId<AgentSoftwareChannels>;
-	using AgentSoftwareVersion = StringId<IAgentSoftwareCollection>;
 	using IStream = Horde.Build.Models.IStream;
 	using JobId = ObjectId<IJob>;
 	using LogId = ObjectId<ILogFile>;
-	using SessionId = ObjectId<ISession>;
-	using StreamId = StringId<IStream>;
 	using RpcAgentCapabilities = HordeCommon.Rpc.Messages.AgentCapabilities;
 	using RpcDeviceCapabilities = HordeCommon.Rpc.Messages.DeviceCapabilities;
+	using SessionId = ObjectId<ISession>;
+	using StreamId = StringId<IStream>;
 
 	/// <summary>
 	/// Implements the Horde gRPC service for bots updating their status and dequeing work
@@ -47,87 +45,82 @@ namespace Horde.Build.Services
 		/// <summary>
 		/// Timeout before closing a long-polling request (client will retry again) 
 		/// </summary>
-		internal TimeSpan LongPollTimeout = TimeSpan.FromMinutes(9);
-
-		DatabaseService DatabaseService;
-		AclService AclService;
-		AgentService AgentService;
-		StreamService StreamService;
-		JobService JobService;
-		AgentSoftwareService AgentSoftwareService;
-		IArtifactCollection ArtifactCollection;
-		ILogFileService LogFileService;
-		CredentialService CredentialService;
-		PoolService PoolService;
-		LifetimeService LifetimeService;
-		IGraphCollection Graphs;
-		ITestDataCollection TestData;
-		IJobStepRefCollection JobStepRefCollection;
-		ConformTaskSource ConformTaskSource;
-		ILogger<RpcService> Logger;
+		internal TimeSpan _longPollTimeout = TimeSpan.FromMinutes(9);
+		readonly AclService _aclService;
+		readonly AgentService _agentService;
+		readonly StreamService _streamService;
+		readonly JobService _jobService;
+		readonly AgentSoftwareService _agentSoftwareService;
+		readonly IArtifactCollection _artifactCollection;
+		readonly ILogFileService _logFileService;
+		readonly PoolService _poolService;
+		readonly LifetimeService _lifetimeService;
+		readonly IGraphCollection _graphs;
+		readonly ITestDataCollection _testData;
+		readonly IJobStepRefCollection _jobStepRefCollection;
+		readonly ConformTaskSource _conformTaskSource;
+		readonly ILogger<RpcService> _logger;
 
 		/// <summary>
 		/// Constructor
 		/// </summary>
-		public RpcService(DatabaseService DatabaseService, AclService AclService, AgentService AgentService, StreamService StreamService, JobService JobService, AgentSoftwareService AgentSoftwareService, IArtifactCollection ArtifactCollection, ILogFileService LogFileService, CredentialService CredentialService, PoolService PoolService, LifetimeService LifetimeService, IGraphCollection Graphs, ITestDataCollection TestData, IJobStepRefCollection JobStepRefCollection, ConformTaskSource ConformTaskSource, ILogger<RpcService> Logger)
+		public RpcService(AclService aclService, AgentService agentService, StreamService streamService, JobService jobService, AgentSoftwareService agentSoftwareService, IArtifactCollection artifactCollection, ILogFileService logFileService, PoolService poolService, LifetimeService lifetimeService, IGraphCollection graphs, ITestDataCollection testData, IJobStepRefCollection jobStepRefCollection, ConformTaskSource conformTaskSource, ILogger<RpcService> logger)
 		{
-			this.DatabaseService = DatabaseService;
-			this.AclService = AclService;
-			this.AgentService = AgentService;
-			this.StreamService = StreamService;
-			this.JobService = JobService;
-			this.AgentSoftwareService = AgentSoftwareService;
-			this.ArtifactCollection = ArtifactCollection;
-			this.LogFileService = LogFileService;
-			this.CredentialService = CredentialService;
-			this.PoolService = PoolService;
-			this.LifetimeService = LifetimeService;
-			this.Graphs = Graphs;
-			this.TestData = TestData;
-			this.JobStepRefCollection = JobStepRefCollection;
-			this.ConformTaskSource = ConformTaskSource;
-			this.Logger = Logger;
+			_aclService = aclService;
+			_agentService = agentService;
+			_streamService = streamService;
+			_jobService = jobService;
+			_agentSoftwareService = agentSoftwareService;
+			_artifactCollection = artifactCollection;
+			_logFileService = logFileService;
+			_poolService = poolService;
+			_lifetimeService = lifetimeService;
+			_graphs = graphs;
+			_testData = testData;
+			_jobStepRefCollection = jobStepRefCollection;
+			_conformTaskSource = conformTaskSource;
+			_logger = logger;
 		}
 
 		/// <summary>
 		/// Waits until the server is terminating
 		/// </summary>
-		/// <param name="Reader">Request reader</param>
-		/// <param name="Writer">Response writer</param>
-		/// <param name="Context">Context for the call</param>
+		/// <param name="reader">Request reader</param>
+		/// <param name="writer">Response writer</param>
+		/// <param name="context">Context for the call</param>
 		/// <returns>Response object</returns>
-		public override async Task QueryServerState(IAsyncStreamReader<QueryServerStateRequest> Reader, IServerStreamWriter<QueryServerStateResponse> Writer, ServerCallContext Context)
+		public override async Task QueryServerState(IAsyncStreamReader<QueryServerStateRequest> reader, IServerStreamWriter<QueryServerStateResponse> writer, ServerCallContext context)
 		{
-			if (await Reader.MoveNext())
+			if (await reader.MoveNext())
 			{
-				QueryServerStateRequest Request = Reader.Current;
-				Logger.LogInformation("Start server query for client {Name}", Request.Name);
+				QueryServerStateRequest request = reader.Current;
+				_logger.LogInformation("Start server query for client {Name}", request.Name);
 
 				// Return the current response
-				QueryServerStateResponse Response = new QueryServerStateResponse();
-				Response.Name = Dns.GetHostName();
-				await Writer.WriteAsync(Response);
+				QueryServerStateResponse response = new QueryServerStateResponse();
+				response.Name = Dns.GetHostName();
+				await writer.WriteAsync(response);
 
 				// Move to the next request from the client. This should always be the end of the stream, but will not occur until the client stops requesting responses.
-				Task<bool> MoveNextTask = Reader.MoveNext();
+				Task<bool> moveNextTask = reader.MoveNext();
 
 				// Wait for the client to close the stream or a shutdown to start
-				Task LongPollDelay = Task.Delay(LongPollTimeout);
-				Task WaitTask = await Task.WhenAny(MoveNextTask, LifetimeService.StoppingTask, LongPollDelay);
+				Task longPollDelay = Task.Delay(_longPollTimeout);
+				Task waitTask = await Task.WhenAny(moveNextTask, _lifetimeService.StoppingTask, longPollDelay);
 
-				if (WaitTask == MoveNextTask)
+				if (waitTask == moveNextTask)
 				{
 					throw new Exception("Unexpected request to QueryServerState posted from client.");
 				}
-				else if (WaitTask == LifetimeService.StoppingTask)
+				else if (waitTask == _lifetimeService.StoppingTask)
 				{
-					Logger.LogInformation("Notifying client {Name} of server shutdown", Request.Name);
-					await Writer.WriteAsync(Response);
+					_logger.LogInformation("Notifying client {Name} of server shutdown", request.Name);
+					await writer.WriteAsync(response);
 				}
-				else if (WaitTask == LongPollDelay)
+				else if (waitTask == longPollDelay)
 				{
 					// Send same response as server shutdown. In the agent perspective, they will be identical.
-					await Writer.WriteAsync(Response);
+					await writer.WriteAsync(response);
 				}
 			}
 		}
@@ -135,44 +128,44 @@ namespace Horde.Build.Services
 		/// <summary>
 		/// Waits until the server is terminating
 		/// </summary>
-		/// <param name="Reader">Request reader</param>
-		/// <param name="Writer">Response writer</param>
-		/// <param name="Context">Context for the call</param>
+		/// <param name="reader">Request reader</param>
+		/// <param name="writer">Response writer</param>
+		/// <param name="context">Context for the call</param>
 		/// <returns>Response object</returns>
-		public override async Task QueryServerStateV2(IAsyncStreamReader<QueryServerStateRequest> Reader, IServerStreamWriter<QueryServerStateResponse> Writer, ServerCallContext Context)
+		public override async Task QueryServerStateV2(IAsyncStreamReader<QueryServerStateRequest> reader, IServerStreamWriter<QueryServerStateResponse> writer, ServerCallContext context)
 		{
-			if (await Reader.MoveNext())
+			if (await reader.MoveNext())
 			{
-				QueryServerStateRequest Request = Reader.Current;
-				Logger.LogDebug("Start server query for client {Name}", Request.Name);
+				QueryServerStateRequest request = reader.Current;
+				_logger.LogDebug("Start server query for client {Name}", request.Name);
 
 				try
 				{
 					// Return the current response
-					QueryServerStateResponse Response = new QueryServerStateResponse();
-					Response.Name = Dns.GetHostName();
-					Response.Stopping = LifetimeService.IsStopping;
-					await Writer.WriteAsync(Response);
+					QueryServerStateResponse response = new QueryServerStateResponse();
+					response.Name = Dns.GetHostName();
+					response.Stopping = _lifetimeService.IsStopping;
+					await writer.WriteAsync(response);
 
 					// Move to the next request from the client. This should always be the end of the stream, but will not occur until the client stops requesting responses.
-					Task<bool> MoveNextTask = Reader.MoveNext();
+					Task<bool> moveNextTask = reader.MoveNext();
 
 					// Wait for the client to close the stream or a shutdown to start
-					if (await Task.WhenAny(MoveNextTask, LifetimeService.StoppingTask) == LifetimeService.StoppingTask)
+					if (await Task.WhenAny(moveNextTask, _lifetimeService.StoppingTask) == _lifetimeService.StoppingTask)
 					{
-						Response.Stopping = true;
-						await Writer.WriteAsync(Response);
+						response.Stopping = true;
+						await writer.WriteAsync(response);
 					}
 
 					// Wait until the client has finished sending
-					while (await MoveNextTask)
+					while (await moveNextTask)
 					{
-						MoveNextTask = Reader.MoveNext();
+						moveNextTask = reader.MoveNext();
 					}
 				}
-				catch (Exception Ex)
+				catch (Exception ex)
 				{
-					Logger.LogError(Ex, "Exception in QueryServerState for {Name}", Request.Name);
+					_logger.LogError(ex, "Exception in QueryServerState for {Name}", request.Name);
 					throw;
 				}
 			}
@@ -181,69 +174,69 @@ namespace Horde.Build.Services
 		/// <summary>
 		/// Updates the workspaces synced for an agent
 		/// </summary>
-		/// <param name="Request">The request parameters</param>
-		/// <param name="Context">Context for the call</param>
+		/// <param name="request">The request parameters</param>
+		/// <param name="context">Context for the call</param>
 		/// <returns>Response object</returns>
-		public override async Task<UpdateAgentWorkspacesResponse> UpdateAgentWorkspaces(UpdateAgentWorkspacesRequest Request, ServerCallContext Context)
+		public override async Task<UpdateAgentWorkspacesResponse> UpdateAgentWorkspaces(UpdateAgentWorkspacesRequest request, ServerCallContext context)
 		{
 			for (; ; )
 			{
 				// Get the current agent state
-				IAgent? Agent = await AgentService.GetAgentAsync(new AgentId(Request.AgentId));
-				if (Agent == null)
+				IAgent? agent = await _agentService.GetAgentAsync(new AgentId(request.AgentId));
+				if (agent == null)
 				{
-					throw new StructuredRpcException(StatusCode.OutOfRange, "Agent {AgentId} does not exist", Request.AgentId);
+					throw new StructuredRpcException(StatusCode.OutOfRange, "Agent {AgentId} does not exist", request.AgentId);
 				}
 
 				// Get the new workspaces
-				List<AgentWorkspace> NewWorkspaces = Request.Workspaces.Select(x => new AgentWorkspace(x)).ToList();
+				List<AgentWorkspace> newWorkspaces = request.Workspaces.Select(x => new AgentWorkspace(x)).ToList();
 
 				// Get the set of workspaces that are currently required
-				HashSet<AgentWorkspace> ConformWorkspaces = await PoolService.GetWorkspacesAsync(Agent, DateTime.UtcNow);
-				bool bPendingConform = !ConformWorkspaces.SetEquals(NewWorkspaces) || (Agent.RequestFullConform && !Request.RemoveUntrackedFiles);
+				HashSet<AgentWorkspace> conformWorkspaces = await _poolService.GetWorkspacesAsync(agent, DateTime.UtcNow);
+				bool bPendingConform = !conformWorkspaces.SetEquals(newWorkspaces) || (agent.RequestFullConform && !request.RemoveUntrackedFiles);
 
 				// Update the workspaces
-				if (await AgentService.TryUpdateWorkspacesAsync(Agent, NewWorkspaces, bPendingConform))
+				if (await _agentService.TryUpdateWorkspacesAsync(agent, newWorkspaces, bPendingConform))
 				{
-					UpdateAgentWorkspacesResponse Response = new UpdateAgentWorkspacesResponse();
+					UpdateAgentWorkspacesResponse response = new UpdateAgentWorkspacesResponse();
 					if (bPendingConform)
 					{
-						Response.Retry = await ConformTaskSource.GetWorkspacesAsync(Agent, Response.PendingWorkspaces);
-						Response.RemoveUntrackedFiles = Request.RemoveUntrackedFiles || Agent.RequestFullConform;
+						response.Retry = await _conformTaskSource.GetWorkspacesAsync(agent, response.PendingWorkspaces);
+						response.RemoveUntrackedFiles = request.RemoveUntrackedFiles || agent.RequestFullConform;
 					}
-					return Response;
+					return response;
 				}
 			}
 		}
 
-		static void CopyPropertyToResource(string Name, List<string> Properties, Dictionary<string, int> Resources)
+		static void CopyPropertyToResource(string name, List<string> properties, Dictionary<string, int> resources)
 		{
-			foreach (string Property in Properties)
+			foreach (string property in properties)
 			{
-				if (Property.Length > Name.Length && Property.StartsWith(Name, StringComparison.OrdinalIgnoreCase) && Property[Name.Length] == '=')
+				if (property.Length > name.Length && property.StartsWith(name, StringComparison.OrdinalIgnoreCase) && property[name.Length] == '=')
 				{
-					int Value;
-					if (int.TryParse(Property.AsSpan(Name.Length + 1), out Value))
+					int value;
+					if (Int32.TryParse(property.AsSpan(name.Length + 1), out value))
 					{
-						Resources[Name] = Value;
+						resources[name] = value;
 					}
 				}
 			}
 		}
 
-		static void GetCapabilities(RpcAgentCapabilities? Capabilities, out List<string> Properties, out Dictionary<string, int> Resources)
+		static void GetCapabilities(RpcAgentCapabilities? capabilities, out List<string> properties, out Dictionary<string, int> resources)
 		{
-			Properties = new List<string>();
-			Resources = new Dictionary<string, int>();
+			properties = new List<string>();
+			resources = new Dictionary<string, int>();
 
-			if (Capabilities != null && Capabilities.Devices.Count > 0)
+			if (capabilities != null && capabilities.Devices.Count > 0)
 			{
-				RpcDeviceCapabilities Device = Capabilities.Devices[0];
-				if (Device.Properties != null)
+				RpcDeviceCapabilities device = capabilities.Devices[0];
+				if (device.Properties != null)
 				{
-					Properties = new List<string>(Device.Properties);
-					CopyPropertyToResource(KnownPropertyNames.LogicalCores, Properties, Resources);
-					CopyPropertyToResource(KnownPropertyNames.RAM, Properties, Resources);
+					properties = new List<string>(device.Properties);
+					CopyPropertyToResource(KnownPropertyNames.LogicalCores, properties, resources);
+					CopyPropertyToResource(KnownPropertyNames.Ram, properties, resources);
 				}
 			}
 		}
@@ -251,125 +244,129 @@ namespace Horde.Build.Services
 		/// <summary>
 		/// Creates a new session
 		/// </summary>
-		/// <param name="Request">Request to create a new agent</param>
-		/// <param name="Context">Context for the RPC call</param>
+		/// <param name="request">Request to create a new agent</param>
+		/// <param name="context">Context for the RPC call</param>
 		/// <returns>Information about the new agent</returns>
-		public async override Task<CreateSessionResponse> CreateSession(CreateSessionRequest Request, ServerCallContext Context)
+		public override async Task<CreateSessionResponse> CreateSession(CreateSessionRequest request, ServerCallContext context)
 		{
-			if (Request.Capabilities == null)
+			if (request.Capabilities == null)
 			{
 				throw new StructuredRpcException(StatusCode.InvalidArgument, "Capabilities may not be null");
 			}
 
-			AgentId AgentId = new AgentId(Request.Name);
-			using IDisposable Scope = Logger.BeginScope("CreateSession({AgentId})", AgentId.ToString());
+			AgentId agentId = new AgentId(request.Name);
+			using IDisposable scope = _logger.BeginScope("CreateSession({AgentId})", agentId.ToString());
 
 			// Find the agent
-			IAgent? Agent = await AgentService.GetAgentAsync(AgentId);
-			if (Agent == null)
+			IAgent? agent = await _agentService.GetAgentAsync(agentId);
+			if (agent == null)
 			{
-				if (!await AclService.AuthorizeAsync(AclAction.CreateAgent, Context.GetHttpContext().User))
+				if (!await _aclService.AuthorizeAsync(AclAction.CreateAgent, context.GetHttpContext().User))
 				{
 					throw new StructuredRpcException(StatusCode.PermissionDenied, "User is not authenticated to create new agents");
 				}
 
-				const bool bEnabled = true;
-				Agent = await AgentService.CreateAgentAsync(Request.Name, bEnabled, null, null);
+				const bool BEnabled = true;
+				agent = await _agentService.CreateAgentAsync(request.Name, BEnabled, null, null);
 			}
 
 			// Make sure we're allowed to create sessions on this agent
-			if (!await AgentService.AuthorizeAsync(Agent, AclAction.CreateSession, Context.GetHttpContext().User, null))
+			if (!await _agentService.AuthorizeAsync(agent, AclAction.CreateSession, context.GetHttpContext().User, null))
 			{
-				throw new StructuredRpcException(StatusCode.PermissionDenied, "User is not authenticated to create session for {AgentId}", Request.Name);
+				throw new StructuredRpcException(StatusCode.PermissionDenied, "User is not authenticated to create session for {AgentId}", request.Name);
 			}
 
 			// Get the known properties for this agent
-			GetCapabilities(Request.Capabilities, out List<string> Properties, out Dictionary<string, int> Resources);
+			GetCapabilities(request.Capabilities, out List<string> properties, out Dictionary<string, int> resources);
 
 			// Create a new session
-			Agent = await AgentService.CreateSessionAsync(Agent, Request.Status, Properties, Resources, Request.Version);
-			if (Agent == null)
+			agent = await _agentService.CreateSessionAsync(agent, request.Status, properties, resources, request.Version);
+			if (agent == null)
 			{
-				throw new StructuredRpcException(StatusCode.NotFound, "Agent {AgentId} not found", Request.Name);
+				throw new StructuredRpcException(StatusCode.NotFound, "Agent {AgentId} not found", request.Name);
 			}
 
 			// Create the response
-			CreateSessionResponse Response = new CreateSessionResponse();
-			Response.AgentId = Agent.Id.ToString();
-			Response.SessionId = Agent.SessionId.ToString();
-			Response.ExpiryTime = Timestamp.FromDateTime(Agent.SessionExpiresAt!.Value);
-			Response.Token = AgentService.IssueSessionToken(Agent.Id, Agent.SessionId!.Value);
-			return Response;
+			CreateSessionResponse response = new CreateSessionResponse();
+			response.AgentId = agent.Id.ToString();
+			response.SessionId = agent.SessionId.ToString();
+			response.ExpiryTime = Timestamp.FromDateTime(agent.SessionExpiresAt!.Value);
+			response.Token = _agentService.IssueSessionToken(agent.Id, agent.SessionId!.Value);
+			return response;
 		}
 
 		/// <summary>
 		/// Updates an agent session
 		/// </summary>
-		/// <param name="Reader">Request to create a new agent</param>
-		/// <param name="Writer">Writer for response objects</param>
-		/// <param name="Context">Context for the RPC call</param>
+		/// <param name="reader">Request to create a new agent</param>
+		/// <param name="writer">Writer for response objects</param>
+		/// <param name="context">Context for the RPC call</param>
 		/// <returns>Information about the new agent</returns>
-		public async override Task UpdateSession(IAsyncStreamReader<UpdateSessionRequest> Reader, IServerStreamWriter<UpdateSessionResponse> Writer, ServerCallContext Context)
+		public override async Task UpdateSession(IAsyncStreamReader<UpdateSessionRequest> reader, IServerStreamWriter<UpdateSessionResponse> writer, ServerCallContext context)
 		{
 			// Read the request object
-			Task<bool> NextRequestTask = Reader.MoveNext();
-			if (await NextRequestTask)
+			Task<bool> nextRequestTask = reader.MoveNext();
+			if (await nextRequestTask)
 			{
-				UpdateSessionRequest Request = Reader.Current;
-				using IDisposable Scope = Logger.BeginScope("UpdateSession for agent {AgentId}, session {SessionId}", Request.AgentId, Request.SessionId);
+				UpdateSessionRequest request = reader.Current;
+				using IDisposable scope = _logger.BeginScope("UpdateSession for agent {AgentId}, session {SessionId}", request.AgentId, request.SessionId);
 
-				Logger.LogDebug("Updating session for {AgentId}", Request.AgentId);
-				foreach (HordeCommon.Rpc.Messages.Lease Lease in Request.Leases)
+				_logger.LogDebug("Updating session for {AgentId}", request.AgentId);
+				foreach (HordeCommon.Rpc.Messages.Lease lease in request.Leases)
 				{
-					Logger.LogDebug("Session {SessionId}, Lease {LeaseId} - State: {LeaseState}, Outcome: {LeaseOutcome}", Request.SessionId, Lease.Id, Lease.State, Lease.Outcome);
+					_logger.LogDebug("Session {SessionId}, Lease {LeaseId} - State: {LeaseState}, Outcome: {LeaseOutcome}", request.SessionId, lease.Id, lease.State, lease.Outcome);
 				}
 
 				// Get a task for moving to the next item. This will only complete once the call has closed.
-				using CancellationTokenSource CancellationSource = CancellationTokenSource.CreateLinkedTokenSource(Context.CancellationToken);
-				NextRequestTask = Reader.MoveNext();
-				NextRequestTask = NextRequestTask.ContinueWith(Task => { CancellationSource.Cancel(); return Task.Result; }, TaskScheduler.Current);
+				using CancellationTokenSource cancellationSource = CancellationTokenSource.CreateLinkedTokenSource(context.CancellationToken);
+				nextRequestTask = reader.MoveNext();
+				nextRequestTask = nextRequestTask.ContinueWith(task => 
+				{ 
+					cancellationSource.Cancel(); 
+					return task.Result; 
+				}, TaskScheduler.Current);
 
 				// Get the current agent state
-				IAgent? Agent = await AgentService.GetAgentAsync(new AgentId(Request.AgentId));
-				if(Agent != null)
+				IAgent? agent = await _agentService.GetAgentAsync(new AgentId(request.AgentId));
+				if(agent != null)
 				{
 					// Check we're authorized to update it
-					if (!AgentService.AuthorizeSession(Agent, Context.GetHttpContext().User))
+					if (!_agentService.AuthorizeSession(agent, context.GetHttpContext().User))
 					{
-						throw new StructuredRpcException(StatusCode.PermissionDenied, "Not authenticated for {AgentId}", Request.AgentId);
+						throw new StructuredRpcException(StatusCode.PermissionDenied, "Not authenticated for {AgentId}", request.AgentId);
 					}
 
 					// Get the new capabilities of this agent
-					List<string>? Properties = null;
-					Dictionary<string, int>? Resources = null;
-					if (Request.Capabilities != null)
+					List<string>? properties = null;
+					Dictionary<string, int>? resources = null;
+					if (request.Capabilities != null)
 					{
-						GetCapabilities(Request.Capabilities, out Properties, out Resources);
+						GetCapabilities(request.Capabilities, out properties, out resources);
 					}
 
 					// Update the session
-					Agent = await AgentService.UpdateSessionWithWaitAsync(Agent, SessionId.Parse(Request.SessionId), Request.Status, Properties, Resources, Request.Leases, CancellationSource.Token);
+					agent = await _agentService.UpdateSessionWithWaitAsync(agent, SessionId.Parse(request.SessionId), request.Status, properties, resources, request.Leases, cancellationSource.Token);
 				}
 
 				// Handle the invalid agent case
-				if (Agent == null)
+				if (agent == null)
 				{
-					throw new StructuredRpcException(StatusCode.NotFound, "Invalid agent name '{AgentId}'", Request.AgentId);
+					throw new StructuredRpcException(StatusCode.NotFound, "Invalid agent name '{AgentId}'", request.AgentId);
 				}
 
 				// Create the new session info
-				if (!Context.CancellationToken.IsCancellationRequested)
+				if (!context.CancellationToken.IsCancellationRequested)
 				{
-					UpdateSessionResponse Response = new UpdateSessionResponse();
-					Response.Leases.Add(Agent.Leases.Select(x => x.ToRpcMessage()));
-					Response.ExpiryTime = (Agent.SessionExpiresAt == null) ? new Timestamp() : Timestamp.FromDateTime(Agent.SessionExpiresAt.Value);
-					await Writer.WriteAsync(Response);
+					UpdateSessionResponse response = new UpdateSessionResponse();
+					response.Leases.Add(agent.Leases.Select(x => x.ToRpcMessage()));
+					response.ExpiryTime = (agent.SessionExpiresAt == null) ? new Timestamp() : Timestamp.FromDateTime(agent.SessionExpiresAt.Value);
+					await writer.WriteAsync(response);
 				}
 
 				// Wait for the client to close the stream
-				while (await NextRequestTask)
+				while (await nextRequestTask)
 				{
-					NextRequestTask = Reader.MoveNext();
+					nextRequestTask = reader.MoveNext();
 				}
 			}
 		}
@@ -377,169 +374,168 @@ namespace Horde.Build.Services
 		/// <summary>
 		/// Gets information about a stream
 		/// </summary>
-		/// <param name="Request">Request arguments</param>
-		/// <param name="Context">Context for the RPC call</param>
+		/// <param name="request">Request arguments</param>
+		/// <param name="context">Context for the RPC call</param>
 		/// <returns>Information about the new agent</returns>
-		public override async Task<GetStreamResponse> GetStream(GetStreamRequest Request, ServerCallContext Context)
+		public override async Task<GetStreamResponse> GetStream(GetStreamRequest request, ServerCallContext context)
 		{
-			StreamId StreamIdValue = new StreamId(Request.StreamId);
+			StreamId streamIdValue = new StreamId(request.StreamId);
 
-			IStream? Stream = await StreamService.GetStreamAsync(StreamIdValue);
-			if (Stream == null)
+			IStream? stream = await _streamService.GetStreamAsync(streamIdValue);
+			if (stream == null)
 			{
-				throw new StructuredRpcException(StatusCode.NotFound, "Stream {StreamId} does not exist", Request.StreamId);
+				throw new StructuredRpcException(StatusCode.NotFound, "Stream {StreamId} does not exist", request.StreamId);
 			}
-			if (!await StreamService.AuthorizeAsync(Stream, AclAction.ViewStream, Context.GetHttpContext().User, null))
+			if (!await _streamService.AuthorizeAsync(stream, AclAction.ViewStream, context.GetHttpContext().User, null))
 			{
-				throw new StructuredRpcException(StatusCode.PermissionDenied, "Not authenticated to access stream {StreamId}", Request.StreamId);
+				throw new StructuredRpcException(StatusCode.PermissionDenied, "Not authenticated to access stream {StreamId}", request.StreamId);
 			}
 
-			return Stream.ToRpcResponse();
+			return stream.ToRpcResponse();
 		}
 
 		/// <summary>
 		/// Gets information about a job
 		/// </summary>
-		/// <param name="Request">Request arguments</param>
-		/// <param name="Context">Context for the RPC call</param>
+		/// <param name="request">Request arguments</param>
+		/// <param name="context">Context for the RPC call</param>
 		/// <returns>Information about the new agent</returns>
-		public async override Task<GetJobResponse> GetJob(GetJobRequest Request, ServerCallContext Context)
+		public override async Task<GetJobResponse> GetJob(GetJobRequest request, ServerCallContext context)
 		{
-			JobId JobIdValue = new JobId(Request.JobId.ToObjectId());
+			JobId jobIdValue = new JobId(request.JobId.ToObjectId());
 
-			IJob? Job = await JobService.GetJobAsync(JobIdValue);
-			if (Job == null)
+			IJob? job = await _jobService.GetJobAsync(jobIdValue);
+			if (job == null)
 			{
-				throw new StructuredRpcException(StatusCode.NotFound, "Job {JobId} does not exist", Request.JobId);
+				throw new StructuredRpcException(StatusCode.NotFound, "Job {JobId} does not exist", request.JobId);
 			}
-			if (!JobService.AuthorizeSession(Job, Context.GetHttpContext().User))
+			if (!JobService.AuthorizeSession(job, context.GetHttpContext().User))
 			{
-				throw new StructuredRpcException(StatusCode.PermissionDenied, "Not authenticated to access job {JobId}", Request.JobId);
+				throw new StructuredRpcException(StatusCode.PermissionDenied, "Not authenticated to access job {JobId}", request.JobId);
 			}
 
-			return Job.ToRpcResponse();
+			return job.ToRpcResponse();
 		}
 
 		/// <summary>
 		/// Updates properties on a job
 		/// </summary>
-		/// <param name="Request">Request arguments</param>
-		/// <param name="Context">Context for the RPC call</param>
+		/// <param name="request">Request arguments</param>
+		/// <param name="context">Context for the RPC call</param>
 		/// <returns>Information about the new agent</returns>
-		public async override Task<Empty> UpdateJob(UpdateJobRequest Request, ServerCallContext Context)
+		public override async Task<Empty> UpdateJob(UpdateJobRequest request, ServerCallContext context)
 		{
-			JobId JobIdValue = new JobId(Request.JobId);
+			JobId jobIdValue = new JobId(request.JobId);
 
-			IJob? Job = await JobService.GetJobAsync(JobIdValue);
-			if (Job == null)
+			IJob? job = await _jobService.GetJobAsync(jobIdValue);
+			if (job == null)
 			{
-				throw new StructuredRpcException(StatusCode.NotFound, "Job {JobId} does not exist", Request.JobId);
+				throw new StructuredRpcException(StatusCode.NotFound, "Job {JobId} does not exist", request.JobId);
 			}
-			if (!JobService.AuthorizeSession(Job, Context.GetHttpContext().User))
+			if (!JobService.AuthorizeSession(job, context.GetHttpContext().User))
 			{
-				throw new StructuredRpcException(StatusCode.PermissionDenied, "Not authenticated to modify job {JobId}", Request.JobId);
+				throw new StructuredRpcException(StatusCode.PermissionDenied, "Not authenticated to modify job {JobId}", request.JobId);
 			}
 
-			await JobService.UpdateJobAsync(Job, Name: Request.Name);
+			await _jobService.UpdateJobAsync(job, name: request.Name);
 			return new Empty();
 		}
 
 		/// <summary>
 		/// Starts executing a batch
 		/// </summary>
-		/// <param name="Request">Request arguments</param>
-		/// <param name="Context">Context for the RPC call</param>
+		/// <param name="request">Request arguments</param>
+		/// <param name="context">Context for the RPC call</param>
 		/// <returns>Information about the new agent</returns>
-		public async override Task<BeginBatchResponse> BeginBatch(BeginBatchRequest Request, ServerCallContext Context)
+		public override async Task<BeginBatchResponse> BeginBatch(BeginBatchRequest request, ServerCallContext context)
 		{
-			ObjectId JobId = Request.JobId.ToObjectId();
-			SubResourceId BatchId = Request.BatchId.ToSubResourceId();
+			SubResourceId batchId = request.BatchId.ToSubResourceId();
 
-			IJob? Job = await JobService.GetJobAsync(new JobId(Request.JobId));
-			if (Job == null)
+			IJob? job = await _jobService.GetJobAsync(new JobId(request.JobId));
+			if (job == null)
 			{
-				throw new StructuredRpcException(StatusCode.NotFound, "Job {JobId} not found", Request.JobId);
+				throw new StructuredRpcException(StatusCode.NotFound, "Job {JobId} not found", request.JobId);
 			}
 
-			IJobStepBatch Batch = AuthorizeBatch(Job, Request.BatchId.ToSubResourceId(), Context);
-			Job = await JobService.UpdateBatchAsync(Job, BatchId, null, Api.JobStepBatchState.Starting);
+			IJobStepBatch batch = AuthorizeBatch(job, request.BatchId.ToSubResourceId(), context);
+			job = await _jobService.UpdateBatchAsync(job, batchId, null, Api.JobStepBatchState.Starting);
 
-			if (Job == null)
+			if (job == null)
 			{
-				throw new StructuredRpcException(StatusCode.NotFound, "Batch {JobId}:{BatchId} not found for updating", Request.JobId, Request.BatchId);
+				throw new StructuredRpcException(StatusCode.NotFound, "Batch {JobId}:{BatchId} not found for updating", request.JobId, request.BatchId);
 			}
 
-			IGraph Graph = await JobService.GetGraphAsync(Job);
+			IGraph graph = await _jobService.GetGraphAsync(job);
 
-			BeginBatchResponse Response = new BeginBatchResponse();
-			Response.LogId = Batch.LogId.ToString();
-			Response.AgentType = Graph.Groups[Batch.GroupIdx].AgentType;
-			return Response;
+			BeginBatchResponse response = new BeginBatchResponse();
+			response.LogId = batch.LogId.ToString();
+			response.AgentType = graph.Groups[batch.GroupIdx].AgentType;
+			return response;
 		}
 
 		/// <summary>
 		/// Finishes executing a batch
 		/// </summary>
-		/// <param name="Request">Request arguments</param>
-		/// <param name="Context">Context for the RPC call</param>
+		/// <param name="request">Request arguments</param>
+		/// <param name="context">Context for the RPC call</param>
 		/// <returns>Information about the new agent</returns>
-		public async override Task<Empty> FinishBatch(FinishBatchRequest Request, ServerCallContext Context)
+		public override async Task<Empty> FinishBatch(FinishBatchRequest request, ServerCallContext context)
 		{
-			IJob? Job = await JobService.GetJobAsync(new JobId(Request.JobId));
-			if (Job == null)
+			IJob? job = await _jobService.GetJobAsync(new JobId(request.JobId));
+			if (job == null)
 			{
-				throw new StructuredRpcException(StatusCode.NotFound, "Job {JobId} not found", Request.JobId);
+				throw new StructuredRpcException(StatusCode.NotFound, "Job {JobId} not found", request.JobId);
 			}
 
-			IJobStepBatch Batch = AuthorizeBatch(Job, Request.BatchId.ToSubResourceId(), Context);
-			await JobService.UpdateBatchAsync(Job, Batch.Id, null, Api.JobStepBatchState.Complete);
+			IJobStepBatch batch = AuthorizeBatch(job, request.BatchId.ToSubResourceId(), context);
+			await _jobService.UpdateBatchAsync(job, batch.Id, null, Api.JobStepBatchState.Complete);
 			return new Empty();
 		}
 
 		/// <summary>
 		/// Starts executing a step
 		/// </summary>
-		/// <param name="Request">Request arguments</param>
-		/// <param name="Context">Context for the RPC call</param>
+		/// <param name="request">Request arguments</param>
+		/// <param name="context">Context for the RPC call</param>
 		/// <returns>Information about the new agent</returns>
-		public async override Task<BeginStepResponse> BeginStep(BeginStepRequest Request, ServerCallContext Context)
+		public override async Task<BeginStepResponse> BeginStep(BeginStepRequest request, ServerCallContext context)
 		{
-			Boxed<ILogFile?> Log = new Boxed<ILogFile?>(null);
+			Boxed<ILogFile?> log = new Boxed<ILogFile?>(null);
 			for (; ; )
 			{
-				BeginStepResponse? Response = await TryBeginStep(Request, Log, Context);
-				if (Response != null)
+				BeginStepResponse? response = await TryBeginStep(request, log, context);
+				if (response != null)
 				{
-					return Response;
+					return response;
 				}
 			}
 		}
 
-		async Task<BeginStepResponse?> TryBeginStep(BeginStepRequest Request, Boxed<ILogFile?> Log, ServerCallContext Context)
+		async Task<BeginStepResponse?> TryBeginStep(BeginStepRequest request, Boxed<ILogFile?> log, ServerCallContext context)
 		{
 			// Check the job exists and we can access it
-			IJob? Job = await JobService.GetJobAsync(new JobId(Request.JobId));
-			if (Job == null)
+			IJob? job = await _jobService.GetJobAsync(new JobId(request.JobId));
+			if (job == null)
 			{
-				throw new StructuredRpcException(StatusCode.NotFound, "Job {JobId} not found", Request.JobId);
+				throw new StructuredRpcException(StatusCode.NotFound, "Job {JobId} not found", request.JobId);
 			}
 
 			// Find the batch being executed
-			IJobStepBatch Batch = AuthorizeBatch(Job, Request.BatchId.ToSubResourceId(), Context);
-			if (Batch.State != Api.JobStepBatchState.Starting && Batch.State != Api.JobStepBatchState.Running)
+			IJobStepBatch batch = AuthorizeBatch(job, request.BatchId.ToSubResourceId(), context);
+			if (batch.State != Api.JobStepBatchState.Starting && batch.State != Api.JobStepBatchState.Running)
 			{
 				return new BeginStepResponse { State = BeginStepResponse.Types.Result.Complete };
 			}
 
 			// Figure out which step to execute next
-			IJobStep? Step;
-			for (int StepIdx = 0; ; StepIdx++)
+			IJobStep? step;
+			for (int stepIdx = 0; ; stepIdx++)
 			{
 				// If there aren't any more steps, send a complete message
-				if (StepIdx == Batch.Steps.Count)
+				if (stepIdx == batch.Steps.Count)
 				{
-					Logger.LogDebug("Job {JobId} batch {BatchId} is complete", Job.Id, Batch.Id);
-					if (await JobService.TryUpdateBatchAsync(Job, Batch.Id, NewState: Api.JobStepBatchState.Stopping) == null)
+					_logger.LogDebug("Job {JobId} batch {BatchId} is complete", job.Id, batch.Id);
+					if (await _jobService.TryUpdateBatchAsync(job, batch.Id, newState: Api.JobStepBatchState.Stopping) == null)
 					{
 						return null;
 					}
@@ -547,30 +543,30 @@ namespace Horde.Build.Services
 				}
 
 				// Check if this step is ready to be executed
-				Step = Batch.Steps[StepIdx];
-				if (Step.State == JobStepState.Ready)
+				step = batch.Steps[stepIdx];
+				if (step.State == JobStepState.Ready)
 				{
 					break;
 				}
-				if (Step.State == JobStepState.Waiting)
+				if (step.State == JobStepState.Waiting)
 				{
-					Logger.LogDebug("Waiting for job {JobId}, batch {BatchId}, step {StepId}", Job.Id, Batch.Id, Step.Id);
+					_logger.LogDebug("Waiting for job {JobId}, batch {BatchId}, step {StepId}", job.Id, batch.Id, step.Id);
 					return new BeginStepResponse { State = BeginStepResponse.Types.Result.Waiting };
 				}
 			}
 
 			// Create a log file if necessary
-			if (Log.Value == null)
+			if (log.Value == null)
 			{
-				Log.Value = await LogFileService.CreateLogFileAsync(Job.Id, Batch.SessionId, Api.LogType.Json);
+				log.Value = await _logFileService.CreateLogFileAsync(job.Id, batch.SessionId, Api.LogType.Json);
 			}
 
 			// Get the node for this step
-			IGraph Graph = await JobService.GetGraphAsync(Job);
-			INode Node = Graph.Groups[Batch.GroupIdx].Nodes[Step.NodeIdx];
+			IGraph graph = await _jobService.GetGraphAsync(job);
+			INode node = graph.Groups[batch.GroupIdx].Nodes[step.NodeIdx];
 
 			// Figure out all the credentials for it (and check we can access them)
-			Dictionary<string, string> Credentials = new Dictionary<string, string>();
+			Dictionary<string, string> credentials = new Dictionary<string, string>();
 			//				if (Node.Credentials != null)
 			//				{
 			//					ClaimsPrincipal Principal = new ClaimsPrincipal(new ClaimsIdentity(Job.Claims.Select(x => new Claim(x.Type, x.Value))));
@@ -582,244 +578,181 @@ namespace Horde.Build.Services
 			//				}
 
 			// Update the step state
-			IJob? NewJob = await JobService.TryUpdateStepAsync(Job, Batch.Id, Step.Id, JobStepState.Running, JobStepOutcome.Unspecified, null, null, Log.Value.Id, null, null, null, null);
-			if (NewJob != null)
-			{
-				BeginStepResponse Response = new BeginStepResponse();
-				Response.State = BeginStepResponse.Types.Result.Ready;
-				Response.LogId = Log.Value.Id.ToString();
-				Response.StepId = Step.Id.ToString();
-				Response.Name = Node.Name;
-				Response.Credentials.Add(Credentials);
-
-				IJobStepRef? LastStep = await JobStepRefCollection.GetPrevStepForNodeAsync(Job.StreamId, Job.TemplateId, Node.Name, Job.Change);
-				if (LastStep != null)
+			IJob? newJob = await _jobService.TryUpdateStepAsync(job, batch.Id, step.Id, JobStepState.Running, JobStepOutcome.Unspecified, null, null, log.Value.Id, null, null, null, null);
+			if (newJob != null)
+			{
+				BeginStepResponse response = new BeginStepResponse();
+				response.State = BeginStepResponse.Types.Result.Ready;
+				response.LogId = log.Value.Id.ToString();
+				response.StepId = step.Id.ToString();
+				response.Name = node.Name;
+				response.Credentials.Add(credentials);
+
+				IJobStepRef? lastStep = await _jobStepRefCollection.GetPrevStepForNodeAsync(job.StreamId, job.TemplateId, node.Name, job.Change);
+				if (lastStep != null)
 				{
-					Response.EnvVars.Add("UE_HORDE_LAST_CL", LastStep.Change.ToString(CultureInfo.InvariantCulture));
+					response.EnvVars.Add("UE_HORDE_LAST_CL", lastStep.Change.ToString(CultureInfo.InvariantCulture));
 
-					if (LastStep.Outcome == JobStepOutcome.Success)
+					if (lastStep.Outcome == JobStepOutcome.Success)
 					{
-						Response.EnvVars.Add("UE_HORDE_LAST_SUCCESS_CL", LastStep.Change.ToString(CultureInfo.InvariantCulture));
+						response.EnvVars.Add("UE_HORDE_LAST_SUCCESS_CL", lastStep.Change.ToString(CultureInfo.InvariantCulture));
 					}
-					else if (LastStep.LastSuccess != null)
+					else if (lastStep.LastSuccess != null)
 					{
-						Response.EnvVars.Add("UE_HORDE_LAST_SUCCESS_CL", LastStep.LastSuccess.Value.ToString(CultureInfo.InvariantCulture));
+						response.EnvVars.Add("UE_HORDE_LAST_SUCCESS_CL", lastStep.LastSuccess.Value.ToString(CultureInfo.InvariantCulture));
 					}
 
-					if (LastStep.Outcome == JobStepOutcome.Success || LastStep.Outcome == JobStepOutcome.Warnings)
+					if (lastStep.Outcome == JobStepOutcome.Success || lastStep.Outcome == JobStepOutcome.Warnings)
 					{
-						Response.EnvVars.Add("UE_HORDE_LAST_WARNING_CL", LastStep.Change.ToString(CultureInfo.InvariantCulture));
+						response.EnvVars.Add("UE_HORDE_LAST_WARNING_CL", lastStep.Change.ToString(CultureInfo.InvariantCulture));
 					}
-					else if (LastStep.LastWarning != null)
+					else if (lastStep.LastWarning != null)
 					{
-						Response.EnvVars.Add("UE_HORDE_LAST_WARNING_CL", LastStep.LastWarning.Value.ToString(CultureInfo.InvariantCulture));
+						response.EnvVars.Add("UE_HORDE_LAST_WARNING_CL", lastStep.LastWarning.Value.ToString(CultureInfo.InvariantCulture));
 					}
 				}
 
-				if (Node.Properties != null)
+				if (node.Properties != null)
 				{
-					Response.Properties.Add(Node.Properties);
+					response.Properties.Add(node.Properties);
 				}
-				Response.Warnings = Node.Warnings;
-				return Response;
+				response.Warnings = node.Warnings;
+				return response;
 			}
 
 			return null;
 		}
 
-		/// <summary>
-		/// Gets all the required credentials for the given step
-		/// </summary>
-		/// <param name="User">The user to validate</param>
-		/// <param name="Node">The node being executed</param>
-		/// <param name="Credentials">Receives a list of credentials for the step</param>
-		/// <param name="WriteError">Delegate used to write error messages</param>
-		/// <returns>Async task</returns>
-		private async Task<bool> GetCredentialsForStep(ClaimsPrincipal User, INode Node, Dictionary<string, string> Credentials, Func<string, Task> WriteError)
-		{
-			if (Node.Credentials != null)
-			{
-				Dictionary<string, Credential> Cache = new Dictionary<string, Credential>(StringComparer.OrdinalIgnoreCase);
-				foreach (KeyValuePair<string, string> Pair in Node.Credentials)
-				{
-					// Get the credential path. We expect this in the format <CredentialName>.<PropertyName>
-					string Path = Pair.Value;
-
-					// Find the separator
-					int Idx = Path.IndexOf('.', StringComparison.Ordinal);
-					if (Idx == -1)
-					{
-						await WriteError($"Invalid credential path '{Path}'. Requested credentials should be in the form '<CredentialName>.<PropertyName>'.\n");
-						return false;
-					}
-
-					// Split the path into credential and property names
-					string CredentialName = Path.Substring(0, Idx);
-					string PropertyName = Path.Substring(Idx + 1);
-
-					// Try to get the credential with this name
-					Credential? Credential;
-					if (!Cache.TryGetValue(CredentialName, out Credential))
-					{
-						Credential = await CredentialService.GetCredentialAsync(CredentialName);
-						if (Credential == null)
-						{
-							await WriteError($"No credential called '{CredentialName}' could be found");
-							return false;
-						}
-						if (!await CredentialService.AuthorizeAsync(Credential, AclAction.ViewCredential, User, null))
-						{
-							await WriteError($"User is not allowed to view credential '{CredentialName}'");
-							return false;
-						}
-						Cache[CredentialName] = Credential;
-					}
-
-					// Get the property
-					string? PropertyValue;
-					if (!Credential!.Properties.TryGetValue(PropertyName, out PropertyValue))
-					{
-						await WriteError($"No property called '{PropertyName}' found in credential '{CredentialName}'");
-						return false;
-					}
-
-					// Add it to the output list
-					Credentials.Add(Pair.Key, PropertyValue);
-				}
-			}
-			return true;
-		}
-
-		async Task<IJob> GetJobAsync(JobId JobId)
+		async Task<IJob> GetJobAsync(JobId jobId)
 		{
-			IJob? Job = await JobService.GetJobAsync(JobId);
-			if (Job == null)
+			IJob? job = await _jobService.GetJobAsync(jobId);
+			if (job == null)
 			{
-				throw new StructuredRpcException(StatusCode.NotFound, "Job {JobId} not found", JobId);
+				throw new StructuredRpcException(StatusCode.NotFound, "Job {JobId} not found", jobId);
 			}
-			return Job;
+			return job;
 		}
 
-		static IJobStepBatch AuthorizeBatch(IJob Job, SubResourceId BatchId, ServerCallContext Context)
+		static IJobStepBatch AuthorizeBatch(IJob job, SubResourceId batchId, ServerCallContext context)
 		{
-			IJobStepBatch? Batch;
-			if (!Job.TryGetBatch(BatchId, out Batch))
+			IJobStepBatch? batch;
+			if (!job.TryGetBatch(batchId, out batch))
 			{
-				throw new StructuredRpcException(StatusCode.NotFound, "Unable to find batch {JobId}:{BatchId}", Job.Id, BatchId);
+				throw new StructuredRpcException(StatusCode.NotFound, "Unable to find batch {JobId}:{BatchId}", job.Id, batchId);
 			}
-			if (Batch.SessionId == null)
+			if (batch.SessionId == null)
 			{
-				throw new StructuredRpcException(StatusCode.PermissionDenied, "Batch {JobId}:{BatchId} has no session id", Job.Id, BatchId);
+				throw new StructuredRpcException(StatusCode.PermissionDenied, "Batch {JobId}:{BatchId} has no session id", job.Id, batchId);
 			}
 
-			ClaimsPrincipal Principal = Context.GetHttpContext().User;
-			if (!Principal.HasSessionClaim(Batch.SessionId.Value))
+			ClaimsPrincipal principal = context.GetHttpContext().User;
+			if (!principal.HasSessionClaim(batch.SessionId.Value))
 			{
-				throw new StructuredRpcException(StatusCode.PermissionDenied, "Session id {SessionId} not valid for batch {JobId}:{BatchId}. Expected {ExpectedSessionId}.", Principal.GetSessionClaim() ?? SessionId.Empty, Job.Id, BatchId, Batch.SessionId.Value);
+				throw new StructuredRpcException(StatusCode.PermissionDenied, "Session id {SessionId} not valid for batch {JobId}:{BatchId}. Expected {ExpectedSessionId}.", principal.GetSessionClaim() ?? SessionId.Empty, job.Id, batchId, batch.SessionId.Value);
 			}
 
-			return Batch;
+			return batch;
 		}
 
 		/// <summary>
 		/// Updates the state of a jobstep
 		/// </summary>
-		/// <param name="Request">Request arguments</param>
-		/// <param name="Context">Context for the RPC call</param>
+		/// <param name="request">Request arguments</param>
+		/// <param name="context">Context for the RPC call</param>
 		/// <returns>Information about the new agent</returns>
-		public async override Task<Empty> UpdateStep(UpdateStepRequest Request, ServerCallContext Context)
+		public override async Task<Empty> UpdateStep(UpdateStepRequest request, ServerCallContext context)
 		{
-			IJob? Job = await JobService.GetJobAsync(new JobId(Request.JobId));
-			if (Job == null)
+			IJob? job = await _jobService.GetJobAsync(new JobId(request.JobId));
+			if (job == null)
 			{
-				throw new StructuredRpcException(StatusCode.NotFound, "Job {JobId} not found", Request.JobId);
+				throw new StructuredRpcException(StatusCode.NotFound, "Job {JobId} not found", request.JobId);
 			}
 
-			IJobStepBatch Batch = AuthorizeBatch(Job, Request.BatchId.ToSubResourceId(), Context);
+			IJobStepBatch batch = AuthorizeBatch(job, request.BatchId.ToSubResourceId(), context);
 
-			await JobService.UpdateStepAsync(Job, Batch.Id, Request.StepId.ToSubResourceId(), Request.State, Request.Outcome, null, null, null, null, null);
+			await _jobService.UpdateStepAsync(job, batch.Id, request.StepId.ToSubResourceId(), request.State, request.Outcome, null, null, null, null, null);
 			return new Empty();
 		}
 
 		/// <summary>
 		/// Get the state of a jobstep
 		/// </summary>
-		/// <param name="Request">Request arguments</param>
-		/// <param name="Context">Context for the RPC call</param>
+		/// <param name="request">Request arguments</param>
+		/// <param name="context">Context for the RPC call</param>
 		/// <returns>Information about the step</returns>
-		public async override Task<GetStepResponse> GetStep(GetStepRequest Request, ServerCallContext Context)
+		public override async Task<GetStepResponse> GetStep(GetStepRequest request, ServerCallContext context)
 		{
-			IJob Job = await GetJobAsync(new JobId(Request.JobId));
-			IJobStepBatch Batch = AuthorizeBatch(Job, Request.BatchId.ToSubResourceId(), Context);
+			IJob job = await GetJobAsync(new JobId(request.JobId));
+			IJobStepBatch batch = AuthorizeBatch(job, request.BatchId.ToSubResourceId(), context);
 
-			SubResourceId StepId = Request.StepId.ToSubResourceId();
-			if (!Batch.TryGetStep(StepId, out IJobStep? Step))
+			SubResourceId stepId = request.StepId.ToSubResourceId();
+			if (!batch.TryGetStep(stepId, out IJobStep? step))
 			{
-				throw new StructuredRpcException(StatusCode.NotFound, "Unable to find step {JobId}:{BatchId}:{StepId}", Job.Id, Batch.Id, StepId);
+				throw new StructuredRpcException(StatusCode.NotFound, "Unable to find step {JobId}:{BatchId}:{StepId}", job.Id, batch.Id, stepId);
 			}
 
-			return new GetStepResponse { Outcome = Step.Outcome, State = Step.State, AbortRequested = Step.AbortRequested };
+			return new GetStepResponse { Outcome = step.Outcome, State = step.State, AbortRequested = step.AbortRequested };
 		}
 
 		/// <summary>
 		/// Updates the state of a jobstep
 		/// </summary>
-		/// <param name="Request">Request arguments</param>
-		/// <param name="Context">Context for the RPC call</param>
+		/// <param name="request">Request arguments</param>
+		/// <param name="context">Context for the RPC call</param>
 		/// <returns>Information about the new agent</returns>
-		public async override Task<UpdateGraphResponse> UpdateGraph(UpdateGraphRequest Request, ServerCallContext Context)
+		public override async Task<UpdateGraphResponse> UpdateGraph(UpdateGraphRequest request, ServerCallContext context)
 		{
-			List<NewGroup> NewGroups = new List<NewGroup>();
-			foreach (CreateGroupRequest Group in Request.Groups)
+			List<NewGroup> newGroups = new List<NewGroup>();
+			foreach (CreateGroupRequest group in request.Groups)
 			{
-				List<NewNode> NewNodes = new List<NewNode>();
-				foreach (CreateNodeRequest Node in Group.Nodes)
+				List<NewNode> newNodes = new List<NewNode>();
+				foreach (CreateNodeRequest node in group.Nodes)
 				{
-					NewNode NewNode = new NewNode(Node.Name, Node.InputDependencies.ToList(), Node.OrderDependencies.ToList(), Node.Priority, Node.AllowRetry, Node.RunEarly, Node.Warnings, new Dictionary<string, string>(Node.Credentials), new Dictionary<string, string>(Node.Properties));
-					NewNodes.Add(NewNode);
+					NewNode newNode = new NewNode(node.Name, node.InputDependencies.ToList(), node.OrderDependencies.ToList(), node.Priority, node.AllowRetry, node.RunEarly, node.Warnings, new Dictionary<string, string>(node.Credentials), new Dictionary<string, string>(node.Properties));
+					newNodes.Add(newNode);
 				}
-				NewGroups.Add(new NewGroup(Group.AgentType, NewNodes));
+				newGroups.Add(new NewGroup(group.AgentType, newNodes));
 			}
 
-			List<NewAggregate> NewAggregates = new List<NewAggregate>();
-			foreach (CreateAggregateRequest Aggregate in Request.Aggregates)
+			List<NewAggregate> newAggregates = new List<NewAggregate>();
+			foreach (CreateAggregateRequest aggregate in request.Aggregates)
 			{
-				NewAggregate NewAggregate = new NewAggregate(Aggregate.Name, Aggregate.Nodes.ToList());
-				NewAggregates.Add(NewAggregate);
+				NewAggregate newAggregate = new NewAggregate(aggregate.Name, aggregate.Nodes.ToList());
+				newAggregates.Add(newAggregate);
 			}
 
-			List<NewLabel> NewLabels = new List<NewLabel>();
-			foreach (CreateLabelRequest Label in Request.Labels)
+			List<NewLabel> newLabels = new List<NewLabel>();
+			foreach (CreateLabelRequest label in request.Labels)
 			{
-				NewLabel NewLabel = new NewLabel();
-				NewLabel.DashboardName = String.IsNullOrEmpty(Label.DashboardName) ? null : Label.DashboardName;
-				NewLabel.DashboardCategory = String.IsNullOrEmpty(Label.DashboardCategory) ? null : Label.DashboardCategory;
-				NewLabel.UgsName = String.IsNullOrEmpty(Label.UgsName) ? null : Label.UgsName;
-				NewLabel.UgsProject = String.IsNullOrEmpty(Label.UgsProject) ? null : Label.UgsProject;
-				NewLabel.Change = Label.Change;
-				NewLabel.RequiredNodes = Label.RequiredNodes.ToList();
-				NewLabel.IncludedNodes = Label.IncludedNodes.ToList();
-				NewLabels.Add(NewLabel);
+				NewLabel newLabel = new NewLabel();
+				newLabel.DashboardName = String.IsNullOrEmpty(label.DashboardName) ? null : label.DashboardName;
+				newLabel.DashboardCategory = String.IsNullOrEmpty(label.DashboardCategory) ? null : label.DashboardCategory;
+				newLabel.UgsName = String.IsNullOrEmpty(label.UgsName) ? null : label.UgsName;
+				newLabel.UgsProject = String.IsNullOrEmpty(label.UgsProject) ? null : label.UgsProject;
+				newLabel.Change = label.Change;
+				newLabel.RequiredNodes = label.RequiredNodes.ToList();
+				newLabel.IncludedNodes = label.IncludedNodes.ToList();
+				newLabels.Add(newLabel);
 			}
 
-			JobId JobIdValue = new JobId(Request.JobId);
+			JobId jobIdValue = new JobId(request.JobId);
 			for (; ; )
 			{
-				IJob? Job = await JobService.GetJobAsync(JobIdValue);
-				if (Job == null)
+				IJob? job = await _jobService.GetJobAsync(jobIdValue);
+				if (job == null)
 				{
 					throw new StructuredRpcException(StatusCode.NotFound, "Resource not found");
 				}
-				if (!JobService.AuthorizeSession(Job, Context.GetHttpContext().User))
+				if (!JobService.AuthorizeSession(job, context.GetHttpContext().User))
 				{
 					throw new StructuredRpcException(StatusCode.PermissionDenied, "Access denied");
 				}
 
-				IGraph Graph = await JobService.GetGraphAsync(Job);
-				Graph = await Graphs.AppendAsync(Graph, NewGroups, NewAggregates, NewLabels);
+				IGraph graph = await _jobService.GetGraphAsync(job);
+				graph = await _graphs.AppendAsync(graph, newGroups, newAggregates, newLabels);
 
-				IJob? NewJob = await JobService.TryUpdateGraphAsync(Job, Graph);
-				if (NewJob != null)
+				IJob? newJob = await _jobService.TryUpdateGraphAsync(job, graph);
+				if (newJob != null)
 				{
 					return new UpdateGraphResponse();
 				}
@@ -829,185 +762,185 @@ namespace Horde.Build.Services
 		/// <summary>
 		/// Creates a set of events
 		/// </summary>
-		/// <param name="Request">Request arguments</param>
-		/// <param name="Context">Context for the RPC call</param>
+		/// <param name="request">Request arguments</param>
+		/// <param name="context">Context for the RPC call</param>
 		/// <returns>Information about the new agent</returns>
-		public async override Task<Empty> CreateEvents(CreateEventsRequest Request, ServerCallContext Context)
+		public override async Task<Empty> CreateEvents(CreateEventsRequest request, ServerCallContext context)
 		{
-			if (!await AclService.AuthorizeAsync(AclAction.CreateEvent, Context.GetHttpContext().User))
+			if (!await _aclService.AuthorizeAsync(AclAction.CreateEvent, context.GetHttpContext().User))
 			{
 				throw new StructuredRpcException(StatusCode.NotFound, "Access denied");
 			}
 
-			List<NewLogEventData> NewEvents = new List<NewLogEventData>();
-			foreach (CreateEventRequest Event in Request.Events)
+			List<NewLogEventData> newEvents = new List<NewLogEventData>();
+			foreach (CreateEventRequest createEvent in request.Events)
 			{
-				NewLogEventData NewEvent = new NewLogEventData();
-				NewEvent.LogId = new LogId(Event.LogId);
-				NewEvent.Severity = Event.Severity;
-				NewEvent.LineIndex = Event.LineIndex;
-				NewEvent.LineCount = Event.LineCount;
-				NewEvents.Add(NewEvent);
+				NewLogEventData newEvent = new NewLogEventData();
+				newEvent.LogId = new LogId(createEvent.LogId);
+				newEvent.Severity = createEvent.Severity;
+				newEvent.LineIndex = createEvent.LineIndex;
+				newEvent.LineCount = createEvent.LineCount;
+				newEvents.Add(newEvent);
 			}
-			await LogFileService.CreateEventsAsync(NewEvents);
+			await _logFileService.CreateEventsAsync(newEvents);
 			return new Empty();
 		}
 
 		/// <summary>
 		/// Writes output to a log file
 		/// </summary>
-		/// <param name="Request">Request arguments</param>
-		/// <param name="Context">Context for the RPC call</param>
+		/// <param name="request">Request arguments</param>
+		/// <param name="context">Context for the RPC call</param>
 		/// <returns>Information about the new agent</returns>
-		public async override Task<Empty> WriteOutput(WriteOutputRequest Request, ServerCallContext Context)
+		public override async Task<Empty> WriteOutput(WriteOutputRequest request, ServerCallContext context)
 		{
-			ILogFile? LogFile = await LogFileService.GetCachedLogFileAsync(new LogId(Request.LogId));
-			if (LogFile == null)
+			ILogFile? logFile = await _logFileService.GetCachedLogFileAsync(new LogId(request.LogId));
+			if (logFile == null)
 			{
 				throw new StructuredRpcException(StatusCode.NotFound, "Resource not found");
 			}
-			if (!Horde.Build.Services.LogFileService.AuthorizeForSession(LogFile, Context.GetHttpContext().User))
+			if (!Horde.Build.Services.LogFileService.AuthorizeForSession(logFile, context.GetHttpContext().User))
 			{
 				throw new StructuredRpcException(StatusCode.PermissionDenied, "Access denied");
 			}
 
-			await LogFileService.WriteLogDataAsync(LogFile, Request.Offset, Request.LineIndex, Request.Data.ToArray(), Request.Flush);
+			await _logFileService.WriteLogDataAsync(logFile, request.Offset, request.LineIndex, request.Data.ToArray(), request.Flush);
 			return new Empty();
 		}
 
 		/// <summary>
 		/// Uploads a new agent archive
 		/// </summary>
-		/// <param name="Request">Request arguments</param>
-		/// <param name="Context">Context for the RPC call</param>
+		/// <param name="request">Request arguments</param>
+		/// <param name="context">Context for the RPC call</param>
 		/// <returns>Information about the new agent</returns>
-		public async override Task<UploadSoftwareResponse> UploadSoftware(UploadSoftwareRequest Request, ServerCallContext Context)
+		public override async Task<UploadSoftwareResponse> UploadSoftware(UploadSoftwareRequest request, ServerCallContext context)
 		{
-			if (!await AclService.AuthorizeAsync(AclAction.UploadSoftware, Context.GetHttpContext().User))
+			if (!await _aclService.AuthorizeAsync(AclAction.UploadSoftware, context.GetHttpContext().User))
 			{
 				throw new StructuredRpcException(StatusCode.PermissionDenied, "Access to software is forbidden");
 			}
 
-			string Version = await AgentSoftwareService.SetArchiveAsync(new AgentSoftwareChannelName(Request.Channel), null, Request.Data.ToArray());
+			string version = await _agentSoftwareService.SetArchiveAsync(new AgentSoftwareChannelName(request.Channel), null, request.Data.ToArray());
 
-			UploadSoftwareResponse Response = new UploadSoftwareResponse();
-			Response.Version = Version;
-			return Response;
+			UploadSoftwareResponse response = new UploadSoftwareResponse();
+			response.Version = version;
+			return response;
 		}
 
 		/// <summary>
 		/// Downloads a new agent archive
 		/// </summary>
-		/// <param name="Request">Request arguments</param>
-		/// <param name="ResponseStream">Writer for the output data</param>
-		/// <param name="Context">Context for the RPC call</param>
+		/// <param name="request">Request arguments</param>
+		/// <param name="responseStream">Writer for the output data</param>
+		/// <param name="context">Context for the RPC call</param>
 		/// <returns>Information about the new agent</returns>
-		public async override Task DownloadSoftware(DownloadSoftwareRequest Request, IServerStreamWriter<DownloadSoftwareResponse> ResponseStream, ServerCallContext Context)
+		public override async Task DownloadSoftware(DownloadSoftwareRequest request, IServerStreamWriter<DownloadSoftwareResponse> responseStream, ServerCallContext context)
 		{
-			if (!await AclService.AuthorizeAsync(AclAction.DownloadSoftware, Context.GetHttpContext().User))
+			if (!await _aclService.AuthorizeAsync(AclAction.DownloadSoftware, context.GetHttpContext().User))
 			{
 				throw new StructuredRpcException(StatusCode.NotFound, "Access to software is forbidden");
 			}
 
-			byte[]? Data = await AgentSoftwareService.GetArchiveAsync(Request.Version);
-			if (Data == null)
+			byte[]? data = await _agentSoftwareService.GetArchiveAsync(request.Version);
+			if (data == null)
 			{
 				throw new StructuredRpcException(StatusCode.NotFound, "Missing version {Version}");
 			}
 
-			for (int Offset = 0; Offset < Data.Length;)
+			for (int offset = 0; offset < data.Length;)
 			{
-				int NextOffset = Math.Min(Offset + 128 * 1024, Data.Length);
+				int nextOffset = Math.Min(offset + 128 * 1024, data.Length);
 
-				DownloadSoftwareResponse Response = new DownloadSoftwareResponse();
-				Response.Data = Google.Protobuf.ByteString.CopyFrom(Data.AsSpan(Offset, NextOffset - Offset));
+				DownloadSoftwareResponse response = new DownloadSoftwareResponse();
+				response.Data = Google.Protobuf.ByteString.CopyFrom(data.AsSpan(offset, nextOffset - offset));
 
-				await ResponseStream.WriteAsync(Response);
+				await responseStream.WriteAsync(response);
 
-				Offset = NextOffset;
+				offset = nextOffset;
 			}
 		}
 
 		/// <summary>
 		/// Uploads a new artifact
 		/// </summary>
-		/// <param name="Reader">Request arguments</param>
-		/// <param name="Context">Context for the RPC call</param>
+		/// <param name="reader">Request arguments</param>
+		/// <param name="context">Context for the RPC call</param>
 		/// <returns>Information about the new agent</returns>
-		public async override Task<UploadArtifactResponse> UploadArtifact(IAsyncStreamReader<UploadArtifactRequest> Reader, ServerCallContext Context)
+		public override async Task<UploadArtifactResponse> UploadArtifact(IAsyncStreamReader<UploadArtifactRequest> reader, ServerCallContext context)
 		{
 			// Advance to the metadata object
-			if (!await Reader.MoveNext())
+			if (!await reader.MoveNext())
 			{
 				throw new StructuredRpcException(StatusCode.DataLoss, "Missing request for artifact upload");
 			}
 
 			// Read the request object
-			UploadArtifactMetadata? Metadata = Reader.Current.Metadata;
-			if (Metadata == null)
+			UploadArtifactMetadata? metadata = reader.Current.Metadata;
+			if (metadata == null)
 			{
 				throw new StructuredRpcException(StatusCode.DataLoss, "Expected metadata in first artifact request");
 			}
 
 			// Get the job and step
-			IJob Job = await GetJobAsync(new JobId(Metadata.JobId));
-			IJobStepBatch Batch = AuthorizeBatch(Job, Metadata.BatchId.ToSubResourceId(), Context);
+			IJob job = await GetJobAsync(new JobId(metadata.JobId));
+			AuthorizeBatch(job, metadata.BatchId.ToSubResourceId(), context);
 
-			IJobStep? Step;
-			if (!Job.TryGetStep(Metadata.BatchId.ToSubResourceId(), Metadata.StepId.ToSubResourceId(), out Step))
+			IJobStep? step;
+			if (!job.TryGetStep(metadata.BatchId.ToSubResourceId(), metadata.StepId.ToSubResourceId(), out step))
 			{
-				throw new StructuredRpcException(StatusCode.NotFound, "Unable to find step {JobId}:{BatchId}:{StepId}", Job.Id, Metadata.BatchId, Metadata.StepId);
+				throw new StructuredRpcException(StatusCode.NotFound, "Unable to find step {JobId}:{BatchId}:{StepId}", job.Id, metadata.BatchId, metadata.StepId);
 			}
 
 			// Upload the stream
-			using (ArtifactChunkStream InputStream = new ArtifactChunkStream(Reader, Metadata.Length))
+			using (ArtifactChunkStream inputStream = new ArtifactChunkStream(reader, metadata.Length))
 			{
-				IArtifact Artifact = await ArtifactCollection.CreateArtifactAsync(Job.Id, Step.Id, Metadata.Name, Metadata.MimeType, InputStream);
+				IArtifact artifact = await _artifactCollection.CreateArtifactAsync(job.Id, step.Id, metadata.Name, metadata.MimeType, inputStream);
 
-				UploadArtifactResponse Response = new UploadArtifactResponse();
-				Response.Id = Artifact.Id.ToString();
-				return Response;
+				UploadArtifactResponse response = new UploadArtifactResponse();
+				response.Id = artifact.Id.ToString();
+				return response;
 			}
 		}
 
 		/// <summary>
 		/// Uploads new test data
 		/// </summary>
-		/// <param name="Reader">Request arguments</param>
-		/// <param name="Context">Context for the RPC call</param>
+		/// <param name="reader">Request arguments</param>
+		/// <param name="context">Context for the RPC call</param>
 		/// <returns>Information about the new agent</returns>
-		public override async Task<UploadTestDataResponse> UploadTestData(IAsyncStreamReader<UploadTestDataRequest> Reader, ServerCallContext Context)
+		public override async Task<UploadTestDataResponse> UploadTestData(IAsyncStreamReader<UploadTestDataRequest> reader, ServerCallContext context)
 		{
-			IJob? Job = null;
-			IJobStep? JobStep = null;
+			IJob? job = null;
+			IJobStep? jobStep = null;
 
-			while (await Reader.MoveNext())
+			while (await reader.MoveNext())
 			{
-				UploadTestDataRequest Request = Reader.Current;
+				UploadTestDataRequest request = reader.Current;
 
-				JobId JobId = new JobId(Request.JobId);
-				if (Job == null || JobId != Job.Id)
+				JobId jobId = new JobId(request.JobId);
+				if (job == null || jobId != job.Id)
 				{
-					Job = await JobService.GetJobAsync(JobId);
-					if (Job == null)
+					job = await _jobService.GetJobAsync(jobId);
+					if (job == null)
 					{
-						throw new StructuredRpcException(StatusCode.NotFound, "Unable to find job {JobId}", JobId);
+						throw new StructuredRpcException(StatusCode.NotFound, "Unable to find job {JobId}", jobId);
 					}
-					JobStep = null;
+					jobStep = null;
 				}
 
-				SubResourceId JobStepId = Request.JobStepId.ToSubResourceId();
-				if (JobStep == null || JobStepId != JobStep.Id)
+				SubResourceId jobStepId = request.JobStepId.ToSubResourceId();
+				if (jobStep == null || jobStepId != jobStep.Id)
 				{
-					if (!Job.TryGetStep(JobStepId, out JobStep))
+					if (!job.TryGetStep(jobStepId, out jobStep))
 					{
-						throw new StructuredRpcException(StatusCode.NotFound, "Unable to find step {JobStepId} on job {JobId}", JobStepId, JobId);
+						throw new StructuredRpcException(StatusCode.NotFound, "Unable to find step {JobStepId} on job {JobId}", jobStepId, jobId);
 					}
 				}
 
-				string Text = Encoding.UTF8.GetString(Request.Value.ToArray());
-				BsonDocument Document = BsonSerializer.Deserialize<BsonDocument>(Text);
-				await TestData.AddAsync(Job, JobStep, Request.Key, Document);
+				string text = Encoding.UTF8.GetString(request.Value.ToArray());
+				BsonDocument document = BsonSerializer.Deserialize<BsonDocument>(text);
+				await _testData.AddAsync(job, jobStep, request.Key, document);
 			}
 
 			return new UploadTestDataResponse();
@@ -1016,24 +949,24 @@ namespace Horde.Build.Services
 		/// <summary>
 		/// Create a new report on a job or job step
 		/// </summary>
-		/// <param name="Request"></param>
-		/// <param name="Context"></param>
+		/// <param name="request"></param>
+		/// <param name="context"></param>
 		/// <returns></returns>
-		public override async Task<CreateReportResponse> CreateReport(CreateReportRequest Request, ServerCallContext Context)
+		public override async Task<CreateReportResponse> CreateReport(CreateReportRequest request, ServerCallContext context)
 		{
-			IJob Job = await GetJobAsync(new JobId(Request.JobId));
-			IJobStepBatch Batch = AuthorizeBatch(Job, Request.BatchId.ToSubResourceId(), Context);
+			IJob job = await GetJobAsync(new JobId(request.JobId));
+			IJobStepBatch batch = AuthorizeBatch(job, request.BatchId.ToSubResourceId(), context);
 
-			Report NewReport = new Report { Name = Request.Name, Placement = Request.Placement, ArtifactId = Request.ArtifactId.ToObjectId() };
-			if (Request.Scope == ReportScope.Job)
+			Report newReport = new Report { Name = request.Name, Placement = request.Placement, ArtifactId = request.ArtifactId.ToObjectId() };
+			if (request.Scope == ReportScope.Job)
 			{
-				Logger.LogDebug("Adding report to job {JobId}: {Name} -> {ArtifactId}", Job.Id, Request.Name, Request.ArtifactId);
-				await JobService.UpdateJobAsync(Job, Reports: new List<Report> { NewReport });
+				_logger.LogDebug("Adding report to job {JobId}: {Name} -> {ArtifactId}", job.Id, request.Name, request.ArtifactId);
+				await _jobService.UpdateJobAsync(job, reports: new List<Report> { newReport });
 			}
 			else
 			{
-				Logger.LogDebug("Adding report to step {JobId}:{BatchId}:{StepId}: {Name} -> {ArtifactId}", Job.Id, Batch.Id, Request.StepId, Request.Name, Request.ArtifactId);
-				await JobService.UpdateStepAsync(Job, Batch.Id, Request.StepId.ToSubResourceId(), NewReports: new List<Report> { NewReport });
+				_logger.LogDebug("Adding report to step {JobId}:{BatchId}:{StepId}: {Name} -> {ArtifactId}", job.Id, batch.Id, request.StepId, request.Name, request.ArtifactId);
+				await _jobService.UpdateStepAsync(job, batch.Id, request.StepId.ToSubResourceId(), newReports: new List<Report> { newReport });
 			}
 
 			return new CreateReportResponse();
