Horde: Fix coding conventions within Horde.Build.

diff --git a/Engine/Source/Programs/Horde/Horde.Build/Startup.cs b/Engine/Source/Programs/Horde/Horde.Build/Startup.cs
--- a/Engine/Source/Programs/Horde/Horde.Build/Startup.cs
+++ b/Engine/Source/Programs/Horde/Horde.Build/Startup.cs
@@ -1,187 +1,169 @@
-﻿// Copyright Epic Games, Inc. All Rights Reserved.
+// Copyright Epic Games, Inc. All Rights Reserved.
 
 using System;
 using System.Collections.Generic;
 using System.Diagnostics;
-using System.IdentityModel.Tokens.Jwt;
+using System.Globalization;
 using System.IO;
 using System.Linq;
 using System.Reflection;
-using System.Security.Claims;
-using System.Text.Encodings.Web;
+using System.Runtime.InteropServices;
 using System.Text.Json;
 using System.Text.Json.Serialization;
+using System.Threading;
 using System.Threading.Tasks;
-using Serilog;
+using Amazon.Extensions.NETCore.Setup;
+using Amazon.Runtime;
+using Amazon.SecurityToken.Model;
+using EpicGames.AspNet;
+using EpicGames.Core;
+using EpicGames.Horde.Storage;
+using Grpc.Core;
+using Grpc.Core.Interceptors;
+using Horde.Build.Acls;
+using Horde.Build.Authentication;
+using Horde.Build.Collections;
+using Horde.Build.Collections.Impl;
+using Horde.Build.Commits;
+using Horde.Build.Commits.Impl;
+using Horde.Build.Compute;
+using Horde.Build.Compute.Impl;
+using Horde.Build.Fleet.Autoscale;
+using Horde.Build.Jobs;
+using Horde.Build.Logs;
+using Horde.Build.Logs.Builder;
+using Horde.Build.Logs.Readers;
+using Horde.Build.Logs.Storage;
+using Horde.Build.Logs.Storage.Impl;
 using Horde.Build.Models;
+using Horde.Build.Notifications;
+using Horde.Build.Notifications.Impl;
 using Horde.Build.Services;
+using Horde.Build.Services.Impl;
+using Horde.Build.Storage;
+using Horde.Build.Storage.Backends;
+using Horde.Build.Tasks;
+using Horde.Build.Tasks.Impl;
+using Horde.Build.Utilities;
+using HordeCommon;
 using Microsoft.AspNetCore.Authentication;
 using Microsoft.AspNetCore.Authentication.Cookies;
 using Microsoft.AspNetCore.Authentication.JwtBearer;
-using Microsoft.AspNetCore.Authentication.OAuth.Claims;
 using Microsoft.AspNetCore.Authentication.OpenIdConnect;
 using Microsoft.AspNetCore.Authorization;
 using Microsoft.AspNetCore.Builder;
+using Microsoft.AspNetCore.DataProtection;
 using Microsoft.AspNetCore.Hosting;
+using Microsoft.AspNetCore.Hosting.Server.Features;
+using Microsoft.AspNetCore.Http;
+using Microsoft.AspNetCore.HttpOverrides;
+using Microsoft.AspNetCore.Mvc;
+using Microsoft.AspNetCore.Mvc.ModelBinding;
 using Microsoft.Extensions.Configuration;
 using Microsoft.Extensions.DependencyInjection;
+using Microsoft.Extensions.Diagnostics.HealthChecks;
 using Microsoft.Extensions.Hosting;
 using Microsoft.Extensions.Logging;
 using Microsoft.Extensions.Options;
-using Microsoft.IdentityModel.Protocols.OpenIdConnect;
-using Microsoft.IdentityModel.Tokens;
 using Microsoft.OpenApi.Models;
 using MongoDB.Bson;
 using MongoDB.Bson.Serialization;
 using MongoDB.Bson.Serialization.Conventions;
-using Microsoft.AspNetCore.HttpOverrides;
-using Horde.Build.Acls;
-using Horde.Build.Utilities;
-using Grpc.Core.Interceptors;
-using Grpc.Core;
-using EpicGames.Core;
-using Microsoft.Extensions.Diagnostics.HealthChecks;
-using Microsoft.AspNetCore.Diagnostics.HealthChecks;
-using Horde.Build.Collections;
-using Microsoft.AspNetCore.Http;
-using System.Diagnostics.CodeAnalysis;
-using System.Globalization;
-using System.Threading;
-using HordeCommon;
-using Microsoft.AspNetCore.DataProtection;
-using StackExchange.Redis;
-using Horde.Build.Storage;
-using Horde.Build.Logs;
-using Horde.Build.Logs.Builder;
-using Horde.Build.Logs.Storage;
-using Horde.Build.Logs.Readers;
-using Horde.Build.Logs.Storage.Impl;
-using Horde.Build.Collections.Impl;
-using Horde.Build.Services.Impl;
-using Microsoft.Extensions.DependencyInjection.Extensions;
-using Horde.Build.Authentication;
-using Horde.Build.Tasks.Impl;
-using Horde.Build.Tasks;
-using StatsdClient;
-using Status = Grpc.Core.Status;
-using Horde.Build.Notifications.Impl;
-using Horde.Build.Notifications;
-using Microsoft.AspNetCore.Hosting.Server.Features;
-using System.Runtime.InteropServices;
-using Amazon.Extensions.NETCore.Setup;
-using Amazon.S3;
-using Amazon.Runtime;
-using Horde.Build.Storage.Backends;
-using Horde.Build.Commits.Impl;
-using Horde.Build.Commits;
+using MongoDB.Bson.Serialization.Serializers;
 using OpenTracing.Contrib.Grpc.Interceptors;
 using OpenTracing.Util;
-using EpicGames.Horde.Compute;
-using Horde.Build.Compute.Impl;
-using Horde.Build.Compute;
-using System.Net.Http.Headers;
-using Amazon.SecurityToken.Model;
-using Horde.Build.Fleet.Autoscale;
+using Serilog;
 using Serilog.Events;
-using Horde.Build.Jobs;
-using Microsoft.AspNetCore.Mvc;
-using Microsoft.AspNetCore.Mvc.ModelBinding;
-using Horde.Build.Storage.Services;
-using EpicGames.Horde.Storage;
-using System.Net.Http;
-using EpicGames.AspNet;
-using EpicGames.Horde.Storage.Impl;
-using MongoDB.Bson.Serialization.Serializers;
+using StackExchange.Redis;
+using StatsdClient;
+using Status = Grpc.Core.Status;
 
 namespace Horde.Build
 {
-	using MessageTemplate = EpicGames.Core.MessageTemplate;
 	using JobId = ObjectId<IJob>;
 	using LeaseId = ObjectId<ILease>;
 	using LogId = ObjectId<ILogFile>;
 	using UserId = ObjectId<IUser>;
-	using ILogger = Microsoft.Extensions.Logging.ILogger;
 
 	class Startup
 	{
 		class GrpcExceptionInterceptor : Interceptor
 		{
-			ILogger<GrpcExceptionInterceptor> Logger;
+			readonly ILogger<GrpcExceptionInterceptor> _logger;
 
-			public GrpcExceptionInterceptor(ILogger<GrpcExceptionInterceptor> Logger)
+			public GrpcExceptionInterceptor(ILogger<GrpcExceptionInterceptor> logger)
 			{
-				this.Logger = Logger;
+				_logger = logger;
 			}
 
-			[SuppressMessage("Design", "CA1031:Do not catch general exception types", Justification = "<Pending>")]
-			public override Task<TResponse> UnaryServerHandler<TRequest, TResponse>(TRequest Request, ServerCallContext Context, UnaryServerMethod<TRequest, TResponse> Continuation)
+			public override Task<TResponse> UnaryServerHandler<TRequest, TResponse>(TRequest request, ServerCallContext context, UnaryServerMethod<TRequest, TResponse> continuation)
 			{
-				return Guard(Context, () => base.UnaryServerHandler(Request, Context, Continuation));
+				return Guard(context, () => base.UnaryServerHandler(request, context, continuation));
 			}
 
-			public override Task<TResponse> ClientStreamingServerHandler<TRequest, TResponse>(IAsyncStreamReader<TRequest> RequestStream, ServerCallContext Context, ClientStreamingServerMethod<TRequest, TResponse> Continuation) where TRequest : class where TResponse : class
+			public override Task<TResponse> ClientStreamingServerHandler<TRequest, TResponse>(IAsyncStreamReader<TRequest> requestStream, ServerCallContext context, ClientStreamingServerMethod<TRequest, TResponse> continuation) where TRequest : class where TResponse : class
 			{
-				return Guard(Context, () => base.ClientStreamingServerHandler(RequestStream, Context, Continuation));
+				return Guard(context, () => base.ClientStreamingServerHandler(requestStream, context, continuation));
 			}
 
-			public override Task ServerStreamingServerHandler<TRequest, TResponse>(TRequest Request, IServerStreamWriter<TResponse> ResponseStream, ServerCallContext Context, ServerStreamingServerMethod<TRequest, TResponse> Continuation) where TRequest : class where TResponse : class
+			public override Task ServerStreamingServerHandler<TRequest, TResponse>(TRequest request, IServerStreamWriter<TResponse> responseStream, ServerCallContext context, ServerStreamingServerMethod<TRequest, TResponse> continuation) where TRequest : class where TResponse : class
 			{
-				return Guard(Context, () => base.ServerStreamingServerHandler(Request, ResponseStream, Context, Continuation));
+				return Guard(context, () => base.ServerStreamingServerHandler(request, responseStream, context, continuation));
 			}
 
-			public override Task DuplexStreamingServerHandler<TRequest, TResponse>(IAsyncStreamReader<TRequest> RequestStream, IServerStreamWriter<TResponse> ResponseStream, ServerCallContext Context, DuplexStreamingServerMethod<TRequest, TResponse> Continuation) where TRequest : class where TResponse : class
+			public override Task DuplexStreamingServerHandler<TRequest, TResponse>(IAsyncStreamReader<TRequest> requestStream, IServerStreamWriter<TResponse> responseStream, ServerCallContext context, DuplexStreamingServerMethod<TRequest, TResponse> continuation) where TRequest : class where TResponse : class
 			{
-				return Guard(Context, () => base.DuplexStreamingServerHandler(RequestStream, ResponseStream, Context, Continuation));
+				return Guard(context, () => base.DuplexStreamingServerHandler(requestStream, responseStream, context, continuation));
 			}
 
-			async Task<T> Guard<T>(ServerCallContext Context, Func<Task<T>> CallFunc) where T : class
+			async Task<T> Guard<T>(ServerCallContext context, Func<Task<T>> callFunc) where T : class
 			{
-				T Result = null!;
-				await Guard(Context, async () => { Result = await CallFunc(); });
-				return Result;
+				T result = null!;
+				await Guard(context, async () => { result = await callFunc(); });
+				return result;
 			}
 
-			async Task Guard(ServerCallContext Context, Func<Task> CallFunc)
+			async Task Guard(ServerCallContext context, Func<Task> callFunc)
 			{
-				HttpContext HttpContext = Context.GetHttpContext();
+				HttpContext httpContext = context.GetHttpContext();
 
-				AgentId? AgentId = AclService.GetAgentId(HttpContext.User);
-				if (AgentId != null)
+				AgentId? agentId = AclService.GetAgentId(httpContext.User);
+				if (agentId != null)
 				{
-					using IDisposable Scope = Logger.BeginScope("Agent: {AgentId}, RemoteIP: {RemoteIP}, Method: {Method}", AgentId.Value, HttpContext.Connection.RemoteIpAddress, Context.Method);
-					await GuardInner(Context, CallFunc);
+					using IDisposable scope = _logger.BeginScope("Agent: {AgentId}, RemoteIP: {RemoteIP}, Method: {Method}", agentId.Value, httpContext.Connection.RemoteIpAddress, context.Method);
+					await GuardInner(context, callFunc);
 				}
 				else
 				{
-					using IDisposable Scope = Logger.BeginScope("RemoteIP: {RemoteIP}, Method: {Method}", HttpContext.Connection.RemoteIpAddress, Context.Method);
-					await GuardInner(Context, CallFunc);
+					using IDisposable scope = _logger.BeginScope("RemoteIP: {RemoteIP}, Method: {Method}", httpContext.Connection.RemoteIpAddress, context.Method);
+					await GuardInner(context, callFunc);
 				}
 			}
 
-			async Task GuardInner(ServerCallContext Context, Func<Task> CallFunc)
+			async Task GuardInner(ServerCallContext context, Func<Task> callFunc)
 			{
 				try
 				{
-					await CallFunc();
+					await callFunc();
 				}
-				catch (StructuredRpcException Ex)
+				catch (StructuredRpcException ex)
 				{
 #pragma warning disable CA2254 // Template should be a static expression
-					Logger.LogError(Ex, Ex.Format, Ex.Args);
+					_logger.LogError(ex, ex.Format, ex.Args);
 #pragma warning restore CA2254 // Template should be a static expression
 					throw;
 				}
-				catch (Exception Ex)
+				catch (Exception ex)
 				{
-					if (Context.CancellationToken.IsCancellationRequested)
+					if (context.CancellationToken.IsCancellationRequested)
 					{
-						Logger.LogInformation(Ex, "Call to method {Method} was cancelled", Context.Method);
+						_logger.LogInformation(ex, "Call to method {Method} was cancelled", context.Method);
 						throw;
 					}
 					else
 					{
-						Logger.LogError(Ex, "Exception in call to {Method}", Context.Method);
-						throw new RpcException(new Status(StatusCode.Internal, $"An exception was thrown on the server: {Ex}"));
+						_logger.LogError(ex, "Exception in call to {Method}", context.Method);
+						throw new RpcException(new Status(StatusCode.Internal, $"An exception was thrown on the server: {ex}"));
 					}
 				}
 			}
@@ -189,13 +171,13 @@ namespace Horde.Build
 
 		class BsonSerializationProvider : IBsonSerializationProvider
 		{
-			public IBsonSerializer? GetSerializer(Type Type)
+			public IBsonSerializer? GetSerializer(Type type)
 			{
-				if (Type == typeof(ContentHash))
+				if (type == typeof(ContentHash))
 				{
 					return new ContentHashSerializer();
 				}
-				if (Type == typeof(DateTimeOffset))
+				if (type == typeof(DateTimeOffset))
 				{
 					return new DateTimeOffsetStringSerializer();
 				}
@@ -205,466 +187,465 @@ namespace Horde.Build
 
 		class JsonObjectIdConverter : JsonConverter<ObjectId>
 		{
-			public override ObjectId Read(ref Utf8JsonReader Reader, Type TypeToConvert, JsonSerializerOptions Options)
+			public override ObjectId Read(ref Utf8JsonReader reader, Type typeToConvert, JsonSerializerOptions options)
 			{
-				string? String = Reader.GetString();
-				if (String == null)
+				string? str = reader.GetString();
+				if (str == null)
 				{
 					throw new InvalidDataException("Unable to parse object id");
 				}
-				return String.ToObjectId();
+				return str.ToObjectId();
 			}
 
-			public override void Write(Utf8JsonWriter Writer, ObjectId ObjectId, JsonSerializerOptions Options)
+			public override void Write(Utf8JsonWriter writer, ObjectId objectId, JsonSerializerOptions options)
 			{
-				Writer.WriteStringValue(ObjectId.ToString());
+				writer.WriteStringValue(objectId.ToString());
 			}
 		}
 		
 		class JsonDateTimeConverter : JsonConverter<DateTime>
 		{
-			public override DateTime Read(ref Utf8JsonReader Reader, Type TypeToConvert, JsonSerializerOptions Options)
+			public override DateTime Read(ref Utf8JsonReader reader, Type typeToConvert, JsonSerializerOptions options)
 			{
-				Debug.Assert(TypeToConvert == typeof(DateTime));
+				Debug.Assert(typeToConvert == typeof(DateTime));
 
-				string? String = Reader.GetString();
-				if (String == null)
+				string? str = reader.GetString();
+				if (str == null)
 				{
 					throw new InvalidDataException("Unable to parse DateTime");
 				}
-				return DateTime.Parse(String, CultureInfo.CurrentCulture);
+				return DateTime.Parse(str, CultureInfo.CurrentCulture);
 			}
 
-			public override void Write(Utf8JsonWriter Writer, DateTime DateTime, JsonSerializerOptions Options)
+			public override void Write(Utf8JsonWriter writer, DateTime dateTime, JsonSerializerOptions options)
 			{
-				Writer.WriteStringValue(DateTime.ToUniversalTime().ToString("yyyy'-'MM'-'dd'T'HH':'mm':'ssZ", CultureInfo.CurrentCulture));
+				writer.WriteStringValue(dateTime.ToUniversalTime().ToString("yyyy'-'MM'-'dd'T'HH':'mm':'ssZ", CultureInfo.CurrentCulture));
 			}
 		}
 
-		public Startup(IConfiguration Configuration)
+		public Startup(IConfiguration configuration)
 		{
-			this.Configuration = Configuration;
+			Configuration = configuration;
 		}
 
 		public IConfiguration Configuration { get; }
 
 		// This method gets called *multiple times* by the runtime. Use this method to add services to the container.
-		public void ConfigureServices(IServiceCollection Services)
+		public void ConfigureServices(IServiceCollection services)
 		{
 			// IOptionsMonitor pattern for live updating of configuration settings
-			Services.Configure<ServerSettings>(Configuration.GetSection("Horde"));
+			services.Configure<ServerSettings>(Configuration.GetSection("Horde"));
 
 			// Settings used for configuring services
-			IConfigurationSection ConfigSection = Configuration.GetSection("Horde");
-			ServerSettings Settings = new ServerSettings();
-			ConfigSection.Bind(Settings);
+			IConfigurationSection configSection = Configuration.GetSection("Horde");
+			ServerSettings settings = new ServerSettings();
+			configSection.Bind(settings);
 
-			Settings.Validate();
+			settings.Validate();
 			
-			Services.Configure<CommitServiceOptions>(ConfigSection.GetSection("Replication"));
+			services.Configure<CommitServiceOptions>(configSection.GetSection("Replication"));
 
-			if (Settings.GlobalThreadPoolMinSize != null)
+			if (settings.GlobalThreadPoolMinSize != null)
 			{
 				// Min thread pool size is set to combat timeouts seen with the Redis client.
 				// See comments for <see cref="ServerSettings.GlobalThreadPoolMinSize" /> and
 				// [REDACTED_PROFILE_URL]/StackExchange.Redis/issues/1680
-				int Min = Settings.GlobalThreadPoolMinSize.Value;
-				ThreadPool.SetMinThreads(Min, Min);
+				int min = settings.GlobalThreadPoolMinSize.Value;
+				ThreadPool.SetMinThreads(min, min);
 			}
 
 #pragma warning disable CA2000 // Dispose objects before losing scope
-			RedisService RedisService = new RedisService(Settings);
+			RedisService redisService = new RedisService(settings);
 #pragma warning restore CA2000 // Dispose objects before losing scope
-			Services.AddSingleton<RedisService>(SP => RedisService);
-			Services.AddSingleton<IDatabase>(RedisService.Database);
-			Services.AddSingleton<ConnectionMultiplexer>(SP => RedisService.Multiplexer);
-			Services.AddDataProtection().PersistKeysToStackExchangeRedis(() => RedisService.Database, "aspnet-data-protection");
+			services.AddSingleton<RedisService>(sp => redisService);
+			services.AddSingleton<IDatabase>(redisService.Database);
+			services.AddSingleton<ConnectionMultiplexer>(sp => redisService.Multiplexer);
+			services.AddDataProtection().PersistKeysToStackExchangeRedis(() => redisService.Database, "aspnet-data-protection");
 
-			if (Settings.CorsEnabled)
+			if (settings.CorsEnabled)
 			{
-				Services.AddCors(Options =>
+				services.AddCors(options =>
 				{
-					Options.AddPolicy("CorsPolicy",
-						Builder => Builder.WithOrigins(Settings.CorsOrigin.Split(";"))
+					options.AddPolicy("CorsPolicy",
+						builder => builder.WithOrigins(settings.CorsOrigin.Split(";"))
 						.AllowAnyMethod()
 						.AllowAnyHeader()
 						.AllowCredentials());
 				});
 			}
 
-			Services.AddGrpc(Options =>
+			services.AddGrpc(options =>
 			{
-				Options.EnableDetailedErrors = true;
-				Options.MaxReceiveMessageSize = 200 * 1024 * 1024; // 100 MB (packaged builds of Horde agent can be large) 
-				Options.Interceptors.Add(typeof(LifetimeGrpcInterceptor));
-				Options.Interceptors.Add(typeof(GrpcExceptionInterceptor));
-				Options.Interceptors.Add<ServerTracingInterceptor>(GlobalTracer.Instance);
+				options.EnableDetailedErrors = true;
+				options.MaxReceiveMessageSize = 200 * 1024 * 1024; // 100 MB (packaged builds of Horde agent can be large) 
+				options.Interceptors.Add(typeof(GrpcExceptionInterceptor));
+				options.Interceptors.Add<ServerTracingInterceptor>(GlobalTracer.Instance);
 			});
-			Services.AddGrpcReflection();
-
-			Services.AddSingleton<IAgentCollection, AgentCollection>();
-			Services.AddSingleton<IAgentSoftwareCollection, AgentSoftwareCollection>();
-			Services.AddSingleton<IArtifactCollection, ArtifactCollection>();
-			Services.AddSingleton<ICommitCollection, CommitCollection>();
-			Services.AddSingleton<IGraphCollection, GraphCollection>();
-			Services.AddSingleton<IIssueCollection, IssueCollection>();
-			Services.AddSingleton<IJobCollection, JobCollection>();
-			Services.AddSingleton<IJobStepRefCollection, JobStepRefCollection>();
-			Services.AddSingleton<IJobTimingCollection, JobTimingCollection>();
-			Services.AddSingleton<ILeaseCollection, LeaseCollection>();
-			Services.AddSingleton<ILogEventCollection, LogEventCollection>();
-			Services.AddSingleton<ILogFileCollection, LogFileCollection>();
-			Services.AddSingleton<INotificationTriggerCollection, NotificationTriggerCollection>();
-			Services.AddSingleton<IPoolCollection, PoolCollection>();
-			Services.AddSingleton<IProjectCollection, ProjectCollection>();
-			Services.AddSingleton<ISessionCollection, SessionCollection>();
-			Services.AddSingleton<IServiceAccountCollection, ServiceAccountCollection>();
-			Services.AddSingleton<ISubscriptionCollection, SubscriptionCollection>();
-			Services.AddSingleton<IStreamCollection, StreamCollection>();
-			Services.AddSingleton<ITemplateCollection, TemplateCollection>();
-			Services.AddSingleton<ITestDataCollection, TestDataCollection>();
-			Services.AddSingleton<ITelemetryCollection, TelemetryCollection>();
-			Services.AddSingleton<ITemplateCollection, TemplateCollection>();
-			Services.AddSingleton<IUgsMetadataCollection, UgsMetadataCollection>();
-			Services.AddSingleton<IUserCollection, UserCollectionV2>();
-			Services.AddSingleton<IDeviceCollection, DeviceCollection>();
-			Services.AddSingleton<INoticeCollection, NoticeCollection>();
+			services.AddGrpcReflection();
+
+			services.AddSingleton<IAgentCollection, AgentCollection>();
+			services.AddSingleton<IAgentSoftwareCollection, AgentSoftwareCollection>();
+			services.AddSingleton<IArtifactCollection, ArtifactCollection>();
+			services.AddSingleton<ICommitCollection, CommitCollection>();
+			services.AddSingleton<IGraphCollection, GraphCollection>();
+			services.AddSingleton<IIssueCollection, IssueCollection>();
+			services.AddSingleton<IJobCollection, JobCollection>();
+			services.AddSingleton<IJobStepRefCollection, JobStepRefCollection>();
+			services.AddSingleton<IJobTimingCollection, JobTimingCollection>();
+			services.AddSingleton<ILeaseCollection, LeaseCollection>();
+			services.AddSingleton<ILogEventCollection, LogEventCollection>();
+			services.AddSingleton<ILogFileCollection, LogFileCollection>();
+			services.AddSingleton<INotificationTriggerCollection, NotificationTriggerCollection>();
+			services.AddSingleton<IPoolCollection, PoolCollection>();
+			services.AddSingleton<IProjectCollection, ProjectCollection>();
+			services.AddSingleton<ISessionCollection, SessionCollection>();
+			services.AddSingleton<IServiceAccountCollection, ServiceAccountCollection>();
+			services.AddSingleton<ISubscriptionCollection, SubscriptionCollection>();
+			services.AddSingleton<IStreamCollection, StreamCollection>();
+			services.AddSingleton<ITemplateCollection, TemplateCollection>();
+			services.AddSingleton<ITestDataCollection, TestDataCollection>();
+			services.AddSingleton<ITelemetryCollection, TelemetryCollection>();
+			services.AddSingleton<ITemplateCollection, TemplateCollection>();
+			services.AddSingleton<IUgsMetadataCollection, UgsMetadataCollection>();
+			services.AddSingleton<IUserCollection, UserCollectionV2>();
+			services.AddSingleton<IDeviceCollection, DeviceCollection>();
+			services.AddSingleton<INoticeCollection, NoticeCollection>();
 
 			// Auditing
-			Services.AddSingleton<IAuditLog<AgentId>>(SP => SP.GetRequiredService<IAuditLogFactory<AgentId>>().Create("Agents.Log", "AgentId"));
+			services.AddSingleton<IAuditLog<AgentId>>(sp => sp.GetRequiredService<IAuditLogFactory<AgentId>>().Create("Agents.Log", "AgentId"));
 
-			Services.AddSingleton(typeof(IAuditLogFactory<>), typeof(AuditLogFactory<>));
-			Services.AddSingleton(typeof(ISingletonDocument<>), typeof(SingletonDocument<>));
+			services.AddSingleton(typeof(IAuditLogFactory<>), typeof(AuditLogFactory<>));
+			services.AddSingleton(typeof(ISingletonDocument<>), typeof(SingletonDocument<>));
 
-			Services.AddSingleton<AutoscaleService>();
-			Services.AddSingleton<AutoscaleServiceV2>();
-			Services.AddSingleton<LeaseUtilizationStrategy>();
-			Services.AddSingleton<JobQueueStrategy>();
-			Services.AddSingleton<NoOpPoolSizeStrategy>();
+			services.AddSingleton<AutoscaleService>();
+			services.AddSingleton<AutoscaleServiceV2>();
+			services.AddSingleton<LeaseUtilizationStrategy>();
+			services.AddSingleton<JobQueueStrategy>();
+			services.AddSingleton<NoOpPoolSizeStrategy>();
 			
-			switch (Settings.FleetManager)
+			switch (settings.FleetManager)
 			{
 				case FleetManagerType.Aws:
-					Services.AddSingleton<IFleetManager, AwsFleetManager>();
+					services.AddSingleton<IFleetManager, AwsFleetManager>();
 					break;
 				default:
-					Services.AddSingleton<IFleetManager, DefaultFleetManager>();
+					services.AddSingleton<IFleetManager, DefaultFleetManager>();
 					break;
 			}
 
-			Services.AddSingleton<AclService>();
-			Services.AddSingleton<AgentService>();			
-			Services.AddSingleton<AgentSoftwareService>();
-			Services.AddSingleton<ConsistencyService>();
-			Services.AddSingleton<RequestTrackerService>();
-			Services.AddSingleton<CredentialService>();
-			Services.AddSingleton<DatabaseService>();
-			Services.AddSingleton<ConfigService>();
-			Services.AddSingleton<IDogStatsd>(ctx =>
+			services.AddSingleton<AclService>();
+			services.AddSingleton<AgentService>();			
+			services.AddSingleton<AgentSoftwareService>();
+			services.AddSingleton<ConsistencyService>();
+			services.AddSingleton<RequestTrackerService>();
+			services.AddSingleton<CredentialService>();
+			services.AddSingleton<DatabaseService>();
+			services.AddSingleton<ConfigService>();
+			services.AddSingleton<IDogStatsd>(ctx =>
 			{
-				string? DatadogAgentHost = Environment.GetEnvironmentVariable("DD_AGENT_HOST");
-				if (DatadogAgentHost != null)
+				string? datadogAgentHost = Environment.GetEnvironmentVariable("DD_AGENT_HOST");
+				if (datadogAgentHost != null)
 				{
 					// Datadog agent is configured, enable DogStatsD for metric collection
-					StatsdConfig Config = new StatsdConfig
+					StatsdConfig config = new StatsdConfig
 					{
-						StatsdServerName = DatadogAgentHost,
+						StatsdServerName = datadogAgentHost,
 						StatsdPort = 8125,
 					};
 
-					DogStatsdService DogStatsdService = new DogStatsdService();
-					DogStatsdService.Configure(Config);
-					return DogStatsdService;
+					DogStatsdService dogStatsdService = new DogStatsdService();
+					dogStatsdService.Configure(config);
+					return dogStatsdService;
 				}
 				return new NoOpDogStatsd();
 			});
-			Services.AddSingleton<CommitService>();
-			Services.AddSingleton<ICommitService>(SP => SP.GetRequiredService<CommitService>());
-			Services.AddSingleton<IClock, Clock>();
-			Services.AddSingleton<IDowntimeService, DowntimeService>();
-			Services.AddSingleton<IIssueService, IssueService>();
-			Services.AddSingleton<JobService>();
-			Services.AddSingleton<LifetimeService>();
-			Services.AddSingleton<ILogFileService, LogFileService>();
-			Services.AddSingleton<INotificationService, NotificationService>();
-			Services.AddSingleton<IPerforceService, PerforceService>();
-
-			Services.AddSingleton<PerforceLoadBalancer>();
-			Services.AddSingleton<PoolService>();
-			Services.AddSingleton<ProjectService>();
-			Services.AddSingleton<ScheduleService>();
-			Services.AddSingleton<SlackNotificationSink>();
-			Services.AddSingleton<IAvatarService, SlackNotificationSink>(SP => SP.GetRequiredService<SlackNotificationSink>());
-			Services.AddSingleton<INotificationSink, SlackNotificationSink>(SP => SP.GetRequiredService<SlackNotificationSink>());
-			Services.AddSingleton<StreamService>();
-			Services.AddSingleton<UpgradeService>();
-			Services.AddSingleton<DeviceService>();			
-			Services.AddSingleton<JiraService>();
-			Services.AddSingleton<NoticeService>();
-
-			AWSOptions AwsOptions = Configuration.GetAWSOptions();
-			if (Settings.S3CredentialType == "AssumeRole" && Settings.S3AssumeArn != null)
-			{
-				AwsOptions.Credentials = new AssumeRoleAWSCredentials(FallbackCredentialsFactory.GetCredentials(), Settings.S3AssumeArn, "Horde");
-			}
-			else if(Settings.S3CredentialType == "AssumeRoleWebIdentity")
-			{
-				AwsOptions.Credentials = AssumeRoleWithWebIdentityCredentials.FromEnvironmentVariables();
-			}
-			else if(Settings.S3CredentialType == "Fallback")
+			services.AddSingleton<CommitService>();
+			services.AddSingleton<ICommitService>(sp => sp.GetRequiredService<CommitService>());
+			services.AddSingleton<IClock, Clock>();
+			services.AddSingleton<IDowntimeService, DowntimeService>();
+			services.AddSingleton<IIssueService, IssueService>();
+			services.AddSingleton<JobService>();
+			services.AddSingleton<LifetimeService>();
+			services.AddSingleton<ILogFileService, LogFileService>();
+			services.AddSingleton<INotificationService, NotificationService>();
+			services.AddSingleton<IPerforceService, PerforceService>();
+
+			services.AddSingleton<PerforceLoadBalancer>();
+			services.AddSingleton<PoolService>();
+			services.AddSingleton<ProjectService>();
+			services.AddSingleton<ScheduleService>();
+			services.AddSingleton<SlackNotificationSink>();
+			services.AddSingleton<IAvatarService, SlackNotificationSink>(sp => sp.GetRequiredService<SlackNotificationSink>());
+			services.AddSingleton<INotificationSink, SlackNotificationSink>(sp => sp.GetRequiredService<SlackNotificationSink>());
+			services.AddSingleton<StreamService>();
+			services.AddSingleton<UpgradeService>();
+			services.AddSingleton<DeviceService>();			
+			services.AddSingleton<JiraService>();
+			services.AddSingleton<NoticeService>();
+
+			AWSOptions awsOptions = Configuration.GetAWSOptions();
+			if (settings.S3CredentialType == "AssumeRole" && settings.S3AssumeArn != null)
+			{
+				awsOptions.Credentials = new AssumeRoleAWSCredentials(FallbackCredentialsFactory.GetCredentials(), settings.S3AssumeArn, "Horde");
+			}
+			else if(settings.S3CredentialType == "AssumeRoleWebIdentity")
+			{
+				awsOptions.Credentials = AssumeRoleWithWebIdentityCredentials.FromEnvironmentVariables();
+			}
+			else if(settings.S3CredentialType == "Fallback")
 			{
 				// Using the fallback credentials from the AWS SDK, it will pick up credentials through a number of default mechanisms.
-				AwsOptions.Credentials = FallbackCredentialsFactory.GetCredentials();
+				awsOptions.Credentials = FallbackCredentialsFactory.GetCredentials();
 			}
-			else if(Settings.S3CredentialType == "SharedCredentials" && Settings.S3AwsProfile != null)
+			else if(settings.S3CredentialType == "SharedCredentials" && settings.S3AwsProfile != null)
 			{
 				// This SharedCredentials option is primarily for development purposes.
-				var (AccessKey, SecretAccessKey, SecretToken) = AwsHelper.ReadAwsCredentials(Settings.S3AwsProfile);
-				AwsOptions.Credentials = new Credentials(AccessKey, SecretAccessKey, SecretToken, DateTime.Now + TimeSpan.FromHours(12));
+				(string accessKey, string secretAccessKey, string secretToken) = AwsHelper.ReadAwsCredentials(settings.S3AwsProfile);
+				awsOptions.Credentials = new Credentials(accessKey, secretAccessKey, secretToken, DateTime.Now + TimeSpan.FromHours(12));
 			}
 			
-			Services.AddSingleton(AwsOptions);
+			services.AddSingleton(awsOptions);
 
-			Services.AddSingleton(new StorageBackendSettings<PersistentLogStorage> { Type = Settings.ExternalStorageProviderType, BaseDir = Settings.LocalLogsDir, BucketName = Settings.S3LogBucketName });
-			Services.AddSingleton(new StorageBackendSettings<ArtifactCollection> { Type = Settings.ExternalStorageProviderType, BaseDir = Settings.LocalArtifactsDir, BucketName = Settings.S3ArtifactBucketName });
-			Services.AddSingleton(typeof(IStorageBackend<>), typeof(StorageBackendFactory<>));
+			services.AddSingleton(new StorageBackendSettings<PersistentLogStorage> { Type = settings.ExternalStorageProviderType, BaseDir = settings.LocalLogsDir, BucketName = settings.S3LogBucketName });
+			services.AddSingleton(new StorageBackendSettings<ArtifactCollection> { Type = settings.ExternalStorageProviderType, BaseDir = settings.LocalArtifactsDir, BucketName = settings.S3ArtifactBucketName });
+			services.AddSingleton(typeof(IStorageBackend<>), typeof(StorageBackendFactory<>));
 
-			Services.AddHordeStorage(Settings => ConfigSection.GetSection("Storage").Bind(Settings));
+			services.AddHordeStorage(settings => configSection.GetSection("Storage").Bind(settings));
 
-			ConfigureLogStorage(Services);
+			ConfigureLogStorage(services);
 
-			AuthenticationBuilder AuthBuilder = Services.AddAuthentication(Options =>
+			AuthenticationBuilder authBuilder = services.AddAuthentication(options =>
 				{
-					switch (Settings.AuthMethod)
+					switch (settings.AuthMethod)
 					{
 						case AuthMethod.Anonymous:
-							Options.DefaultAuthenticateScheme = AnonymousAuthenticationHandler.AuthenticationScheme;
-							Options.DefaultSignInScheme = AnonymousAuthenticationHandler.AuthenticationScheme;
-							Options.DefaultChallengeScheme = AnonymousAuthenticationHandler.AuthenticationScheme;
+							options.DefaultAuthenticateScheme = AnonymousAuthenticationHandler.AuthenticationScheme;
+							options.DefaultSignInScheme = AnonymousAuthenticationHandler.AuthenticationScheme;
+							options.DefaultChallengeScheme = AnonymousAuthenticationHandler.AuthenticationScheme;
 							break;
 						
 						case AuthMethod.Okta:
 							// If an authentication cookie is present, use it to get authentication information
-							Options.DefaultAuthenticateScheme = CookieAuthenticationDefaults.AuthenticationScheme;
-							Options.DefaultSignInScheme = CookieAuthenticationDefaults.AuthenticationScheme;
+							options.DefaultAuthenticateScheme = CookieAuthenticationDefaults.AuthenticationScheme;
+							options.DefaultSignInScheme = CookieAuthenticationDefaults.AuthenticationScheme;
 
 							// If authentication is required, and no cookie is present, use OIDC to sign in
-							Options.DefaultChallengeScheme = OktaDefaults.AuthenticationScheme;
+							options.DefaultChallengeScheme = OktaDefaults.AuthenticationScheme;
 							break;
 						
 						case AuthMethod.OpenIdConnect:
 							// If an authentication cookie is present, use it to get authentication information
-							Options.DefaultAuthenticateScheme = CookieAuthenticationDefaults.AuthenticationScheme;
-							Options.DefaultSignInScheme = CookieAuthenticationDefaults.AuthenticationScheme;
+							options.DefaultAuthenticateScheme = CookieAuthenticationDefaults.AuthenticationScheme;
+							options.DefaultSignInScheme = CookieAuthenticationDefaults.AuthenticationScheme;
 
 							// If authentication is required, and no cookie is present, use OIDC to sign in
-							Options.DefaultChallengeScheme = OpenIdConnectDefaults.AuthenticationScheme;
+							options.DefaultChallengeScheme = OpenIdConnectDefaults.AuthenticationScheme;
 							break;
 						
 						default:
-							throw new ArgumentException($"Invalid auth method {Settings.AuthMethod}");
+							throw new ArgumentException($"Invalid auth method {settings.AuthMethod}");
 					}
 				});
 
-			List<string> Schemes = new List<string>();
+			List<string> schemes = new List<string>();
 
-			AuthBuilder.AddCookie(Options =>
+			authBuilder.AddCookie(options =>
 				 {
-					 Options.Events.OnValidatePrincipal = Context =>
+					 options.Events.OnValidatePrincipal = context =>
 					 {
-						 if (Context.Principal?.FindFirst(HordeClaimTypes.UserId) == null)
+						 if (context.Principal?.FindFirst(HordeClaimTypes.UserId) == null)
 						 {
-							 Context.RejectPrincipal();
+							 context.RejectPrincipal();
 						 }
 						 return Task.CompletedTask;
 					 };
 
-					 Options.Events.OnRedirectToAccessDenied = Context =>
+					 options.Events.OnRedirectToAccessDenied = context =>
 					 {
-						 Context.Response.StatusCode = StatusCodes.Status403Forbidden;
-						 return Context.Response.CompleteAsync();
+						 context.Response.StatusCode = StatusCodes.Status403Forbidden;
+						 return context.Response.CompleteAsync();
 					 };
 				 });
-			Schemes.Add(CookieAuthenticationDefaults.AuthenticationScheme);
+			schemes.Add(CookieAuthenticationDefaults.AuthenticationScheme);
 
-			AuthBuilder.AddServiceAccount(Options => { });
-			Schemes.Add(ServiceAccountAuthHandler.AuthenticationScheme);
+			authBuilder.AddServiceAccount(options => { });
+			schemes.Add(ServiceAccountAuthHandler.AuthenticationScheme);
 
 			
-			switch (Settings.AuthMethod)
+			switch (settings.AuthMethod)
 			{
 				case AuthMethod.Anonymous:
-					AuthBuilder.AddAnonymous(Options =>
+					authBuilder.AddAnonymous(options =>
 					{
-						Options.AdminClaimType = Settings.AdminClaimType;
-						Options.AdminClaimValue = Settings.AdminClaimValue;
+						options.AdminClaimType = settings.AdminClaimType;
+						options.AdminClaimValue = settings.AdminClaimValue;
 					});
-					Schemes.Add(AnonymousAuthenticationHandler.AuthenticationScheme);
+					schemes.Add(AnonymousAuthenticationHandler.AuthenticationScheme);
 					break;
 						
 				case AuthMethod.Okta:
-					AuthBuilder.AddOkta(OktaDefaults.AuthenticationScheme, OpenIdConnectDefaults.DisplayName, Options =>
+					authBuilder.AddOkta(OktaDefaults.AuthenticationScheme, OpenIdConnectDefaults.DisplayName, options =>
 						{
-							Options.Authority = Settings.OidcAuthority;
-							Options.ClientId = Settings.OidcClientId;
+							options.Authority = settings.OidcAuthority;
+							options.ClientId = settings.OidcClientId;
 
-							if (!String.IsNullOrEmpty(Settings.OidcSigninRedirect))
+							if (!String.IsNullOrEmpty(settings.OidcSigninRedirect))
 							{
-								Options.Events = new OpenIdConnectEvents
+								options.Events = new OpenIdConnectEvents
 								{
-									OnRedirectToIdentityProvider = async RedirectContext =>
+									OnRedirectToIdentityProvider = async redirectContext =>
 									{
-										RedirectContext.ProtocolMessage.RedirectUri = Settings.OidcSigninRedirect;
+										redirectContext.ProtocolMessage.RedirectUri = settings.OidcSigninRedirect;
 										await Task.CompletedTask;
 									}
 								};
 							}
 						});
-					Schemes.Add(OktaDefaults.AuthenticationScheme);
+					schemes.Add(OktaDefaults.AuthenticationScheme);
 					break;
 						
 				case AuthMethod.OpenIdConnect:
-					AuthBuilder.AddHordeOpenId(Settings, OpenIdConnectDefaults.AuthenticationScheme, OpenIdConnectDefaults.DisplayName, Options =>
+					authBuilder.AddHordeOpenId(settings, OpenIdConnectDefaults.AuthenticationScheme, OpenIdConnectDefaults.DisplayName, options =>
 						{
-							Options.Authority = Settings.OidcAuthority;
-							Options.ClientId = Settings.OidcClientId;
-							Options.ClientSecret = Settings.OidcClientSecret;
-							foreach (string Scope in Settings.OidcRequestedScopes)
+							options.Authority = settings.OidcAuthority;
+							options.ClientId = settings.OidcClientId;
+							options.ClientSecret = settings.OidcClientSecret;
+							foreach (string scope in settings.OidcRequestedScopes)
 							{
-								Options.Scope.Add(Scope);
+								options.Scope.Add(scope);
 							}
 
-							if (!String.IsNullOrEmpty(Settings.OidcSigninRedirect))
+							if (!String.IsNullOrEmpty(settings.OidcSigninRedirect))
 							{
-								Options.Events = new OpenIdConnectEvents
+								options.Events = new OpenIdConnectEvents
 								{
-									OnRedirectToIdentityProvider = async RedirectContext =>
+									OnRedirectToIdentityProvider = async redirectContext =>
 									{
-										RedirectContext.ProtocolMessage.RedirectUri = Settings.OidcSigninRedirect;
+										redirectContext.ProtocolMessage.RedirectUri = settings.OidcSigninRedirect;
 										await Task.CompletedTask;
 									}
 								};
 							}
 						});
-					Schemes.Add(OpenIdConnectDefaults.AuthenticationScheme);
+					schemes.Add(OpenIdConnectDefaults.AuthenticationScheme);
 					break;
 						
 				default:
-					throw new ArgumentException($"Invalid auth method {Settings.AuthMethod}");
+					throw new ArgumentException($"Invalid auth method {settings.AuthMethod}");
 			}
 
-			AuthBuilder.AddScheme<JwtBearerOptions, HordeJwtBearerHandler>(HordeJwtBearerHandler.AuthenticationScheme, Options => { });
-			Schemes.Add(HordeJwtBearerHandler.AuthenticationScheme);
+			authBuilder.AddScheme<JwtBearerOptions, HordeJwtBearerHandler>(HordeJwtBearerHandler.AuthenticationScheme, options => { });
+			schemes.Add(HordeJwtBearerHandler.AuthenticationScheme);
 
-			Services.AddAuthorization(Options =>
+			services.AddAuthorization(options =>
 				{
-					Options.DefaultPolicy = new AuthorizationPolicyBuilder(Schemes.ToArray())
+					options.DefaultPolicy = new AuthorizationPolicyBuilder(schemes.ToArray())
 						.RequireAuthenticatedUser()							
 						.Build();
 				});
 
-			if (Settings.IsRunModeActive(RunMode.Worker))
+			if (settings.IsRunModeActive(RunMode.Worker))
 			{
-				Services.AddHostedService(Provider => Provider.GetRequiredService<AutoscaleServiceV2>());
+				services.AddHostedService(provider => provider.GetRequiredService<AutoscaleServiceV2>());
 				
-				Services.AddHostedService(Provider => Provider.GetRequiredService<AgentService>());
-				Services.AddHostedService(Provider => Provider.GetRequiredService<CommitService>());
-				Services.AddHostedService(Provider => Provider.GetRequiredService<ConsistencyService>());
-				Services.AddHostedService(Provider => (DowntimeService)Provider.GetRequiredService<IDowntimeService>());
-				Services.AddHostedService(Provider => Provider.GetRequiredService<IIssueService>());
-				Services.AddHostedService(Provider => (LogFileService)Provider.GetRequiredService<ILogFileService>());
-				Services.AddHostedService(Provider => (NotificationService)Provider.GetRequiredService<INotificationService>());
-				if (!Settings.DisableSchedules)
+				services.AddHostedService(provider => provider.GetRequiredService<AgentService>());
+				services.AddHostedService(provider => provider.GetRequiredService<CommitService>());
+				services.AddHostedService(provider => provider.GetRequiredService<ConsistencyService>());
+				services.AddHostedService(provider => (DowntimeService)provider.GetRequiredService<IDowntimeService>());
+				services.AddHostedService(provider => provider.GetRequiredService<IIssueService>());
+				services.AddHostedService(provider => (LogFileService)provider.GetRequiredService<ILogFileService>());
+				services.AddHostedService(provider => (NotificationService)provider.GetRequiredService<INotificationService>());
+				if (!settings.DisableSchedules)
 				{
-					Services.AddHostedService(Provider => Provider.GetRequiredService<ScheduleService>());
+					services.AddHostedService(provider => provider.GetRequiredService<ScheduleService>());
 				}
 
-				Services.AddHostedService<MetricService>();
-				Services.AddHostedService(Provider => Provider.GetRequiredService<PerforceLoadBalancer>());
-				Services.AddHostedService<PoolUpdateService>();
-				Services.AddHostedService(Provider => Provider.GetRequiredService<SlackNotificationSink>());
-				Services.AddHostedService<ConfigService>();
-				Services.AddHostedService<TelemetryService>();
-				Services.AddHostedService(Provider => Provider.GetRequiredService<DeviceService>());
+				services.AddHostedService<MetricService>();
+				services.AddHostedService(provider => provider.GetRequiredService<PerforceLoadBalancer>());
+				services.AddHostedService<PoolUpdateService>();
+				services.AddHostedService(provider => provider.GetRequiredService<SlackNotificationSink>());
+				services.AddHostedService<ConfigService>();
+				services.AddHostedService<TelemetryService>();
+				services.AddHostedService(provider => provider.GetRequiredService<DeviceService>());
 			}
 
 			// Task sources. Order of registration is important here; it dictates the priority in which sources are served.
-			Services.AddSingleton<JobTaskSource>();
-			Services.AddHostedService<JobTaskSource>(Provider => Provider.GetRequiredService<JobTaskSource>());
-			Services.AddSingleton<ConformTaskSource>();
-			Services.AddHostedService<ConformTaskSource>(Provider => Provider.GetRequiredService<ConformTaskSource>());
-			Services.AddSingleton<IComputeService, ComputeService>();
+			services.AddSingleton<JobTaskSource>();
+			services.AddHostedService<JobTaskSource>(provider => provider.GetRequiredService<JobTaskSource>());
+			services.AddSingleton<ConformTaskSource>();
+			services.AddHostedService<ConformTaskSource>(provider => provider.GetRequiredService<ConformTaskSource>());
+			services.AddSingleton<IComputeService, ComputeService>();
 
-			Services.AddSingleton<ITaskSource, UpgradeTaskSource>();
-			Services.AddSingleton<ITaskSource, ShutdownTaskSource>();
-			Services.AddSingleton<ITaskSource, RestartTaskSource>();
-			Services.AddSingleton<ITaskSource, ConformTaskSource>(Provider => Provider.GetRequiredService<ConformTaskSource>());
-			Services.AddSingleton<ITaskSource, JobTaskSource>(Provider => Provider.GetRequiredService<JobTaskSource>());
-			Services.AddSingleton<ITaskSource, ComputeService>();
+			services.AddSingleton<ITaskSource, UpgradeTaskSource>();
+			services.AddSingleton<ITaskSource, ShutdownTaskSource>();
+			services.AddSingleton<ITaskSource, RestartTaskSource>();
+			services.AddSingleton<ITaskSource, ConformTaskSource>(provider => provider.GetRequiredService<ConformTaskSource>());
+			services.AddSingleton<ITaskSource, JobTaskSource>(provider => provider.GetRequiredService<JobTaskSource>());
+			services.AddSingleton<ITaskSource, ComputeService>();
 
-			Services.AddHostedService(Provider => Provider.GetRequiredService<ConformTaskSource>());
+			services.AddHostedService(provider => provider.GetRequiredService<ConformTaskSource>());
 
 			// Allow longer to shutdown so we can debug missing cancellation tokens
-			Services.Configure<HostOptions>(Options =>
+			services.Configure<HostOptions>(options =>
 			{
-				Options.ShutdownTimeout = TimeSpan.FromSeconds(30.0);
+				options.ShutdownTimeout = TimeSpan.FromSeconds(30.0);
 			});
 
 			// Allow forwarded headers
-			Services.Configure<ForwardedHeadersOptions>(Options =>
+			services.Configure<ForwardedHeadersOptions>(options =>
 			{
-				Options.ForwardedHeaders = ForwardedHeaders.XForwardedFor | ForwardedHeaders.XForwardedProto;
-				Options.KnownProxies.Clear();
-				Options.KnownNetworks.Clear();
+				options.ForwardedHeaders = ForwardedHeaders.XForwardedFor | ForwardedHeaders.XForwardedProto;
+				options.KnownProxies.Clear();
+				options.KnownNetworks.Clear();
 			});
 
-			Services.AddMvc().AddJsonOptions(Options => ConfigureJsonSerializer(Options.JsonSerializerOptions));
+			services.AddMvc().AddJsonOptions(options => ConfigureJsonSerializer(options.JsonSerializerOptions));
 
-			Services.AddControllers(Options =>
+			services.AddControllers(options =>
 			{
-				Options.InputFormatters.Add(new CbInputFormatter());
-				Options.OutputFormatters.Add(new CbOutputFormatter());
-				Options.OutputFormatters.Insert(0, new CbPreferredOutputFormatter());
-				Options.FormatterMappings.SetMediaTypeMappingForFormat("uecb", CustomMediaTypeNames.UnrealCompactBinary);
+				options.InputFormatters.Add(new CbInputFormatter());
+				options.OutputFormatters.Add(new CbOutputFormatter());
+				options.OutputFormatters.Insert(0, new CbPreferredOutputFormatter());
+				options.FormatterMappings.SetMediaTypeMappingForFormat("uecb", CustomMediaTypeNames.UnrealCompactBinary);
 			});
 
-			Services.AddSwaggerGen(Config =>
+			services.AddSwaggerGen(config =>
 			{
-				Config.SwaggerDoc("v1", new OpenApiInfo { Title = "Horde Server API", Version = "v1" });
-				Config.IncludeXmlComments(Path.Combine(AppContext.BaseDirectory, $"{Assembly.GetExecutingAssembly().GetName().Name}.xml"));
+				config.SwaggerDoc("v1", new OpenApiInfo { Title = "Horde Server API", Version = "v1" });
+				config.IncludeXmlComments(Path.Combine(AppContext.BaseDirectory, $"{Assembly.GetExecutingAssembly().GetName().Name}.xml"));
 			});
 
-			Services.Configure<ApiBehaviorOptions>(Options =>
+			services.Configure<ApiBehaviorOptions>(options =>
 			{
-				Options.InvalidModelStateResponseFactory = Context =>
+				options.InvalidModelStateResponseFactory = context =>
 				{
-					foreach(KeyValuePair<string, ModelStateEntry> Pair in Context.ModelState)
+					foreach(KeyValuePair<string, ModelStateEntry> pair in context.ModelState)
 					{
-						ModelError? Error = Pair.Value.Errors.FirstOrDefault();
-						if (Error != null)
+						ModelError? error = pair.Value.Errors.FirstOrDefault();
+						if (error != null)
 						{
-							string Message = Error.ErrorMessage;
-							if (String.IsNullOrEmpty(Message))
+							string message = error.ErrorMessage;
+							if (String.IsNullOrEmpty(message))
 							{
-								Message = Error.Exception?.Message ?? "Invalid error object";
+								message = error.Exception?.Message ?? "Invalid error object";
 							}
-							return new BadRequestObjectResult(EpicGames.Core.LogEvent.Create(LogLevel.Error, KnownLogEvents.None, Error.Exception, "Invalid value for {Name}: {Message}", Pair.Key, Message));
+							return new BadRequestObjectResult(EpicGames.Core.LogEvent.Create(LogLevel.Error, KnownLogEvents.None, error.Exception, "Invalid value for {Name}: {Message}", pair.Key, message));
 						}
 					}
-					return new BadRequestObjectResult(Context.ModelState);
+					return new BadRequestObjectResult(context.ModelState);
 				};
 			});
 
-			DirectoryReference DashboardDir = DirectoryReference.Combine(Program.AppDir, "DashboardApp");
-			if (DirectoryReference.Exists(DashboardDir)) 
+			DirectoryReference dashboardDir = DirectoryReference.Combine(Program.AppDir, "DashboardApp");
+			if (DirectoryReference.Exists(dashboardDir)) 
 			{
-				Services.AddSpaStaticFiles(Config => {Config.RootPath = "DashboardApp";});
+				services.AddSpaStaticFiles(config => {config.RootPath = "DashboardApp";});
 			}
 
 			ConfigureMongoDbClient();
 			ConfigureFormatters();
 
-			OnAddHealthChecks(Services);
+			OnAddHealthChecks(services);
 		}
 
 		public static void ConfigureFormatters()
@@ -676,16 +657,16 @@ namespace Horde.Build
 			LogEventFormatter.RegisterFormatter(typeof(UserId), new LogEventFormatter.AnnotateTypeFormatter("UserId"));
 		}
 
-		public static void ConfigureJsonSerializer(JsonSerializerOptions Options)
+		public static void ConfigureJsonSerializer(JsonSerializerOptions options)
 		{
-			Options.DefaultIgnoreCondition = JsonIgnoreCondition.WhenWritingNull;
-			Options.PropertyNamingPolicy = JsonNamingPolicy.CamelCase;
-			Options.PropertyNameCaseInsensitive = true;
-			Options.Converters.Add(new JsonObjectIdConverter());
-			Options.Converters.Add(new JsonStringEnumConverter());
-			Options.Converters.Add(new JsonKnownTypesConverterFactory());
-			Options.Converters.Add(new JsonStringIdConverterFactory());
-			Options.Converters.Add(new JsonDateTimeConverter());
+			options.DefaultIgnoreCondition = JsonIgnoreCondition.WhenWritingNull;
+			options.PropertyNamingPolicy = JsonNamingPolicy.CamelCase;
+			options.PropertyNameCaseInsensitive = true;
+			options.Converters.Add(new JsonObjectIdConverter());
+			options.Converters.Add(new JsonStringEnumConverter());
+			options.Converters.Add(new JsonKnownTypesConverterFactory());
+			options.Converters.Add(new JsonStringIdConverterFactory());
+			options.Converters.Add(new JsonDateTimeConverter());
 		}
 
 		public class StorageBackendSettings<T> : IFileSystemStorageOptions, IAwsStorageOptions
@@ -707,23 +688,23 @@ namespace Horde.Build
 
 		public class StorageBackendFactory<T> : IStorageBackend<T>
 		{
-			IStorageBackend Inner;
+			readonly IStorageBackend _inner;
 
-			public StorageBackendFactory(IServiceProvider ServiceProvider, StorageBackendSettings<T> Options)
+			public StorageBackendFactory(IServiceProvider serviceProvider, StorageBackendSettings<T> options)
 			{
-				switch (Options.Type)
+				switch (options.Type)
 				{
 					case StorageProviderType.S3:
-						Inner = new AwsStorageBackend(ServiceProvider.GetRequiredService<AWSOptions>(), Options, ServiceProvider.GetRequiredService<ILogger<AwsStorageBackend>>());
+						_inner = new AwsStorageBackend(serviceProvider.GetRequiredService<AWSOptions>(), options, serviceProvider.GetRequiredService<ILogger<AwsStorageBackend>>());
 						break;
 					case StorageProviderType.FileSystem:
-						Inner = new FileSystemStorageBackend(Options);
+						_inner = new FileSystemStorageBackend(options);
 						break;
 					case StorageProviderType.Transient:
-						Inner = new TransientStorageBackend();
+						_inner = new TransientStorageBackend();
 						break;
 					case StorageProviderType.Relay:
-						Inner = new RelayStorageBackend(ServiceProvider.GetRequiredService<IOptions<ServerSettings>>());
+						_inner = new RelayStorageBackend(serviceProvider.GetRequiredService<IOptions<ServerSettings>>());
 						break;
 					default:
 						throw new NotImplementedException();
@@ -731,44 +712,44 @@ namespace Horde.Build
 			}
 
 			/// <inheritdoc/>
-			public Task<Stream?> ReadAsync(string Path) => Inner.ReadAsync(Path);
+			public Task<Stream?> ReadAsync(string path) => _inner.ReadAsync(path);
 
 			/// <inheritdoc/>
-			public Task<ReadOnlyMemory<byte>?> ReadBytesAsync(string Path) => Inner.ReadBytesAsync(Path);
+			public Task<ReadOnlyMemory<byte>?> ReadBytesAsync(string path) => _inner.ReadBytesAsync(path);
 
 			/// <inheritdoc/>
-			public Task WriteAsync(string Path, Stream Stream) => Inner.WriteAsync(Path, Stream);
+			public Task WriteAsync(string path, Stream stream) => _inner.WriteAsync(path, stream);
 
 			/// <inheritdoc/>
-			public Task WriteBytesAsync(string Path, ReadOnlyMemory<byte> Data) => Inner.WriteBytesAsync(Path, Data);
+			public Task WriteBytesAsync(string path, ReadOnlyMemory<byte> data) => _inner.WriteBytesAsync(path, data);
 
 			/// <inheritdoc/>
-			public Task DeleteAsync(string Path) => Inner.DeleteAsync(Path);
+			public Task DeleteAsync(string path) => _inner.DeleteAsync(path);
 
 			/// <inheritdoc/>
-			public Task<bool> ExistsAsync(string Path) => Inner.ExistsAsync(Path);
+			public Task<bool> ExistsAsync(string path) => _inner.ExistsAsync(path);
 		}
 
-		private static void ConfigureLogStorage(IServiceCollection Services)
+		private static void ConfigureLogStorage(IServiceCollection services)
 		{
-			Services.AddSingleton<ILogBuilder>(Provider =>
+			services.AddSingleton<ILogBuilder>(provider =>
 			{
-				RedisService? RedisService = Provider.GetService<RedisService>();
-				if(RedisService == null)
+				RedisService? redisService = provider.GetService<RedisService>();
+				if(redisService == null)
 				{
 					return new LocalLogBuilder();
 				}
 				else
 				{
-					return new RedisLogBuilder(RedisService.ConnectionPool, Provider.GetRequiredService<ILogger<RedisLogBuilder>>());
+					return new RedisLogBuilder(redisService.ConnectionPool, provider.GetRequiredService<ILogger<RedisLogBuilder>>());
 				}
 			});
 
-			Services.AddSingleton<PersistentLogStorage>();
+			services.AddSingleton<PersistentLogStorage>();
 
-			Services.AddSingleton<ILogStorage>(Provider =>
+			services.AddSingleton<ILogStorage>(provider =>
 			{
-				ILogStorage Storage = Provider.GetRequiredService<PersistentLogStorage>();
+				ILogStorage storage = provider.GetRequiredService<PersistentLogStorage>();
 
 //				IDatabase? RedisDb = Provider.GetService<IDatabase>();
 //				if (RedisDb != null)
@@ -776,9 +757,9 @@ namespace Horde.Build
 //					Storage = new RedisLogStorage(RedisDb, Provider.GetService<ILogger<RedisLogStorage>>(), Storage);
 //				}
 
-				Storage = new SequencedLogStorage(Storage);
-				Storage = new LocalLogStorage(50, Storage);
-				return Storage;
+				storage = new SequencedLogStorage(storage);
+				storage = new LocalLogStorage(50, storage);
+				return storage;
 			});
 		}
 		/*
@@ -809,29 +790,29 @@ namespace Horde.Build
 		public sealed class RefIdBsonSerializer : SerializerBase<RefId>
 		{
 			/// <inheritdoc/>
-			public override RefId Deserialize(BsonDeserializationContext Context, BsonDeserializationArgs Args)
+			public override RefId Deserialize(BsonDeserializationContext context, BsonDeserializationArgs args)
 			{
-				return new RefId(IoHash.Parse(Context.Reader.ReadString()));
+				return new RefId(IoHash.Parse(context.Reader.ReadString()));
 			}
 
 			/// <inheritdoc/>
-			public override void Serialize(BsonSerializationContext Context, BsonSerializationArgs Args, RefId Value)
+			public override void Serialize(BsonSerializationContext context, BsonSerializationArgs args, RefId value)
 			{
-				Context.Writer.WriteString(Value.ToString());
+				context.Writer.WriteString(value.ToString());
 			}
 		}
 
-		static int HaveConfiguredMongoDb = 0;
+		static int s_haveConfiguredMongoDb = 0;
 
 		public static void ConfigureMongoDbClient()
 		{
-			if (Interlocked.CompareExchange(ref HaveConfiguredMongoDb, 1, 0) == 0)
+			if (Interlocked.CompareExchange(ref s_haveConfiguredMongoDb, 1, 0) == 0)
 			{
 				// Ignore extra elements on deserialized documents
-				ConventionPack ConventionPack = new ConventionPack();
-				ConventionPack.Add(new IgnoreExtraElementsConvention(true));
-				ConventionPack.Add(new EnumRepresentationConvention(BsonType.String));
-				ConventionRegistry.Register("Horde", ConventionPack, Type => true);
+				ConventionPack conventionPack = new ConventionPack();
+				conventionPack.Add(new IgnoreExtraElementsConvention(true));
+				conventionPack.Add(new EnumRepresentationConvention(BsonType.String));
+				ConventionRegistry.Register("Horde", conventionPack, type => true);
 
 				// Register the custom serializers
 				BsonSerializer.RegisterSerializer(new RefIdBsonSerializer());
@@ -841,124 +822,124 @@ namespace Horde.Build
 			}
 		}
 
-		private static void OnAddHealthChecks(IServiceCollection Services)
+		private static void OnAddHealthChecks(IServiceCollection services)
 		{
-			IHealthChecksBuilder HealthChecks = Services.AddHealthChecks().AddCheck("self", () => HealthCheckResult.Healthy(), tags: new[] { "self" });
+			IHealthChecksBuilder healthChecks = services.AddHealthChecks().AddCheck("self", () => HealthCheckResult.Healthy(), tags: new[] { "self" });
 		}
 
 		// This method gets called by the runtime. Use this method to configure the HTTP request pipeline.
-		public static void Configure(IApplicationBuilder App, IWebHostEnvironment Env, Microsoft.Extensions.Hosting.IHostApplicationLifetime Lifetime, IOptions<ServerSettings> Settings)
+		public static void Configure(IApplicationBuilder app, IWebHostEnvironment env, Microsoft.Extensions.Hosting.IHostApplicationLifetime lifetime, IOptions<ServerSettings> settings)
 		{
-			App.UseForwardedHeaders();
-			App.UseExceptionHandler("/api/v1/error");
+			app.UseForwardedHeaders();
+			app.UseExceptionHandler("/api/v1/error");
 
 			// Used for allowing auth cookies in combination with OpenID Connect auth (for example, Google Auth did not work with these unset)
-			App.UseCookiePolicy(new CookiePolicyOptions()
+			app.UseCookiePolicy(new CookiePolicyOptions()
 			{
 				MinimumSameSitePolicy = SameSiteMode.None,
 				CheckConsentNeeded = _ => true
 			});
 
-			if (Settings.Value.CorsEnabled)
+			if (settings.Value.CorsEnabled)
 			{
-				App.UseCors("CorsPolicy");
+				app.UseCors("CorsPolicy");
 			}
 
 			// Enable middleware to serve generated Swagger as a JSON endpoint.
-			App.UseSwagger();
+			app.UseSwagger();
 
 			// Enable serilog request logging
-			App.UseSerilogRequestLogging(Options => Options.GetLevel = GetRequestLoggingLevel);
+			app.UseSerilogRequestLogging(options => options.GetLevel = GetRequestLoggingLevel);
 
 			// Include the source IP address with requests
-			App.Use(async (Context, Next) => {
-				using (Serilog.Context.LogContext.PushProperty("RemoteIP", Context.Connection.RemoteIpAddress))
+			app.Use(async (context, next) => {
+				using (Serilog.Context.LogContext.PushProperty("RemoteIP", context.Connection.RemoteIpAddress))
 				{
-					await Next();
+					await next();
 				}
 			});
 
 			// Enable middleware to serve swagger-ui (HTML, JS, CSS, etc.),
 			// specifying the Swagger JSON endpoint.
-			App.UseSwaggerUI(c =>
+			app.UseSwaggerUI(c =>
 			{
 				c.SwaggerEndpoint("/swagger/v1/swagger.json", "Horde Server API");
 				c.RoutePrefix = "swagger";
 			});
 
-			if (!Env.IsDevelopment())
+			if (!env.IsDevelopment())
 			{
-				App.UseMiddleware<RequestTrackerMiddleware>();	
+				app.UseMiddleware<RequestTrackerMiddleware>();	
 			}
 
-			App.UseDefaultFiles();
-			App.UseStaticFiles();
+			app.UseDefaultFiles();
+			app.UseStaticFiles();
 
-			DirectoryReference DashboardDir = DirectoryReference.Combine(Program.AppDir, "DashboardApp");
+			DirectoryReference dashboardDir = DirectoryReference.Combine(Program.AppDir, "DashboardApp");
 
-			if (DirectoryReference.Exists(DashboardDir)) 
+			if (DirectoryReference.Exists(dashboardDir)) 
 			{
-				App.UseSpaStaticFiles();
+				app.UseSpaStaticFiles();
 			}
 						
-			App.UseRouting();
+			app.UseRouting();
 
-			App.UseAuthentication();
-			App.UseAuthorization();
+			app.UseAuthentication();
+			app.UseAuthorization();
 
-			App.UseEndpoints(Endpoints =>
+			app.UseEndpoints(endpoints =>
 			{
-				Endpoints.MapGrpcService<HealthService>();
-				Endpoints.MapGrpcService<RpcService>();
+				endpoints.MapGrpcService<HealthService>();
+				endpoints.MapGrpcService<RpcService>();
 
-				Endpoints.MapGrpcReflectionService();
+				endpoints.MapGrpcReflectionService();
 
-				Endpoints.MapControllers();
+				endpoints.MapControllers();
 			});
 
-			if (DirectoryReference.Exists(DashboardDir)) 
+			if (DirectoryReference.Exists(dashboardDir)) 
 			{
-				App.UseSpa(Spa =>
+				app.UseSpa(spa =>
 				{ 
-					Spa.Options.SourcePath = "DashboardApp";        
+					spa.Options.SourcePath = "DashboardApp";        
 				});
 			}
 
-			if (Settings.Value.OpenBrowser && RuntimeInformation.IsOSPlatform(OSPlatform.Windows))
+			if (settings.Value.OpenBrowser && RuntimeInformation.IsOSPlatform(OSPlatform.Windows))
 			{
-				Lifetime.ApplicationStarted.Register(() => LaunchBrowser(App));
+				lifetime.ApplicationStarted.Register(() => LaunchBrowser(app));
 			}
 		}
 
-		static void LaunchBrowser(IApplicationBuilder App)
+		static void LaunchBrowser(IApplicationBuilder app)
 		{
-			IServerAddressesFeature? Feature = App.ServerFeatures.Get<IServerAddressesFeature>();
-			if (Feature != null && Feature.Addresses.Count > 0)
+			IServerAddressesFeature? feature = app.ServerFeatures.Get<IServerAddressesFeature>();
+			if (feature != null && feature.Addresses.Count > 0)
 			{
 				// with a development cert, host will be set by default to localhost, otherwise there will be no host in address
-				string Address = Feature.Addresses.First().Replace("[::]", System.Net.Dns.GetHostName(), StringComparison.OrdinalIgnoreCase);
-				Process.Start(new ProcessStartInfo { FileName = Address, UseShellExecute = true });
+				string address = feature.Addresses.First().Replace("[::]", System.Net.Dns.GetHostName(), StringComparison.OrdinalIgnoreCase);
+				Process.Start(new ProcessStartInfo { FileName = address, UseShellExecute = true });
 			}
 		}
 
-		static LogEventLevel GetRequestLoggingLevel(HttpContext Context, double ElapsedMs, Exception Ex)
+		static LogEventLevel GetRequestLoggingLevel(HttpContext context, double elapsedMs, Exception ex)
 		{
-			if (Context.Request != null && Context.Request.Path.HasValue)
+			if (context.Request != null && context.Request.Path.HasValue)
 			{
-				string RequestPath = Context.Request.Path;
-				if (RequestPath.Equals("/Horde.HordeRpc/QueryServerStateV2", StringComparison.OrdinalIgnoreCase))
+				string requestPath = context.Request.Path;
+				if (requestPath.Equals("/Horde.HordeRpc/QueryServerStateV2", StringComparison.OrdinalIgnoreCase))
 				{
 					return LogEventLevel.Verbose;
 				}
-				if (RequestPath.Equals("/Horde.HordeRpc/UpdateSession", StringComparison.OrdinalIgnoreCase))
+				if (requestPath.Equals("/Horde.HordeRpc/UpdateSession", StringComparison.OrdinalIgnoreCase))
 				{
 					return LogEventLevel.Verbose;
 				}
-				if (RequestPath.Equals("/Horde.HordeRpc/CreateEvents", StringComparison.OrdinalIgnoreCase))
+				if (requestPath.Equals("/Horde.HordeRpc/CreateEvents", StringComparison.OrdinalIgnoreCase))
 				{
 					return LogEventLevel.Verbose;
 				}
-				if (RequestPath.Equals("/Horde.HordeRpc/WriteOutput", StringComparison.OrdinalIgnoreCase))
+				if (requestPath.Equals("/Horde.HordeRpc/WriteOutput", StringComparison.OrdinalIgnoreCase))
 				{
 					return LogEventLevel.Verbose;
 				}
@@ -966,17 +947,17 @@ namespace Horde.Build
 			return LogEventLevel.Information;
 		}
 
-		public static void AddServices(IServiceCollection ServiceCollection, IConfiguration Configuration)
+		public static void AddServices(IServiceCollection serviceCollection, IConfiguration configuration)
 		{
-			Startup Startup = new Startup(Configuration);
-			Startup.ConfigureServices(ServiceCollection);
+			Startup startup = new Startup(configuration);
+			startup.ConfigureServices(serviceCollection);
 		}
 
-		public static IServiceProvider CreateServiceProvider(IConfiguration Configuration)
+		public static IServiceProvider CreateServiceProvider(IConfiguration configuration)
 		{
-			IServiceCollection ServiceCollection = new ServiceCollection();
-			AddServices(ServiceCollection, Configuration);
-			return ServiceCollection.BuildServiceProvider();
+			IServiceCollection serviceCollection = new ServiceCollection();
+			AddServices(serviceCollection, configuration);
+			return serviceCollection.BuildServiceProvider();
 		}
 	}
 }
diff --git a/Engine/Source/Programs/Horde/Horde.Build/Storage/Backends/AwsStorageBackend.cs b/Engine/Source/Programs/Horde/Horde.Build/Storage/Backends/AwsStorageBackend.cs
--- a/Engine/Source/Programs/Horde/Horde.Build/Storage/Backends/AwsStorageBackend.cs
+++ b/Engine/Source/Programs/Horde/Horde.Build/Storage/Backends/AwsStorageBackend.cs
@@ -1,37 +1,32 @@
 // Copyright Epic Games, Inc. All Rights Reserved.
 
-using Amazon.Extensions.NETCore.Setup;
-using Amazon.Runtime;
-using Amazon.S3;
-using Amazon.S3.Model;
-using Horde.Build.Utilities;
-using Microsoft.Extensions.Configuration;
-using Microsoft.Extensions.Logging;
-using Microsoft.Extensions.Options;
 using System;
 using System.Collections.Generic;
-using System.Diagnostics.CodeAnalysis;
 using System.Globalization;
 using System.IO;
 using System.Text;
 using System.Threading;
 using System.Threading.Tasks;
+using Amazon.Extensions.NETCore.Setup;
+using Amazon.S3;
+using Amazon.S3.Model;
+using Horde.Build.Utilities;
+using Microsoft.Extensions.Logging;
 
 namespace Horde.Build.Storage.Backends
 {
 	/// <summary>
 	/// Exception wrapper for S3 requests
 	/// </summary>
-	[SuppressMessage("Design", "CA1032:Implement standard exception constructors", Justification = "<Pending>")]
 	public sealed class AwsException : Exception
 	{
 		/// <summary>
 		/// Constructor
 		/// </summary>
-		/// <param name="Message">Message for the exception</param>
-		/// <param name="InnerException">Inner exception data</param>
-		public AwsException(string? Message, Exception? InnerException)
-			: base(Message, InnerException)
+		/// <param name="message">Message for the exception</param>
+		/// <param name="innerException">Inner exception data</param>
+		public AwsException(string? message, Exception? innerException)
+			: base(message, innerException)
 		{
 		}
 	}
@@ -60,70 +55,70 @@ namespace Horde.Build.Storage.Backends
 		/// <summary>
 		/// S3 Client
 		/// </summary>
-		private readonly IAmazonS3 Client;
+		private readonly IAmazonS3 _client;
 
 		/// <summary>
 		/// Options for AWs
 		/// </summary>
-		private IAwsStorageOptions Options;
+		private readonly IAwsStorageOptions _options;
 
 		/// <summary>
 		/// Semaphore for connecting to AWS
 		/// </summary>
-		private SemaphoreSlim Semaphore;
+		private readonly SemaphoreSlim _semaphore;
 
 		/// <summary>
 		/// Logger interface
 		/// </summary>
-		private ILogger Logger;
+		private readonly ILogger _logger;
 
 		/// <summary>
 		/// Constructor
 		/// </summary>
-		/// <param name="AwsOptions">AWS options</param>
-		/// <param name="Options">Storage options</param>
-		/// <param name="Logger">Logger interface</param>
-		public AwsStorageBackend(AWSOptions AwsOptions, IAwsStorageOptions Options, ILogger<AwsStorageBackend> Logger)
+		/// <param name="awsOptions">AWS options</param>
+		/// <param name="options">Storage options</param>
+		/// <param name="logger">Logger interface</param>
+		public AwsStorageBackend(AWSOptions awsOptions, IAwsStorageOptions options, ILogger<AwsStorageBackend> logger)
 		{
-			this.Client = AwsOptions.CreateServiceClient<IAmazonS3>();
-			this.Options = Options;
-			this.Semaphore = new SemaphoreSlim(16);
-			this.Logger = Logger;
+			_client = awsOptions.CreateServiceClient<IAmazonS3>();
+			_options = options;
+			_semaphore = new SemaphoreSlim(16);
+			_logger = logger;
 
-			Logger.LogInformation("Created AWS storage backend for bucket {BucketName} using credentials {Credentials} {CredentialsStr}", Options.BucketName, AwsOptions.Credentials.GetType(), AwsOptions.Credentials.ToString());
+			logger.LogInformation("Created AWS storage backend for bucket {BucketName} using credentials {Credentials} {CredentialsStr}", options.BucketName, awsOptions.Credentials.GetType(), awsOptions.Credentials.ToString());
 		}
 
 		/// <inheritdoc/>
 		public void Dispose()
 		{
-			Client.Dispose();
-			Semaphore.Dispose();
+			_client.Dispose();
+			_semaphore.Dispose();
 		}
 
 		class WrappedResponseStream : Stream
 		{
-			IDisposable Semaphore;
-			GetObjectResponse Response;
-			Stream ResponseStream;
+			readonly IDisposable _semaphore;
+			readonly GetObjectResponse _response;
+			readonly Stream _responseStream;
 
-			public WrappedResponseStream(IDisposable Semaphore, GetObjectResponse Response)
+			public WrappedResponseStream(IDisposable semaphore, GetObjectResponse response)
 			{
-				this.Semaphore = Semaphore;
-				this.Response = Response;
-				this.ResponseStream = Response.ResponseStream;
+				_semaphore = semaphore;
+				_response = response;
+				_responseStream = response.ResponseStream;
 			}
 
 			public override bool CanRead => true;
 			public override bool CanSeek => false;
 			public override bool CanWrite => false;
-			public override long Length => ResponseStream.Length;
-			public override long Position { get => ResponseStream.Position; set => throw new NotSupportedException(); }
+			public override long Length => _responseStream.Length;
+			public override long Position { get => _responseStream.Position; set => throw new NotSupportedException(); }
 
-			public override void Flush() => ResponseStream.Flush();
+			public override void Flush() => _responseStream.Flush();
 
-			public override int Read(byte[] buffer, int offset, int count) => ResponseStream.Read(buffer, offset, count);
-			public override ValueTask<int> ReadAsync(Memory<byte> buffer, CancellationToken cancellationToken = default) => ResponseStream.ReadAsync(buffer, cancellationToken);
-			public override Task<int> ReadAsync(byte[] buffer, int offset, int count, CancellationToken cancellationToken) => ResponseStream.ReadAsync(buffer, offset, count, cancellationToken);
+			public override int Read(byte[] buffer, int offset, int count) => _responseStream.Read(buffer, offset, count);
+			public override ValueTask<int> ReadAsync(Memory<byte> buffer, CancellationToken cancellationToken = default) => _responseStream.ReadAsync(buffer, cancellationToken);
+			public override Task<int> ReadAsync(byte[] buffer, int offset, int count, CancellationToken cancellationToken) => _responseStream.ReadAsync(buffer, offset, count, cancellationToken);
 
 			public override long Seek(long offset, SeekOrigin origin) => throw new NotSupportedException();
 			public override void SetLength(long value) => throw new NotSupportedException();
@@ -133,180 +128,179 @@ namespace Horde.Build.Storage.Backends
 			{
 				base.Dispose(disposing);
 
-				Semaphore.Dispose();
-				Response.Dispose();
-				ResponseStream.Dispose();
+				_semaphore.Dispose();
+				_response.Dispose();
+				_responseStream.Dispose();
 			}
 
 			public override async ValueTask DisposeAsync()
 			{
 				await base.DisposeAsync();
 
-				Semaphore.Dispose();
-				Response.Dispose();
-				await ResponseStream.DisposeAsync();
+				_semaphore.Dispose();
+				_response.Dispose();
+				await _responseStream.DisposeAsync();
 			}
 		}
 
-		string GetFullPath(string Path)
+		string GetFullPath(string path)
 		{
-			if (Options.BucketPath == null)
+			if (_options.BucketPath == null)
 			{
-				return Path;
+				return path;
 			}
 
-			StringBuilder Result = new StringBuilder();
-			Result.Append(Options.BucketPath);
-			if (Result.Length > 0 && Result[Result.Length - 1] != '/')
+			StringBuilder result = new StringBuilder();
+			result.Append(_options.BucketPath);
+			if (result.Length > 0 && result[^1] != '/')
 			{
-				Result.Append('/');
+				result.Append('/');
 			}
-			Result.Append(Path);
-			return Result.ToString();
+			result.Append(path);
+			return result.ToString();
 		}
 
 		/// <inheritdoc/>
-		public async Task<Stream?> ReadAsync(string Path)
+		public async Task<Stream?> ReadAsync(string path)
 		{
-			string FullPath = GetFullPath(Path);
+			string fullPath = GetFullPath(path);
 
-			IDisposable? Lock = null;
-			GetObjectResponse? Response = null;
+			IDisposable? semaLock = null;
+			GetObjectResponse? response = null;
 			try
 			{
-				Lock = await Semaphore.UseWaitAsync();
+				semaLock = await _semaphore.UseWaitAsync();
 
-				GetObjectRequest NewGetRequest = new GetObjectRequest();
-				NewGetRequest.BucketName = Options.BucketName;
-				NewGetRequest.Key = FullPath;
+				GetObjectRequest newGetRequest = new GetObjectRequest();
+				newGetRequest.BucketName = _options.BucketName;
+				newGetRequest.Key = fullPath;
 
-				Response = await Client.GetObjectAsync(NewGetRequest);
+				response = await _client.GetObjectAsync(newGetRequest);
 
-				return new WrappedResponseStream(Lock, Response);
+				return new WrappedResponseStream(semaLock, response);
 			}
-			catch (Exception Ex)
+			catch (Exception ex)
 			{
-				Logger.LogWarning(Ex, "Unable to read {Path} from S3", FullPath);
+				_logger.LogWarning(ex, "Unable to read {Path} from S3", fullPath);
 
-				Lock?.Dispose();
-				Response?.Dispose();
+				semaLock?.Dispose();
+				response?.Dispose();
 
 				return null;
 			}
 		}
 
 		/// <inheritdoc/>
-		[SuppressMessage("Design", "CA1031:Do not catch general exception types", Justification = "<Pending>")]
-		public async Task WriteAsync(string Path, Stream InputStream)
+		public async Task WriteAsync(string path, Stream inputStream)
 		{
-			TimeSpan[] RetryTimes =
+			TimeSpan[] retryTimes =
 			{
 				TimeSpan.FromSeconds(1.0),
 				TimeSpan.FromSeconds(5.0),
 				TimeSpan.FromSeconds(10.0),
 			};
 
-			string FullPath = GetFullPath(Path);
-			for (int Attempt = 0; ; Attempt++)
+			string fullPath = GetFullPath(path);
+			for (int attempt = 0; ; attempt++)
 			{
 				try
 				{
-					using IDisposable Lock = await Semaphore.UseWaitAsync();
-					await WriteInternalAsync(FullPath, InputStream);
-					Logger.LogDebug("Written data to {Path}", Path);
+					using IDisposable semaLock = await _semaphore.UseWaitAsync();
+					await WriteInternalAsync(fullPath, inputStream);
+					_logger.LogDebug("Written data to {Path}", path);
 					break;
 				}
-				catch (Exception Ex)
+				catch (Exception ex)
 				{
-					Logger.LogError(Ex, "Unable to write data to {Path} ({Attempt}/{AttemptCount})", FullPath, Attempt + 1, RetryTimes.Length + 1);
-					if (Attempt >= RetryTimes.Length)
+					_logger.LogError(ex, "Unable to write data to {Path} ({Attempt}/{AttemptCount})", fullPath, attempt + 1, retryTimes.Length + 1);
+					if (attempt >= retryTimes.Length)
 					{
-						throw new AwsException($"Unable to write to bucket {Options.BucketName} path {FullPath}", Ex);
+						throw new AwsException($"Unable to write to bucket {_options.BucketName} path {fullPath}", ex);
 					}
 				}
 
-				await Task.Delay(RetryTimes[Attempt]);
+				await Task.Delay(retryTimes[attempt]);
 			}
 		}
 
 		/// <inheritdoc/>
-		public async Task WriteInternalAsync(string FullPath, Stream Stream)
+		public async Task WriteInternalAsync(string fullPath, Stream stream)
 		{
 			const int MinPartSize = 5 * 1024 * 1024;
 
-			long StreamLen = Stream.Length;
-			if (StreamLen < MinPartSize)
+			long streamLen = stream.Length;
+			if (streamLen < MinPartSize)
 			{
 				// Read the data into memory. Errors with hash value not matching if we don't do this (?)
-				byte[] Buffer = new byte[StreamLen];
-				await ReadExactLengthAsync(Stream, Buffer, (int)StreamLen);
+				byte[] buffer = new byte[streamLen];
+				await ReadExactLengthAsync(stream, buffer, (int)streamLen);
 
 				// Upload it to S3
-				using (MemoryStream InputStream = new MemoryStream(Buffer))
+				using (MemoryStream inputStream = new MemoryStream(buffer))
 				{
-					PutObjectRequest UploadRequest = new PutObjectRequest();
-					UploadRequest.BucketName = Options.BucketName;
-					UploadRequest.Key = FullPath;
-					UploadRequest.InputStream = InputStream;
-					UploadRequest.Metadata.Add("bytes-written", StreamLen.ToString(CultureInfo.InvariantCulture));
-					await Client.PutObjectAsync(UploadRequest);
+					PutObjectRequest uploadRequest = new PutObjectRequest();
+					uploadRequest.BucketName = _options.BucketName;
+					uploadRequest.Key = fullPath;
+					uploadRequest.InputStream = inputStream;
+					uploadRequest.Metadata.Add("bytes-written", streamLen.ToString(CultureInfo.InvariantCulture));
+					await _client.PutObjectAsync(uploadRequest);
 				}
 			}
 			else
 			{
 				// Initiate a multi-part upload
-				InitiateMultipartUploadRequest InitiateRequest = new InitiateMultipartUploadRequest();
-				InitiateRequest.BucketName = Options.BucketName;
-				InitiateRequest.Key = FullPath;
+				InitiateMultipartUploadRequest initiateRequest = new InitiateMultipartUploadRequest();
+				initiateRequest.BucketName = _options.BucketName;
+				initiateRequest.Key = fullPath;
 
-				InitiateMultipartUploadResponse InitiateResponse = await Client.InitiateMultipartUploadAsync(InitiateRequest);
+				InitiateMultipartUploadResponse initiateResponse = await _client.InitiateMultipartUploadAsync(initiateRequest);
 				try
 				{
 					// Buffer for reading the data
-					byte[] Buffer = new byte[MinPartSize];
+					byte[] buffer = new byte[MinPartSize];
 
 					// Upload all the parts
-					List<PartETag> PartTags = new List<PartETag>();
-					for (long StreamPos = 0; StreamPos < StreamLen;)
+					List<PartETag> partTags = new List<PartETag>();
+					for (long streamPos = 0; streamPos < streamLen;)
 					{
 						// Read the next chunk of data into the buffer
-						int BufferLen = (int)Math.Min((long)MinPartSize, StreamLen - StreamPos);
-						await ReadExactLengthAsync(Stream, Buffer, BufferLen);
-						StreamPos += BufferLen;
+						int bufferLen = (int)Math.Min((long)MinPartSize, streamLen - streamPos);
+						await ReadExactLengthAsync(stream, buffer, bufferLen);
+						streamPos += bufferLen;
 
 						// Upload the part
-						using (MemoryStream InputStream = new MemoryStream(Buffer, 0, BufferLen))
+						using (MemoryStream inputStream = new MemoryStream(buffer, 0, bufferLen))
 						{
-							UploadPartRequest PartRequest = new UploadPartRequest();
-							PartRequest.BucketName = Options.BucketName;
-							PartRequest.Key = FullPath;
-							PartRequest.UploadId = InitiateResponse.UploadId;
-							PartRequest.InputStream = InputStream;
-							PartRequest.PartSize = BufferLen;
-							PartRequest.PartNumber = PartTags.Count + 1;
-							PartRequest.IsLastPart = (StreamPos == StreamLen);
-
-							UploadPartResponse PartResponse = await Client.UploadPartAsync(PartRequest);
-							PartTags.Add(new PartETag(PartResponse.PartNumber, PartResponse.ETag));
+							UploadPartRequest partRequest = new UploadPartRequest();
+							partRequest.BucketName = _options.BucketName;
+							partRequest.Key = fullPath;
+							partRequest.UploadId = initiateResponse.UploadId;
+							partRequest.InputStream = inputStream;
+							partRequest.PartSize = bufferLen;
+							partRequest.PartNumber = partTags.Count + 1;
+							partRequest.IsLastPart = (streamPos == streamLen);
+
+							UploadPartResponse partResponse = await _client.UploadPartAsync(partRequest);
+							partTags.Add(new PartETag(partResponse.PartNumber, partResponse.ETag));
 						}
 					}
 
 					// Mark the upload as complete
-					CompleteMultipartUploadRequest CompleteRequest = new CompleteMultipartUploadRequest();
-					CompleteRequest.BucketName = Options.BucketName;
-					CompleteRequest.Key = FullPath;
-					CompleteRequest.UploadId = InitiateResponse.UploadId;
-					CompleteRequest.PartETags = PartTags;
-					await Client.CompleteMultipartUploadAsync(CompleteRequest);
+					CompleteMultipartUploadRequest completeRequest = new CompleteMultipartUploadRequest();
+					completeRequest.BucketName = _options.BucketName;
+					completeRequest.Key = fullPath;
+					completeRequest.UploadId = initiateResponse.UploadId;
+					completeRequest.PartETags = partTags;
+					await _client.CompleteMultipartUploadAsync(completeRequest);
 				}
 				catch
 				{
 					// Abort the upload
-					AbortMultipartUploadRequest AbortRequest = new AbortMultipartUploadRequest();
-					AbortRequest.BucketName = Options.BucketName;
-					AbortRequest.Key = FullPath;
-					AbortRequest.UploadId = InitiateResponse.UploadId;
-					await Client.AbortMultipartUploadAsync(AbortRequest);
+					AbortMultipartUploadRequest abortRequest = new AbortMultipartUploadRequest();
+					abortRequest.BucketName = _options.BucketName;
+					abortRequest.Key = fullPath;
+					abortRequest.UploadId = initiateResponse.UploadId;
+					await _client.AbortMultipartUploadAsync(abortRequest);
 
 					throw;
 				}
@@ -316,45 +310,45 @@ namespace Horde.Build.Storage.Backends
 		/// <summary>
 		/// Reads data of an exact length into a stream
 		/// </summary>
-		/// <param name="Stream">The stream to read from</param>
-		/// <param name="Buffer">The buffer to read into</param>
-		/// <param name="Length">Length of the data to read</param>
+		/// <param name="stream">The stream to read from</param>
+		/// <param name="buffer">The buffer to read into</param>
+		/// <param name="length">Length of the data to read</param>
 		/// <returns>Async task</returns>
-		static async Task ReadExactLengthAsync(System.IO.Stream Stream, byte[] Buffer, int Length)
+		static async Task ReadExactLengthAsync(System.IO.Stream stream, byte[] buffer, int length)
 		{
-			int BufferPos = 0;
-			while (BufferPos < Length)
+			int bufferPos = 0;
+			while (bufferPos < length)
 			{
-				int BytesRead = await Stream.ReadAsync(Buffer, BufferPos, Length - BufferPos);
-				if (BytesRead == 0)
+				int bytesRead = await stream.ReadAsync(buffer, bufferPos, length - bufferPos);
+				if (bytesRead == 0)
 				{
 					throw new InvalidOperationException("Unexpected end of stream");
 				}
-				BufferPos += BytesRead;
+				bufferPos += bytesRead;
 			}
 		}
 
 		/// <inheritdoc/>
-		public async Task DeleteAsync(string Path)
+		public async Task DeleteAsync(string path)
 		{
-			DeleteObjectRequest NewDeleteRequest = new DeleteObjectRequest();
-			NewDeleteRequest.BucketName = Options.BucketName;
-			NewDeleteRequest.Key = GetFullPath(Path);
-			await Client.DeleteObjectAsync(NewDeleteRequest);
+			DeleteObjectRequest newDeleteRequest = new DeleteObjectRequest();
+			newDeleteRequest.BucketName = _options.BucketName;
+			newDeleteRequest.Key = GetFullPath(path);
+			await _client.DeleteObjectAsync(newDeleteRequest);
 		}
 
 		/// <inheritdoc/>
-		public async Task<bool> ExistsAsync(string Path)
+		public async Task<bool> ExistsAsync(string path)
 		{
 			try
 			{
-				GetObjectMetadataRequest Request = new GetObjectMetadataRequest();
-				Request.BucketName = Options.BucketName;
-				Request.Key = GetFullPath(Path);
-				await Client.GetObjectMetadataAsync(Request);
+				GetObjectMetadataRequest request = new GetObjectMetadataRequest();
+				request.BucketName = _options.BucketName;
+				request.Key = GetFullPath(path);
+				await _client.GetObjectMetadataAsync(request);
 				return true;
 			}
-			catch (AmazonS3Exception Ex) when (Ex.StatusCode == System.Net.HttpStatusCode.NotFound)
+			catch (AmazonS3Exception ex) when (ex.StatusCode == System.Net.HttpStatusCode.NotFound)
 			{
 				return false;
 			}
diff --git a/Engine/Source/Programs/Horde/Horde.Build/Storage/Backends/FileSystemStorageBackend.cs b/Engine/Source/Programs/Horde/Horde.Build/Storage/Backends/FileSystemStorageBackend.cs
--- a/Engine/Source/Programs/Horde/Horde.Build/Storage/Backends/FileSystemStorageBackend.cs
+++ b/Engine/Source/Programs/Horde/Horde.Build/Storage/Backends/FileSystemStorageBackend.cs
@@ -1,15 +1,10 @@
 // Copyright Epic Games, Inc. All Rights Reserved.
 
-using EpicGames.Core;
-using Horde.Build.Utilities;
-using Microsoft.Extensions.Logging;
-using Microsoft.Extensions.Options;
 using System;
-using System.Collections.Generic;
 using System.IO;
-using System.Linq;
 using System.Threading;
 using System.Threading.Tasks;
+using EpicGames.Core;
 
 namespace Horde.Build.Storage.Backends
 {
@@ -32,41 +27,41 @@ namespace Horde.Build.Storage.Backends
 		/// <summary>
 		/// Base directory for log files
 		/// </summary>
-		private readonly DirectoryReference BaseDir;
+		private readonly DirectoryReference _baseDir;
 
 		/// <summary>
 		/// Unique identifier for this instance
 		/// </summary>
-		private string InstanceId;
+		private readonly string _instanceId;
 
 		/// <summary>
 		/// Unique id for each write
 		/// </summary>
-		private int UniqueId;
+		private int _uniqueId;
 
 		/// <summary>
 		/// Constructor
 		/// </summary>
-		/// <param name="Options">Current Horde Settings</param>
-		public FileSystemStorageBackend(IFileSystemStorageOptions Options)
+		/// <param name="options">Current Horde Settings</param>
+		public FileSystemStorageBackend(IFileSystemStorageOptions options)
 		{
-			this.BaseDir = DirectoryReference.Combine(Program.DataDir, Options.BaseDir ?? "Storage");
-			this.InstanceId = Guid.NewGuid().ToString("N");
-			DirectoryReference.CreateDirectory(BaseDir);
+			_baseDir = DirectoryReference.Combine(Program.DataDir, options.BaseDir ?? "Storage");
+			_instanceId = Guid.NewGuid().ToString("N");
+			DirectoryReference.CreateDirectory(_baseDir);
 		}
 
 		/// <inheritdoc/>
-		public Task<Stream?> ReadAsync(string Path)
+		public Task<Stream?> ReadAsync(string path)
 		{
-			FileReference Location = FileReference.Combine(BaseDir, Path);
-			if (!FileReference.Exists(Location))
+			FileReference location = FileReference.Combine(_baseDir, path);
+			if (!FileReference.Exists(location))
 			{
 				return Task.FromResult<Stream?>(null);
 			}
 
 			try
 			{
-				return Task.FromResult<Stream?>(FileReference.Open(Location, FileMode.Open, FileAccess.Read, FileShare.Read));
+				return Task.FromResult<Stream?>(FileReference.Open(location, FileMode.Open, FileAccess.Read, FileShare.Read));
 			}
 			catch (DirectoryNotFoundException)
 			{
@@ -79,32 +74,32 @@ namespace Horde.Build.Storage.Backends
 		}
 
 		/// <inheritdoc/>
-		public async Task WriteAsync(string Path, Stream Stream)
+		public async Task WriteAsync(string path, Stream stream)
 		{
-			FileReference FinalLocation = FileReference.Combine(BaseDir, Path);
-			if (!FileReference.Exists(FinalLocation))
+			FileReference finalLocation = FileReference.Combine(_baseDir, path);
+			if (!FileReference.Exists(finalLocation))
 			{
 				// Write to a temp file first
-				int NewUniqueId = Interlocked.Increment(ref UniqueId);
+				int newUniqueId = Interlocked.Increment(ref _uniqueId);
 
-				DirectoryReference.CreateDirectory(FinalLocation.Directory);
-				FileReference TempLocation = new FileReference($"{FinalLocation}.{InstanceId}.{UniqueId:x8}");
+				DirectoryReference.CreateDirectory(finalLocation.Directory);
+				FileReference tempLocation = new FileReference($"{finalLocation}.{_instanceId}.{newUniqueId:x8}");
 
-				using (Stream OutputStream = FileReference.Open(TempLocation, FileMode.Create, FileAccess.Write, FileShare.Read))
+				using (Stream outputStream = FileReference.Open(tempLocation, FileMode.Create, FileAccess.Write, FileShare.Read))
 				{
-					await Stream.CopyToAsync(OutputStream);
+					await stream.CopyToAsync(outputStream);
 				}
 
 				// Move the temp file into place
 				try
 				{
-					FileReference.Move(TempLocation, FinalLocation, true);
+					FileReference.Move(tempLocation, finalLocation, true);
 				}
 				catch (IOException) // Already exists
 				{
-					if (FileReference.Exists(FinalLocation))
+					if (FileReference.Exists(finalLocation))
 					{
-						FileReference.Delete(TempLocation);
+						FileReference.Delete(tempLocation);
 					}
 					else
 					{
@@ -115,17 +110,17 @@ namespace Horde.Build.Storage.Backends
 		}
 
 		/// <inheritdoc/>
-		public Task<bool> ExistsAsync(string Path)
+		public Task<bool> ExistsAsync(string path)
 		{
-			FileReference Location = FileReference.Combine(BaseDir, Path);
-			return Task.FromResult(FileReference.Exists(Location));
+			FileReference location = FileReference.Combine(_baseDir, path);
+			return Task.FromResult(FileReference.Exists(location));
 		}
 
 		/// <inheritdoc/>
-		public Task DeleteAsync(string Path)
+		public Task DeleteAsync(string path)
 		{
-			FileReference Location = FileReference.Combine(BaseDir, Path);
-			FileReference.Delete(Location);
+			FileReference location = FileReference.Combine(_baseDir, path);
+			FileReference.Delete(location);
 			return Task.CompletedTask;
 		}
 	}
diff --git a/Engine/Source/Programs/Horde/Horde.Build/Storage/Backends/RelayStorageBackend.cs b/Engine/Source/Programs/Horde/Horde.Build/Storage/Backends/RelayStorageBackend.cs
--- a/Engine/Source/Programs/Horde/Horde.Build/Storage/Backends/RelayStorageBackend.cs
+++ b/Engine/Source/Programs/Horde/Horde.Build/Storage/Backends/RelayStorageBackend.cs
@@ -1,14 +1,11 @@
 // Copyright Epic Games, Inc. All Rights Reserved.
 
-using EpicGames.Core;
-using Microsoft.Extensions.Options;
 using System;
-using System.Collections.Generic;
 using System.IO;
-using System.Linq;
 using System.Net.Http;
 using System.Net.Http.Headers;
 using System.Threading.Tasks;
+using Microsoft.Extensions.Options;
 
 namespace Horde.Build.Storage.Backends
 {
@@ -25,65 +22,65 @@ namespace Horde.Build.Storage.Backends
 		/// <summary>
 		/// The client to connect with
 		/// </summary>
-		HttpClient Client;
+		readonly HttpClient _client;
 
 		/// <summary>
 		/// The base server URL
 		/// </summary>
-		Uri ServerUrl;
+		readonly Uri _serverUrl;
 
 		/// <summary>
 		/// Local storage provider
 		/// </summary>
-		FileSystemStorageBackend LocalStorage;
+		readonly FileSystemStorageBackend _localStorage;
 
 		/// <summary>
 		/// Constructor
 		/// </summary>
-		/// <param name="Settings">Settings for the server instance</param>
-		public RelayStorageBackend(IOptions<ServerSettings> Settings)
+		/// <param name="settings">Settings for the server instance</param>
+		public RelayStorageBackend(IOptions<ServerSettings> settings)
 		{
-			ServerSettings CurrentSettings = Settings.Value;
-			if (CurrentSettings.LogRelayServer == null)
+			ServerSettings currentSettings = settings.Value;
+			if (currentSettings.LogRelayServer == null)
 			{
 				throw new InvalidDataException("Missing LogRelayServer in server configuration");
 			}
-			if (CurrentSettings.LogRelayBearerToken == null)
+			if (currentSettings.LogRelayBearerToken == null)
 			{
 				throw new InvalidDataException("Missing LogRelayBearerToken in server configuration");
 			}
 
-			ServerUrl = new Uri(CurrentSettings.LogRelayServer);
+			_serverUrl = new Uri(currentSettings.LogRelayServer);
 
-			Client = new HttpClient();
-			Client.DefaultRequestHeaders.Authorization = new AuthenticationHeaderValue("Bearer", CurrentSettings.LogRelayBearerToken);
+			_client = new HttpClient();
+			_client.DefaultRequestHeaders.Authorization = new AuthenticationHeaderValue("Bearer", currentSettings.LogRelayBearerToken);
 
-			LocalStorage = new FileSystemStorageBackend(new FileSystemSettings { BaseDir = "RelayCache" });
+			_localStorage = new FileSystemStorageBackend(new FileSystemSettings { BaseDir = "RelayCache" });
 		}
 
 		/// <inheritdoc/>
 		public void Dispose()
 		{
-			Client.Dispose();
+			_client.Dispose();
 		}
 
 		/// <inheritdoc/>
-		public async Task<Stream?> ReadAsync(string Path)
+		public async Task<Stream?> ReadAsync(string path)
 		{
-			Stream? LocalResult = await LocalStorage.ReadAsync(Path);
-			if (LocalResult != null)
+			Stream? localResult = await _localStorage.ReadAsync(path);
+			if (localResult != null)
 			{
-				return LocalResult;
+				return localResult;
 			}
 
-			Uri Url = new Uri(ServerUrl, $"api/v1/debug/storage?path={Path}");
-			using (HttpResponseMessage Response = await Client.GetAsync(Url))
+			Uri url = new Uri(_serverUrl, $"api/v1/debug/storage?path={path}");
+			using (HttpResponseMessage response = await _client.GetAsync(url))
 			{
-				if (Response.IsSuccessStatusCode)
+				if (response.IsSuccessStatusCode)
 				{
-					byte[] ResponseData = await Response.Content.ReadAsByteArrayAsync();
-					await LocalStorage.WriteBytesAsync(Path, ResponseData);
-					return new MemoryStream(ResponseData);
+					byte[] responseData = await response.Content.ReadAsByteArrayAsync();
+					await _localStorage.WriteBytesAsync(path, responseData);
+					return new MemoryStream(responseData);
 				}
 			}
 
@@ -91,19 +88,19 @@ namespace Horde.Build.Storage.Backends
 		}
 
 		/// <inheritdoc/>
-		public Task WriteAsync(string Path, Stream Stream)
+		public Task WriteAsync(string path, Stream stream)
 		{
-			return LocalStorage.WriteAsync(Path, Stream);
+			return _localStorage.WriteAsync(path, stream);
 		}
 
 		/// <inheritdoc/>
-		public Task<bool> ExistsAsync(string Path)
+		public Task<bool> ExistsAsync(string path)
 		{
 			throw new NotSupportedException();
 		}
 
 		/// <inheritdoc/>
-		public Task DeleteAsync(string Path)
+		public Task DeleteAsync(string path)
 		{
 			throw new NotImplementedException();
 		}
diff --git a/Engine/Source/Programs/Horde/Horde.Build/Storage/Backends/TransientStorageBackend.cs b/Engine/Source/Programs/Horde/Horde.Build/Storage/Backends/TransientStorageBackend.cs
--- a/Engine/Source/Programs/Horde/Horde.Build/Storage/Backends/TransientStorageBackend.cs
+++ b/Engine/Source/Programs/Horde/Horde.Build/Storage/Backends/TransientStorageBackend.cs
@@ -1,10 +1,7 @@
 // Copyright Epic Games, Inc. All Rights Reserved.
 
-using System;
 using System.Collections.Concurrent;
-using System.Collections.Generic;
 using System.IO;
-using System.Linq;
 using System.Threading.Tasks;
 
 namespace Horde.Build.Storage.Backends
@@ -17,15 +14,15 @@ namespace Horde.Build.Storage.Backends
 		/// <summary>
 		/// Data storage
 		/// </summary>
-		ConcurrentDictionary<string, byte[]> PathToData = new ConcurrentDictionary<string, byte[]>();
+		readonly ConcurrentDictionary<string, byte[]> _pathToData = new ConcurrentDictionary<string, byte[]>();
 
 		/// <inheritdoc/>
-		public Task<Stream?> ReadAsync(string Path)
+		public Task<Stream?> ReadAsync(string path)
 		{
-			byte[]? Data;
-			if (PathToData.TryGetValue(Path, out Data))
+			byte[]? data;
+			if (_pathToData.TryGetValue(path, out data))
 			{
-				return Task.FromResult<Stream?>(new MemoryStream(Data, false));
+				return Task.FromResult<Stream?>(new MemoryStream(data, false));
 			}
 			else
 			{
@@ -34,25 +31,25 @@ namespace Horde.Build.Storage.Backends
 		}
 
 		/// <inheritdoc/>
-		public async Task WriteAsync(string Path, Stream Stream)
+		public async Task WriteAsync(string path, Stream stream)
 		{
-			using (MemoryStream Buffer = new MemoryStream())
+			using (MemoryStream buffer = new MemoryStream())
 			{
-				await Stream.CopyToAsync(Buffer);
-				PathToData[Path] = Buffer.ToArray();
+				await stream.CopyToAsync(buffer);
+				_pathToData[path] = buffer.ToArray();
 			}
 		}
 
 		/// <inheritdoc/>
-		public Task<bool> ExistsAsync(string Path)
+		public Task<bool> ExistsAsync(string path)
 		{
-			return Task.FromResult(PathToData.ContainsKey(Path));
+			return Task.FromResult(_pathToData.ContainsKey(path));
 		}
 
 		/// <inheritdoc/>
-		public Task DeleteAsync(string Path)
+		public Task DeleteAsync(string path)
 		{
-			PathToData.TryRemove(Path, out _);
+			_pathToData.TryRemove(path, out _);
 			return Task.CompletedTask;
 		}
 	}
