Horde: Fix coding conventions within Horde.Build.

diff --git a/Engine/Source/Programs/Horde/Horde.Build/Storage/IStorageBackend.cs b/Engine/Source/Programs/Horde/Horde.Build/Storage/IStorageBackend.cs
--- a/Engine/Source/Programs/Horde/Horde.Build/Storage/IStorageBackend.cs
+++ b/Engine/Source/Programs/Horde/Horde.Build/Storage/IStorageBackend.cs
@@ -1,11 +1,9 @@
 // Copyright Epic Games, Inc. All Rights Reserved.
 
-using EpicGames.Core;
 using System;
-using System.Collections.Generic;
 using System.IO;
-using System.Linq;
 using System.Threading.Tasks;
+using EpicGames.Core;
 
 namespace Horde.Build.Storage
 {
@@ -17,30 +15,30 @@ namespace Horde.Build.Storage
 		/// <summary>
 		/// Opens a read stream for the given path.
 		/// </summary>
-		/// <param name="Path">Relative path within the bucket</param>
+		/// <param name="path">Relative path within the bucket</param>
 		/// <returns></returns>
-		Task<Stream?> ReadAsync(string Path);
+		Task<Stream?> ReadAsync(string path);
 
 		/// <summary>
 		/// Writes a stream to the given path. If the stream throws an exception during read, the write will be aborted.
 		/// </summary>
-		/// <param name="Path">Relative path within the bucket</param>
-		/// <param name="Stream">Stream to write</param>
-		Task WriteAsync(string Path, Stream Stream);
+		/// <param name="path">Relative path within the bucket</param>
+		/// <param name="stream">Stream to write</param>
+		Task WriteAsync(string path, Stream stream);
 
 		/// <summary>
 		/// Tests whether the given path exists
 		/// </summary>
-		/// <param name="Path">Relative path within the bucket</param>
+		/// <param name="path">Relative path within the bucket</param>
 		/// <returns></returns>
-		Task<bool> ExistsAsync(string Path);
+		Task<bool> ExistsAsync(string path);
 
 		/// <summary>
 		/// Deletes a file with the given path
 		/// </summary>
-		/// <param name="Path">Relative path within the bucket</param>
+		/// <param name="path">Relative path within the bucket</param>
 		/// <returns>Async task</returns>
-		Task DeleteAsync(string Path);
+		Task DeleteAsync(string path);
 	}
 
 	/// <summary>
@@ -62,57 +60,57 @@ namespace Horde.Build.Storage
 		/// <typeparam name="T"></typeparam>
 		class StorageBackend<T> : IStorageBackend<T>
 		{
-			IStorageBackend Inner;
+			readonly IStorageBackend _inner;
 
 			/// <summary>
 			/// Constructor
 			/// </summary>
-			/// <param name="Inner"></param>
-			public StorageBackend(IStorageBackend Inner) => this.Inner = Inner;
+			/// <param name="inner"></param>
+			public StorageBackend(IStorageBackend inner) => _inner = inner;
 
 			/// <inheritdoc/>
-			public Task<Stream?> ReadAsync(string Path) => Inner.ReadAsync(Path);
+			public Task<Stream?> ReadAsync(string path) => _inner.ReadAsync(path);
 
 			/// <inheritdoc/>
-			public Task WriteAsync(string Path, Stream Stream) => Inner.WriteAsync(Path, Stream);
+			public Task WriteAsync(string path, Stream stream) => _inner.WriteAsync(path, stream);
 
 			/// <inheritdoc/>
-			public Task DeleteAsync(string Path) => Inner.DeleteAsync(Path);
+			public Task DeleteAsync(string path) => _inner.DeleteAsync(path);
 
 			/// <inheritdoc/>
-			public Task<bool> ExistsAsync(string Path) => Inner.ExistsAsync(Path);
+			public Task<bool> ExistsAsync(string path) => _inner.ExistsAsync(path);
 		}
 
 		/// <summary>
 		/// Creates a typed wrapper around the given storage backend
 		/// </summary>
 		/// <typeparam name="T"></typeparam>
-		/// <param name="Backend"></param>
+		/// <param name="backend"></param>
 		/// <returns></returns>
-		public static IStorageBackend<T> ForType<T>(this IStorageBackend Backend)
+		public static IStorageBackend<T> ForType<T>(this IStorageBackend backend)
 		{
-			return new StorageBackend<T>(Backend);
+			return new StorageBackend<T>(backend);
 		}
 
 		/// <summary>
 		/// Writes a block of memory to storage
 		/// </summary>
-		/// <param name="StorageBackend"></param>
-		/// <param name="Path"></param>
+		/// <param name="storageBackend"></param>
+		/// <param name="path"></param>
 		/// <returns></returns>
-		public static async Task<ReadOnlyMemory<byte>?> ReadBytesAsync(this IStorageBackend StorageBackend, string Path)
+		public static async Task<ReadOnlyMemory<byte>?> ReadBytesAsync(this IStorageBackend storageBackend, string path)
 		{
-			using (Stream? InputStream = await StorageBackend.ReadAsync(Path))
+			using (Stream? inputStream = await storageBackend.ReadAsync(path))
 			{
-				if (InputStream == null)
+				if (inputStream == null)
 				{
 					return null;
 				}
 
-				using (MemoryStream OutputStream = new MemoryStream())
+				using (MemoryStream outputStream = new MemoryStream())
 				{
-					await InputStream.CopyToAsync(OutputStream);
-					return OutputStream.ToArray();
+					await inputStream.CopyToAsync(outputStream);
+					return outputStream.ToArray();
 				}
 			}
 		}
@@ -120,15 +118,15 @@ namespace Horde.Build.Storage
 		/// <summary>
 		/// Writes a block of memory to storage
 		/// </summary>
-		/// <param name="StorageBackend"></param>
-		/// <param name="Path"></param>
-		/// <param name="Data"></param>
+		/// <param name="storageBackend"></param>
+		/// <param name="path"></param>
+		/// <param name="data"></param>
 		/// <returns></returns>
-		public static async Task WriteBytesAsync(this IStorageBackend StorageBackend, string Path, ReadOnlyMemory<byte> Data)
+		public static async Task WriteBytesAsync(this IStorageBackend storageBackend, string path, ReadOnlyMemory<byte> data)
 		{
-			using (ReadOnlyMemoryStream Stream = new ReadOnlyMemoryStream(Data))
+			using (ReadOnlyMemoryStream stream = new ReadOnlyMemoryStream(data))
 			{
-				await StorageBackend.WriteAsync(Path, Stream);
+				await storageBackend.WriteAsync(path, stream);
 			}
 		}
 	}
diff --git a/Engine/Source/Programs/Horde/Horde.Build/Storage/Services/BasicStorageClient.cs b/Engine/Source/Programs/Horde/Horde.Build/Storage/Services/BasicStorageClient.cs
--- a/Engine/Source/Programs/Horde/Horde.Build/Storage/Services/BasicStorageClient.cs
+++ b/Engine/Source/Programs/Horde/Horde.Build/Storage/Services/BasicStorageClient.cs
@@ -1,24 +1,22 @@
 // Copyright Epic Games, Inc. All Rights Reserved.
 
-using EpicGames.Core;
-using EpicGames.Horde.Storage;
-using EpicGames.Serialization;
-using Horde.Build.Services;
-using Horde.Build.Storage;
-using Horde.Build.Utilities;
-using MongoDB.Bson.Serialization.Attributes;
-using MongoDB.Driver;
 using System;
 using System.Collections.Generic;
 using System.IO;
 using System.Linq;
 using System.Threading;
 using System.Threading.Tasks;
+using EpicGames.Core;
+using EpicGames.Horde.Storage;
+using EpicGames.Serialization;
+using Horde.Build.Services;
+using MongoDB.Bson.Serialization.Attributes;
+using MongoDB.Driver;
 
 namespace Horde.Build.Storage.Services
 {
-	using IRef = EpicGames.Horde.Storage.IRef;
 	using BlobNotFoundException = EpicGames.Horde.Storage.BlobNotFoundException;
+	using IRef = EpicGames.Horde.Storage.IRef;
 
 	/// <summary>
 	/// Simple implementation of <see cref="IStorageClient"/> which uses the current <see cref="IStorageBackend"/> implementation without garbage collection.
@@ -45,209 +43,209 @@ namespace Horde.Build.Storage.Services
 			public CbObject Value => new CbObject(Data);
 		}
 
-		IMongoCollection<RefImpl> Refs;
-		IStorageBackend StorageBackend;
+		readonly IMongoCollection<RefImpl> _refs;
+		readonly IStorageBackend _storageBackend;
 
 		/// <summary>
 		/// 
 		/// </summary>
-		/// <param name="DatabaseService"></param>
-		/// <param name="StorageBackend"></param>
-		public BasicStorageClient(DatabaseService DatabaseService, IStorageBackend<BasicStorageClient> StorageBackend)
+		/// <param name="databaseService"></param>
+		/// <param name="storageBackend"></param>
+		public BasicStorageClient(DatabaseService databaseService, IStorageBackend<BasicStorageClient> storageBackend)
 		{
-			Refs = DatabaseService.GetCollection<RefImpl>("Storage.RefsV2");
-			this.StorageBackend = StorageBackend;
+			_refs = databaseService.GetCollection<RefImpl>("Storage.RefsV2");
+			_storageBackend = storageBackend;
 		}
 
 		/// <inheritdoc/>
-		public async Task<Stream> ReadBlobAsync(NamespaceId NamespaceId, IoHash Hash, CancellationToken CancellationToken)
+		public async Task<Stream> ReadBlobAsync(NamespaceId namespaceId, IoHash hash, CancellationToken cancellationToken)
 		{
-			string Path = GetFullBlobPath(NamespaceId, Hash);
+			string path = GetFullBlobPath(namespaceId, hash);
 
-			Stream? Stream = await StorageBackend.ReadAsync(Path);
-			if(Stream == null)
+			Stream? stream = await _storageBackend.ReadAsync(path);
+			if(stream == null)
 			{
-				throw new BlobNotFoundException(NamespaceId, Hash);
+				throw new BlobNotFoundException(namespaceId, hash);
 			}
 
-			return Stream;
+			return stream;
 		}
 
 		/// <inheritdoc/>
-		public async Task WriteBlobAsync(NamespaceId NamespaceId, IoHash Hash, Stream Stream, CancellationToken CancellationToken)
+		public async Task WriteBlobAsync(NamespaceId namespaceId, IoHash hash, Stream stream, CancellationToken cancellationToken)
 		{
-			byte[] Data;
-			using (MemoryStream Buffer = new MemoryStream())
+			byte[] data;
+			using (MemoryStream buffer = new MemoryStream())
 			{
-				await Stream.CopyToAsync(Buffer, CancellationToken);
-				Data = Buffer.ToArray();
+				await stream.CopyToAsync(buffer, cancellationToken);
+				data = buffer.ToArray();
 			}
 
-			IoHash DataHash = IoHash.Compute(Data);
-			if (Hash != DataHash)
+			IoHash dataHash = IoHash.Compute(data);
+			if (hash != dataHash)
 			{
-				throw new InvalidDataException($"Incorrect hash for data (was {DataHash}, expected {Hash})");
+				throw new InvalidDataException($"Incorrect hash for data (was {dataHash}, expected {hash})");
 			}
 
-			string Path = GetFullBlobPath(NamespaceId, Hash);
-			await StorageBackend.WriteBytesAsync(Path, Data);
+			string path = GetFullBlobPath(namespaceId, hash);
+			await _storageBackend.WriteBytesAsync(path, data);
 		}
 
 		/// <inheritdoc/>
-		public async Task<bool> HasBlobAsync(NamespaceId NamespaceId, IoHash Hash, CancellationToken CancellationToken)
+		public async Task<bool> HasBlobAsync(NamespaceId namespaceId, IoHash hash, CancellationToken cancellationToken)
 		{
-			return await StorageBackend.ExistsAsync(GetFullBlobPath(NamespaceId, Hash));
+			return await _storageBackend.ExistsAsync(GetFullBlobPath(namespaceId, hash));
 		}
 
 		/// <inheritdoc/>
-		public async Task<HashSet<IoHash>> FindMissingBlobsAsync(NamespaceId NamespaceId, HashSet<IoHash> Hashes, CancellationToken CancellationToken)
+		public async Task<HashSet<IoHash>> FindMissingBlobsAsync(NamespaceId namespaceId, HashSet<IoHash> hashes, CancellationToken cancellationToken)
 		{
-			var Tasks = Hashes.Select(async Hash => (Hash, await StorageBackend.ExistsAsync(GetFullBlobPath(NamespaceId, Hash))));
-			await Task.WhenAll(Tasks);
-			return Tasks.Where(x => !x.Result.Item2).Select(x => x.Result.Item1).ToHashSet();
+			IEnumerable<Task<(IoHash Hash, bool Exists)>> tasks = hashes.Select(async hash => (Hash: hash, await _storageBackend.ExistsAsync(GetFullBlobPath(namespaceId, hash))));
+			await Task.WhenAll(tasks);
+			return tasks.Where(x => !x.Result.Exists).Select(x => x.Result.Hash).ToHashSet();
 		}
 
 		/// <inheritdoc/>
-		public async Task<IRef> GetRefAsync(NamespaceId NamespaceId, BucketId BucketId, RefId RefId, CancellationToken CancellationToken)
+		public async Task<IRef> GetRefAsync(NamespaceId namespaceId, BucketId bucketId, RefId refId, CancellationToken cancellationToken)
 		{
-			string Id = GetFullRefId(NamespaceId, BucketId, RefId);
+			string id = GetFullRefId(namespaceId, bucketId, refId);
 
-			RefImpl? RefImpl = await Refs.FindOneAndUpdateAsync<RefImpl>(x => x.Id == Id, Builders<RefImpl>.Update.Set(x => x.LastAccessTime, DateTime.UtcNow), new FindOneAndUpdateOptions<RefImpl> { ReturnDocument = ReturnDocument.After }, CancellationToken);
-			if (RefImpl == null || !RefImpl.Finalized)
+			RefImpl? refImpl = await _refs.FindOneAndUpdateAsync<RefImpl>(x => x.Id == id, Builders<RefImpl>.Update.Set(x => x.LastAccessTime, DateTime.UtcNow), new FindOneAndUpdateOptions<RefImpl> { ReturnDocument = ReturnDocument.After }, cancellationToken);
+			if (refImpl == null || !refImpl.Finalized)
 			{
-				throw new RefNotFoundException(NamespaceId, BucketId, RefId);
+				throw new RefNotFoundException(namespaceId, bucketId, refId);
 			}
 
-			return RefImpl;
+			return refImpl;
 		}
 
 		/// <inheritdoc/>
-		public async Task<bool> HasRefAsync(NamespaceId NamespaceId, BucketId BucketId, RefId RefId, CancellationToken CancellationToken)
+		public async Task<bool> HasRefAsync(NamespaceId namespaceId, BucketId bucketId, RefId refId, CancellationToken cancellationToken)
 		{
-			string Id = GetFullRefId(NamespaceId, BucketId, RefId);
-			RefImpl? RefImpl = await Refs.Find<RefImpl>(x => x.Id == Id).FirstOrDefaultAsync(CancellationToken);
-			return RefImpl != null;
+			string id = GetFullRefId(namespaceId, bucketId, refId);
+			RefImpl? refImpl = await _refs.Find<RefImpl>(x => x.Id == id).FirstOrDefaultAsync(cancellationToken);
+			return refImpl != null;
 		}
 
 		/// <inheritdoc/>
-		public async Task<List<IoHash>> TrySetRefAsync(NamespaceId NamespaceId, BucketId BucketId, RefId RefId, CbObject Value, CancellationToken CancellationToken)
+		public async Task<List<IoHash>> TrySetRefAsync(NamespaceId namespaceId, BucketId bucketId, RefId refId, CbObject value, CancellationToken cancellationToken)
 		{
-			RefImpl NewRef = new RefImpl();
-			NewRef.Id = GetFullRefId(NamespaceId, BucketId, RefId);
-			NewRef.Data = Value.GetView().ToArray();
-			await Refs.ReplaceOneAsync(x => x.Id == NewRef.Id, NewRef, new ReplaceOptions { IsUpsert = true }, CancellationToken);
+			RefImpl newRef = new RefImpl();
+			newRef.Id = GetFullRefId(namespaceId, bucketId, refId);
+			newRef.Data = value.GetView().ToArray();
+			await _refs.ReplaceOneAsync(x => x.Id == newRef.Id, newRef, new ReplaceOptions { IsUpsert = true }, cancellationToken);
 
-			IoHash Hash = IoHash.Compute(Value.GetView().Span);
-			return await TryFinalizeRefAsync(NamespaceId, BucketId, RefId, Hash, CancellationToken);
+			IoHash hash = IoHash.Compute(value.GetView().Span);
+			return await TryFinalizeRefAsync(namespaceId, bucketId, refId, hash, cancellationToken);
 		}
 
 		/// <inheritdoc/>
-		public async Task<List<IoHash>> TryFinalizeRefAsync(NamespaceId NamespaceId, BucketId BucketId, RefId RefId, IoHash Hash, CancellationToken CancellationToken)
+		public async Task<List<IoHash>> TryFinalizeRefAsync(NamespaceId namespaceId, BucketId bucketId, RefId refId, IoHash hash, CancellationToken cancellationToken)
 		{
-			string Id = GetFullRefId(NamespaceId, BucketId, RefId);
+			string id = GetFullRefId(namespaceId, bucketId, refId);
 
-			RefImpl Ref = await Refs.Find(x => x.Id == Id).FirstAsync(CancellationToken);
+			RefImpl refImpl = await _refs.Find(x => x.Id == id).FirstAsync(cancellationToken);
 
-			HashSet<IoHash> MissingHashes = new HashSet<IoHash>();
-			if (!Ref.Finalized)
+			HashSet<IoHash> missingHashes = new HashSet<IoHash>();
+			if (!refImpl.Finalized)
 			{
-				HashSet<IoHash> CheckedObjectHashes = new HashSet<IoHash>();
-				HashSet<IoHash> CheckedBinaryHashes = new HashSet<IoHash>();
-				await FinalizeInternalAsync(Ref.NamespaceId, Ref.Value, CheckedObjectHashes, CheckedBinaryHashes, MissingHashes);
+				HashSet<IoHash> checkedObjectHashes = new HashSet<IoHash>();
+				HashSet<IoHash> checkedBinaryHashes = new HashSet<IoHash>();
+				await FinalizeInternalAsync(refImpl.NamespaceId, refImpl.Value, checkedObjectHashes, checkedBinaryHashes, missingHashes);
 
-				if (MissingHashes.Count == 0)
+				if (missingHashes.Count == 0)
 				{
-					await Refs.UpdateOneAsync(x => x.Id == Id, Builders<RefImpl>.Update.Set(x => x.Finalized, true), cancellationToken: CancellationToken);
+					await _refs.UpdateOneAsync(x => x.Id == id, Builders<RefImpl>.Update.Set(x => x.Finalized, true), cancellationToken: cancellationToken);
 				}
 			}
-			return MissingHashes.ToList();
+			return missingHashes.ToList();
 		}
 
-		async Task FinalizeInternalAsync(NamespaceId NamespaceId, CbObject Object, HashSet<IoHash> CheckedObjectHashes, HashSet<IoHash> CheckedBinaryHashes, HashSet<IoHash> MissingHashes)
+		async Task FinalizeInternalAsync(NamespaceId namespaceId, CbObject obj, HashSet<IoHash> checkedObjectHashes, HashSet<IoHash> checkedBinaryHashes, HashSet<IoHash> missingHashes)
 		{
-			List<IoHash> NewObjectHashes = new List<IoHash>();
-			List<IoHash> NewBinaryHashes = new List<IoHash>();
-			Object.IterateAttachments(Field =>
+			List<IoHash> newObjectHashes = new List<IoHash>();
+			List<IoHash> newBinaryHashes = new List<IoHash>();
+			obj.IterateAttachments(field =>
 			{
-				IoHash AttachmentHash = Field.AsAttachment();
-				if (Field.IsObjectAttachment())
+				IoHash attachmentHash = field.AsAttachment();
+				if (field.IsObjectAttachment())
 				{
-					if (CheckedObjectHashes.Add(AttachmentHash))
+					if (checkedObjectHashes.Add(attachmentHash))
 					{
-						NewObjectHashes.Add(AttachmentHash);
+						newObjectHashes.Add(attachmentHash);
 					}
 				}
 				else
 				{
-					if (CheckedBinaryHashes.Add(AttachmentHash))
+					if (checkedBinaryHashes.Add(attachmentHash))
 					{
-						NewBinaryHashes.Add(AttachmentHash);
+						newBinaryHashes.Add(attachmentHash);
 					}
 				}
 			});
 
-			foreach (IoHash NewHash in NewObjectHashes)
+			foreach (IoHash newHash in newObjectHashes)
 			{
-				string ObjPath = GetFullBlobPath(NamespaceId, NewHash);
+				string objPath = GetFullBlobPath(namespaceId, newHash);
 
-				ReadOnlyMemory<byte>? Data = await StorageBackend.ReadBytesAsync(ObjPath);
-				if (Data == null)
+				ReadOnlyMemory<byte>? data = await _storageBackend.ReadBytesAsync(objPath);
+				if (data == null)
 				{
-					MissingHashes.Add(NewHash);
+					missingHashes.Add(newHash);
 				}
 				else
 				{
-					await FinalizeInternalAsync(NamespaceId, new CbObject(Data.Value), CheckedObjectHashes, CheckedBinaryHashes, MissingHashes);
+					await FinalizeInternalAsync(namespaceId, new CbObject(data.Value), checkedObjectHashes, checkedBinaryHashes, missingHashes);
 				}
 			}
 
-			foreach (IoHash NewBinaryHash in NewBinaryHashes)
+			foreach (IoHash newBinaryHash in newBinaryHashes)
 			{
-				string Path = GetFullBlobPath(NamespaceId, NewBinaryHash);
-				if (!await StorageBackend.ExistsAsync(Path))
+				string path = GetFullBlobPath(namespaceId, newBinaryHash);
+				if (!await _storageBackend.ExistsAsync(path))
 				{
-					MissingHashes.Add(NewBinaryHash);
+					missingHashes.Add(newBinaryHash);
 				}
 			}
 		}
 
 		/// <inheritdoc/>
-		public async Task<bool> DeleteRefAsync(NamespaceId NamespaceId, BucketId BucketId, RefId RefId, CancellationToken CancellationToken)
+		public async Task<bool> DeleteRefAsync(NamespaceId namespaceId, BucketId bucketId, RefId refId, CancellationToken cancellationToken)
 		{
-			string Id = GetFullRefId(NamespaceId, BucketId, RefId);
-			DeleteResult Result = await Refs.DeleteOneAsync(x => x.Id == Id, CancellationToken);
-			return Result.DeletedCount > 0;
+			string id = GetFullRefId(namespaceId, bucketId, refId);
+			DeleteResult result = await _refs.DeleteOneAsync(x => x.Id == id, cancellationToken);
+			return result.DeletedCount > 0;
 		}
 
 		/// <inheritdoc/>
-		public async Task<List<RefId>> FindMissingRefsAsync(NamespaceId NamespaceId, BucketId BucketId, List<RefId> RefIds, CancellationToken CancellationToken)
+		public async Task<List<RefId>> FindMissingRefsAsync(NamespaceId namespaceId, BucketId bucketId, List<RefId> refIds, CancellationToken cancellationToken)
 		{
-			List<RefId> MissingRefIds = new List<RefId>();
-			foreach (RefId RefId in RefIds)
+			List<RefId> missingRefIds = new List<RefId>();
+			foreach (RefId refId in refIds)
 			{
-				if(!await HasRefAsync(NamespaceId, BucketId, RefId, CancellationToken))
+				if(!await HasRefAsync(namespaceId, bucketId, refId, cancellationToken))
 				{
-					MissingRefIds.Add(RefId);
+					missingRefIds.Add(refId);
 				}
 			}
-			return MissingRefIds;
+			return missingRefIds;
 		}
 
 		/// <summary>
 		/// Gets the formatted id for a particular ref
 		/// </summary>
-		static string GetFullRefId(NamespaceId NamespaceId, BucketId BucketId, RefId RefId)
+		static string GetFullRefId(NamespaceId namespaceId, BucketId bucketId, RefId refId)
 		{
-			return $"{NamespaceId}/{BucketId}/{RefId}";
+			return $"{namespaceId}/{bucketId}/{refId}";
 		}
 
 		/// <summary>
 		/// Gets the path for a blob
 		/// </summary>
-		static string GetFullBlobPath(NamespaceId NamespaceId, IoHash Hash)
+		static string GetFullBlobPath(NamespaceId namespaceId, IoHash hash)
 		{
-			string HashText = Hash.ToString();
-			return $"blobs2/{NamespaceId}/{HashText[0..2]}/{HashText[2..4]}/{HashText}.blob";
+			string hashText = hash.ToString();
+			return $"blobs2/{namespaceId}/{hashText[0..2]}/{hashText[2..4]}/{hashText}.blob";
 		}
 	}
 }
diff --git a/Engine/Source/Programs/Horde/Horde.Build/Ugs/IUgsMetadataCollection.cs b/Engine/Source/Programs/Horde/Horde.Build/Ugs/IUgsMetadataCollection.cs
--- a/Engine/Source/Programs/Horde/Horde.Build/Ugs/IUgsMetadataCollection.cs
+++ b/Engine/Source/Programs/Horde/Horde.Build/Ugs/IUgsMetadataCollection.cs
@@ -1,18 +1,10 @@
 // Copyright Epic Games, Inc. All Rights Reserved.
 
-using Horde.Build.Api;
-using Horde.Build.Models;
-using Horde.Build.Services;
-using Horde.Build.Utilities;
-using MongoDB.Bson;
-using MongoDB.Bson.Serialization.Attributes;
-using MongoDB.Driver;
-using MongoDB.Driver.Core.Authentication;
 using System;
 using System.Collections.Generic;
-using System.Diagnostics.CodeAnalysis;
-using System.Linq;
 using System.Threading.Tasks;
+using Horde.Build.Api;
+using Horde.Build.Models;
 
 namespace Horde.Build.Collections
 {
@@ -24,43 +16,43 @@ namespace Horde.Build.Collections
 		/// <summary>
 		/// Finds or adds metadata for a given changelist
 		/// </summary>
-		/// <param name="Stream">The stream containing the change</param>
-		/// <param name="Change">The changelist number</param>
-		/// <param name="Project">Project identifier</param>
+		/// <param name="stream">The stream containing the change</param>
+		/// <param name="change">The changelist number</param>
+		/// <param name="project">Project identifier</param>
 		/// <returns>The metadata instance</returns>
-		Task<IUgsMetadata> FindOrAddAsync(string Stream, int Change, string? Project);
+		Task<IUgsMetadata> FindOrAddAsync(string stream, int change, string? project);
 
 		/// <summary>
 		/// Adds information to a change
 		/// </summary>
-		/// <param name="Metadata">The existing metadata</param>
-		/// <param name="UserName">Name of the user posting the change</param>
-		/// <param name="Synced">Time that the change was synced</param>
-		/// <param name="Vote">Vote for this change</param>
-		/// <param name="Investigating">Whether the user is investigating the change</param>
-		/// <param name="Starred">Whether the change is starred</param>
-		/// <param name="Comment">New comment from this user</param>
+		/// <param name="metadata">The existing metadata</param>
+		/// <param name="userName">Name of the user posting the change</param>
+		/// <param name="synced">Time that the change was synced</param>
+		/// <param name="vote">Vote for this change</param>
+		/// <param name="investigating">Whether the user is investigating the change</param>
+		/// <param name="starred">Whether the change is starred</param>
+		/// <param name="comment">New comment from this user</param>
 		/// <returns>Async task task</returns>
-		Task<IUgsMetadata> UpdateUserAsync(IUgsMetadata Metadata, string UserName, bool? Synced, UgsUserVote? Vote, bool? Investigating, bool? Starred, string? Comment);
+		Task<IUgsMetadata> UpdateUserAsync(IUgsMetadata metadata, string userName, bool? synced, UgsUserVote? vote, bool? investigating, bool? starred, string? comment);
 
 		/// <summary>
 		/// Updates the state of a badge
 		/// </summary>
-		/// <param name="Metadata">The existing metadata</param>
-		/// <param name="Name">Name of the badge</param>
-		/// <param name="Url">Url to link to for the badge</param>
-		/// <param name="State">State of the badge</param>
+		/// <param name="metadata">The existing metadata</param>
+		/// <param name="name">Name of the badge</param>
+		/// <param name="url">Url to link to for the badge</param>
+		/// <param name="state">State of the badge</param>
 		/// <returns>Async task</returns>
-		Task<IUgsMetadata> UpdateBadgeAsync(IUgsMetadata Metadata, string Name, Uri? Url, UgsBadgeState State);
+		Task<IUgsMetadata> UpdateBadgeAsync(IUgsMetadata metadata, string name, Uri? url, UgsBadgeState state);
 
 		/// <summary>
 		/// Searches for metadata updates
 		/// </summary>
-		/// <param name="Stream">The stream to search</param>
-		/// <param name="MinChange">Minimum changelist number</param>
-		/// <param name="MaxChange">Maximum changelist number</param>
-		/// <param name="AfterTicks">Last query time</param>
+		/// <param name="stream">The stream to search</param>
+		/// <param name="minChange">Minimum changelist number</param>
+		/// <param name="maxChange">Maximum changelist number</param>
+		/// <param name="afterTicks">Last query time</param>
 		/// <returns>List of metadata updates</returns>
-		Task<List<IUgsMetadata>> FindAsync(string Stream, int MinChange, int? MaxChange = null, long? AfterTicks = null);
+		Task<List<IUgsMetadata>> FindAsync(string stream, int minChange, int? maxChange = null, long? afterTicks = null);
 	}
 }
diff --git a/Engine/Source/Programs/Horde/Horde.Build/Ugs/Impl/UgsMetadataCollection.cs b/Engine/Source/Programs/Horde/Horde.Build/Ugs/Impl/UgsMetadataCollection.cs
--- a/Engine/Source/Programs/Horde/Horde.Build/Ugs/Impl/UgsMetadataCollection.cs
+++ b/Engine/Source/Programs/Horde/Horde.Build/Ugs/Impl/UgsMetadataCollection.cs
@@ -1,5 +1,10 @@
 // Copyright Epic Games, Inc. All Rights Reserved.
 
+using System;
+using System.Collections.Generic;
+using System.Diagnostics.CodeAnalysis;
+using System.Linq;
+using System.Threading.Tasks;
 using Horde.Build.Api;
 using Horde.Build.Models;
 using Horde.Build.Services;
@@ -7,12 +12,6 @@ using Horde.Build.Utilities;
 using MongoDB.Bson;
 using MongoDB.Bson.Serialization.Attributes;
 using MongoDB.Driver;
-using MongoDB.Driver.Core.Authentication;
-using System;
-using System.Collections.Generic;
-using System.Diagnostics.CodeAnalysis;
-using System.Linq;
-using System.Threading.Tasks;
 
 namespace Horde.Build.Collections.Impl
 {
@@ -43,9 +42,9 @@ namespace Horde.Build.Collections.Impl
 			long? IUgsUserData.SyncTime => SyncTime;
 			UgsUserVote IUgsUserData.Vote => Vote ?? UgsUserVote.None;
 
-			public UgsUserDataDocument(string User)
+			public UgsUserDataDocument(string user)
 			{
-				this.User = User;
+				User = user;
 			}
 		}
 
@@ -76,34 +75,34 @@ namespace Horde.Build.Collections.Impl
 			[BsonConstructor]
 			private UgsMetadataDocument()
 			{
-				this.Stream = String.Empty;
+				Stream = String.Empty;
 			}
 
-			public UgsMetadataDocument(string Stream, int Change, string Project)
+			public UgsMetadataDocument(string stream, int change, string project)
 			{
-				this.Id = ObjectId.GenerateNewId();
-				this.Stream = Stream;
-				this.Change = Change;
-				this.Project = Project;
+				Id = ObjectId.GenerateNewId();
+				Stream = stream;
+				Change = change;
+				Project = project;
 			}
 		}
 
-		IMongoCollection<UgsMetadataDocument> Collection;
+		readonly IMongoCollection<UgsMetadataDocument> _collection;
 
 		/// <summary>
 		/// Constructor
 		/// </summary>
-		/// <param name="DatabaseService">Database service instance</param>
-		public UgsMetadataCollection(DatabaseService DatabaseService)
+		/// <param name="databaseService">Database service instance</param>
+		public UgsMetadataCollection(DatabaseService databaseService)
 		{
-			Collection = DatabaseService.GetCollection<UgsMetadataDocument>("UgsMetadata");
+			_collection = databaseService.GetCollection<UgsMetadataDocument>("UgsMetadata");
 			
-			if (!DatabaseService.ReadOnlyMode)
+			if (!databaseService.ReadOnlyMode)
 			{
 				void CreateIndex()
 				{
-					Collection.Indexes.CreateOne(new CreateIndexModel<UgsMetadataDocument>(Builders<UgsMetadataDocument>.IndexKeys.Ascending(x => x.Stream).Descending(x => x.Change).Ascending(x => x.Project), new CreateIndexOptions { Unique = true }));
-					Collection.Indexes.CreateOne(new CreateIndexModel<UgsMetadataDocument>(Builders<UgsMetadataDocument>.IndexKeys.Ascending(x => x.Stream).Descending(x => x.Change).Descending(x => x.UpdateTicks)));
+					_collection.Indexes.CreateOne(new CreateIndexModel<UgsMetadataDocument>(Builders<UgsMetadataDocument>.IndexKeys.Ascending(x => x.Stream).Descending(x => x.Change).Ascending(x => x.Project), new CreateIndexOptions { Unique = true }));
+					_collection.Indexes.CreateOne(new CreateIndexModel<UgsMetadataDocument>(Builders<UgsMetadataDocument>.IndexKeys.Ascending(x => x.Stream).Descending(x => x.Change).Descending(x => x.UpdateTicks)));
 				}
 				
 				try
@@ -113,7 +112,7 @@ namespace Horde.Build.Collections.Impl
 				catch (MongoCommandException)
 				{
 					// If index creation fails, drop all indexes and retry.
-					Collection.Indexes.DropAll();
+					_collection.Indexes.DropAll();
 					CreateIndex();
 				}
 			}
@@ -122,33 +121,33 @@ namespace Horde.Build.Collections.Impl
 		/// <summary>
 		/// Find or add a document for the given change
 		/// </summary>
-		/// <param name="Stream">Stream containing the change</param>
-		/// <param name="Change">The changelist number to add a document for</param>
-		/// <param name="Project">Arbitrary identifier for this project</param>
+		/// <param name="stream">Stream containing the change</param>
+		/// <param name="change">The changelist number to add a document for</param>
+		/// <param name="project">Arbitrary identifier for this project</param>
 		/// <returns>The metadata document</returns>
-		public async Task<IUgsMetadata> FindOrAddAsync(string Stream, int Change, string? Project)
+		public async Task<IUgsMetadata> FindOrAddAsync(string stream, int change, string? project)
 		{
-			string NormalizedStream = GetNormalizedStream(Stream);
-			string NormalizedProject = GetNormalizedProject(Project);
+			string normalizedStream = GetNormalizedStream(stream);
+			string normalizedProject = GetNormalizedProject(project);
 			for (; ; )
 			{
 				// Find an existing document
-				UgsMetadataDocument? Existing = await Collection.Find(x => x.Stream == NormalizedStream && x.Change == Change && x.Project == NormalizedProject).FirstOrDefaultAsync();
-				if (Existing != null)
+				UgsMetadataDocument? existing = await _collection.Find(x => x.Stream == normalizedStream && x.Change == change && x.Project == normalizedProject).FirstOrDefaultAsync();
+				if (existing != null)
 				{
-					return Existing;
+					return existing;
 				}
 
 				// Try to insert a new document
 				try
 				{
-					UgsMetadataDocument NewDocument = new UgsMetadataDocument(NormalizedStream, Change, NormalizedProject);
-					await Collection.InsertOneAsync(NewDocument);
-					return NewDocument;
+					UgsMetadataDocument newDocument = new UgsMetadataDocument(normalizedStream, change, normalizedProject);
+					await _collection.InsertOneAsync(newDocument);
+					return newDocument;
 				}
-				catch (MongoWriteException Ex)
+				catch (MongoWriteException ex)
 				{
-					if (Ex.WriteError.Category != ServerErrorCategory.DuplicateKey)
+					if (ex.WriteError.Category != ServerErrorCategory.DuplicateKey)
 					{
 						throw;
 					}
@@ -157,149 +156,149 @@ namespace Horde.Build.Collections.Impl
 		}
 
 		/// <inheritdoc/>
-		public async Task<IUgsMetadata> UpdateUserAsync(IUgsMetadata Metadata, string UserName, bool? Synced, UgsUserVote? Vote, bool? Investigating, bool? Starred, string? Comment)
+		public async Task<IUgsMetadata> UpdateUserAsync(IUgsMetadata metadata, string userName, bool? synced, UgsUserVote? vote, bool? investigating, bool? starred, string? comment)
 		{
-			UpdateDefinitionBuilder<UgsMetadataDocument> UpdateBuilder = Builders<UgsMetadataDocument>.Update;
+			UpdateDefinitionBuilder<UgsMetadataDocument> updateBuilder = Builders<UgsMetadataDocument>.Update;
 			for (; ; )
 			{
-				UgsMetadataDocument Document = (UgsMetadataDocument)Metadata;
+				UgsMetadataDocument document = (UgsMetadataDocument)metadata;
 
-				int UserIdx = Document.Users.FindIndex(x => x.User.Equals(UserName, StringComparison.OrdinalIgnoreCase));
-				if (UserIdx == -1)
+				int userIdx = document.Users.FindIndex(x => x.User.Equals(userName, StringComparison.OrdinalIgnoreCase));
+				if (userIdx == -1)
 				{
 					// Create a new user entry
-					UgsUserDataDocument UserData = new UgsUserDataDocument(UserName);
-					UserData.SyncTime = (Synced == true) ? (long?)DateTime.UtcNow.Ticks : null;
-					UserData.Vote = Vote;
-					UserData.Investigating = (Investigating == true) ? Investigating : null;
-					UserData.Starred = (Starred == true) ? Starred : null;
-					UserData.Comment = Comment;
-
-					if (await TryUpdateAsync(Document, UpdateBuilder.Push(x => x.Users, UserData)))
+					UgsUserDataDocument userData = new UgsUserDataDocument(userName);
+					userData.SyncTime = (synced == true) ? (long?)DateTime.UtcNow.Ticks : null;
+					userData.Vote = vote;
+					userData.Investigating = (investigating == true) ? investigating : null;
+					userData.Starred = (starred == true) ? starred : null;
+					userData.Comment = comment;
+
+					if (await TryUpdateAsync(document, updateBuilder.Push(x => x.Users, userData)))
 					{
-						Document.Users.Add(UserData);
-						return Metadata;
+						document.Users.Add(userData);
+						return metadata;
 					}
 				}
 				else
 				{
 					// Update an existing entry
-					UgsUserDataDocument UserData = Document.Users[UserIdx];
+					UgsUserDataDocument userData = document.Users[userIdx];
 
-					long? NewSyncTime = Synced.HasValue ? (Synced.Value ? (long?)DateTime.UtcNow.Ticks : null) : UserData.SyncTime;
-					UgsUserVote? NewVote = Vote.HasValue ? ((Vote.Value != UgsUserVote.None) ? Vote : null) : UserData.Vote;
-					bool? NewInvestigating = Investigating.HasValue ? (Investigating.Value ? Investigating : null) : UserData.Investigating;
-					bool? NewStarred = Starred.HasValue ? (Starred.Value ? Starred : null) : UserData.Starred;
+					long? newSyncTime = synced.HasValue ? (synced.Value ? (long?)DateTime.UtcNow.Ticks : null) : userData.SyncTime;
+					UgsUserVote? newVote = vote.HasValue ? ((vote.Value != UgsUserVote.None) ? vote : null) : userData.Vote;
+					bool? newInvestigating = investigating.HasValue ? (investigating.Value ? investigating : null) : userData.Investigating;
+					bool? newStarred = starred.HasValue ? (starred.Value ? starred : null) : userData.Starred;
 
-					List<UpdateDefinition<UgsMetadataDocument>> Updates = new List<UpdateDefinition<UgsMetadataDocument>>();
-					if (NewSyncTime != UserData.SyncTime)
+					List<UpdateDefinition<UgsMetadataDocument>> updates = new List<UpdateDefinition<UgsMetadataDocument>>();
+					if (newSyncTime != userData.SyncTime)
 					{
-						Updates.Add(UpdateBuilder.Set(x => x.Users![UserIdx].SyncTime, NewSyncTime));
+						updates.Add(updateBuilder.Set(x => x.Users![userIdx].SyncTime, newSyncTime));
 					}
-					if (NewVote != UserData.Vote)
+					if (newVote != userData.Vote)
 					{
-						Updates.Add(UpdateBuilder.Set(x => x.Users![UserIdx].Vote, NewVote));
+						updates.Add(updateBuilder.Set(x => x.Users![userIdx].Vote, newVote));
 					}
-					if (NewInvestigating != UserData.Investigating)
+					if (newInvestigating != userData.Investigating)
 					{
-						Updates.Add(UpdateBuilder.SetOrUnsetNull(x => x.Users![UserIdx].Investigating, NewInvestigating));
+						updates.Add(updateBuilder.SetOrUnsetNull(x => x.Users![userIdx].Investigating, newInvestigating));
 					}
-					if (NewStarred != UserData.Starred)
+					if (newStarred != userData.Starred)
 					{
-						Updates.Add(UpdateBuilder.SetOrUnsetNull(x => x.Users![UserIdx].Starred, NewStarred));
+						updates.Add(updateBuilder.SetOrUnsetNull(x => x.Users![userIdx].Starred, newStarred));
 					}
-					if (Comment != null)
+					if (comment != null)
 					{
-						Updates.Add(UpdateBuilder.Set(x => x.Users![UserIdx].Comment, Comment));
+						updates.Add(updateBuilder.Set(x => x.Users![userIdx].Comment, comment));
 					}
 
-					if (Updates.Count == 0 || await TryUpdateAsync(Document, UpdateBuilder.Combine(Updates)))
+					if (updates.Count == 0 || await TryUpdateAsync(document, updateBuilder.Combine(updates)))
 					{
-						UserData.SyncTime = NewSyncTime;
-						UserData.Vote = Vote;
-						UserData.Investigating = NewInvestigating;
-						UserData.Starred = NewStarred;
-						UserData.Comment = Comment;
-						return Metadata;
+						userData.SyncTime = newSyncTime;
+						userData.Vote = vote;
+						userData.Investigating = newInvestigating;
+						userData.Starred = newStarred;
+						userData.Comment = comment;
+						return metadata;
 					}
 				}
 
 				// Update the document and try again
-				Metadata = await FindOrAddAsync(Metadata.Stream, Metadata.Change, Metadata.Project);
+				metadata = await FindOrAddAsync(metadata.Stream, metadata.Change, metadata.Project);
 			}
 		}
 
 		/// <inheritdoc/>
-		public async Task<IUgsMetadata> UpdateBadgeAsync(IUgsMetadata Metadata, string Name, Uri? Url, UgsBadgeState State)
+		public async Task<IUgsMetadata> UpdateBadgeAsync(IUgsMetadata metadata, string name, Uri? url, UgsBadgeState state)
 		{
-			UpdateDefinitionBuilder<UgsMetadataDocument> UpdateBuilder = Builders<UgsMetadataDocument>.Update;
+			UpdateDefinitionBuilder<UgsMetadataDocument> updateBuilder = Builders<UgsMetadataDocument>.Update;
 			for (; ; )
 			{
-				UgsMetadataDocument Document = (UgsMetadataDocument)Metadata;
+				UgsMetadataDocument document = (UgsMetadataDocument)metadata;
 
 				// Update the document
-				int BadgeIdx = Document.Badges.FindIndex(x => x.Name.Equals(Name, StringComparison.OrdinalIgnoreCase));
-				if (BadgeIdx == -1)
+				int badgeIdx = document.Badges.FindIndex(x => x.Name.Equals(name, StringComparison.OrdinalIgnoreCase));
+				if (badgeIdx == -1)
 				{
 					// Create a new badge
-					UgsBadgeDataDocument NewBadge = new UgsBadgeDataDocument();
-					NewBadge.Name = Name;
-					NewBadge.Url = Url;
-					NewBadge.State = State;
+					UgsBadgeDataDocument newBadge = new UgsBadgeDataDocument();
+					newBadge.Name = name;
+					newBadge.Url = url;
+					newBadge.State = state;
 
-					if (await TryUpdateAsync(Document, UpdateBuilder.Push(x => x.Badges, NewBadge)))
+					if (await TryUpdateAsync(document, updateBuilder.Push(x => x.Badges, newBadge)))
 					{
-						Document.Badges.Add(NewBadge);
-						return Metadata;
+						document.Badges.Add(newBadge);
+						return metadata;
 					}
 				}
 				else
 				{
 					// Update an existing badge
-					List<UpdateDefinition<UgsMetadataDocument>> Updates = new List<UpdateDefinition<UgsMetadataDocument>>();
+					List<UpdateDefinition<UgsMetadataDocument>> updates = new List<UpdateDefinition<UgsMetadataDocument>>();
 
-					UgsBadgeDataDocument BadgeData = Document.Badges[BadgeIdx];
-					if (Url != BadgeData.Url)
+					UgsBadgeDataDocument badgeData = document.Badges[badgeIdx];
+					if (url != badgeData.Url)
 					{
-						Updates.Add(UpdateBuilder.Set(x => x.Badges![BadgeIdx].Url, Url));
+						updates.Add(updateBuilder.Set(x => x.Badges![badgeIdx].Url, url));
 					}
-					if (State != BadgeData.State)
+					if (state != badgeData.State)
 					{
-						Updates.Add(UpdateBuilder.Set(x => x.Badges![BadgeIdx].State, State));
+						updates.Add(updateBuilder.Set(x => x.Badges![badgeIdx].State, state));
 					}
 
-					if (Updates.Count == 0 || await TryUpdateAsync(Document, UpdateBuilder.Combine(Updates)))
+					if (updates.Count == 0 || await TryUpdateAsync(document, updateBuilder.Combine(updates)))
 					{
-						BadgeData.Url = Url;
-						BadgeData.State = State;
-						return Metadata;
+						badgeData.Url = url;
+						badgeData.State = state;
+						return metadata;
 					}
 				}
 
 				// Update the document and try again
-				Metadata = await FindOrAddAsync(Metadata.Stream, Metadata.Change, Metadata.Project);
+				metadata = await FindOrAddAsync(metadata.Stream, metadata.Change, metadata.Project);
 			}
 		}
 
 		/// <inheritdoc/>
-		public async Task<List<IUgsMetadata>> FindAsync(string Stream, int MinChange, int? MaxChange = null, long? MinTime = null)
+		public async Task<List<IUgsMetadata>> FindAsync(string stream, int minChange, int? maxChange = null, long? minTime = null)
 		{
-			FilterDefinitionBuilder<UgsMetadataDocument> FilterBuilder = Builders<UgsMetadataDocument>.Filter;
+			FilterDefinitionBuilder<UgsMetadataDocument> filterBuilder = Builders<UgsMetadataDocument>.Filter;
 
-			string NormalizedStream = GetNormalizedStream(Stream);
+			string normalizedStream = GetNormalizedStream(stream);
 
-			FilterDefinition<UgsMetadataDocument> Filter = FilterBuilder.Eq(x => x.Stream, NormalizedStream) & FilterBuilder.Gte(x => x.Change, MinChange);
-			if (MaxChange != null)
+			FilterDefinition<UgsMetadataDocument> filter = filterBuilder.Eq(x => x.Stream, normalizedStream) & filterBuilder.Gte(x => x.Change, minChange);
+			if (maxChange != null)
 			{
-				Filter &= FilterBuilder.Lte(x => x.Change, MaxChange.Value);
+				filter &= filterBuilder.Lte(x => x.Change, maxChange.Value);
 			}
-			if (MinTime != null)
+			if (minTime != null)
 			{
-				Filter &= FilterBuilder.Gt(x => x.UpdateTicks, MinTime.Value);
+				filter &= filterBuilder.Gt(x => x.UpdateTicks, minTime.Value);
 			}
 
-			List<UgsMetadataDocument> Documents = await Collection.Find(Filter).ToListAsync();
-			return Documents.ConvertAll<IUgsMetadata>(x =>
+			List<UgsMetadataDocument> documents = await _collection.Find(filter).ToListAsync();
+			return documents.ConvertAll<IUgsMetadata>(x =>
 			{
 				// Remove polluting null entries. Need to find the source of these.
 				x.Users = x.Users.Where(y => y != null).ToList();
@@ -311,51 +310,51 @@ namespace Horde.Build.Collections.Impl
 		/// <summary>
 		/// Normalize a stream argument
 		/// </summary>
-		/// <param name="Stream">The stream name</param>
+		/// <param name="stream">The stream name</param>
 		/// <returns>Normalized name</returns>
 		[SuppressMessage("Globalization", "CA1308:Normalize strings to uppercase")]
-		private static string GetNormalizedStream(string Stream)
+		private static string GetNormalizedStream(string stream)
 		{
-			return Stream.ToLowerInvariant().TrimEnd('/');
+			return stream.ToLowerInvariant().TrimEnd('/');
 		}
 
 		/// <summary>
 		/// Normalize a project argument
 		/// </summary>
-		/// <param name="Project">The project name</param>
+		/// <param name="project">The project name</param>
 		/// <returns>Normalized name</returns>
 		[SuppressMessage("Globalization", "CA1308:Normalize strings to uppercase")]
-		private static string GetNormalizedProject(string? Project)
+		private static string GetNormalizedProject(string? project)
 		{
-			if(Project == null)
+			if(project == null)
 			{
 				return String.Empty;
 			}
 			else
 			{
-				return Project.ToLowerInvariant();
+				return project.ToLowerInvariant();
 			}
 		}
 
 		/// <summary>
 		/// Try to update a metadata document
 		/// </summary>
-		/// <param name="Document">The document to update</param>
-		/// <param name="Update">Update definition</param>
+		/// <param name="document">The document to update</param>
+		/// <param name="update">Update definition</param>
 		/// <returns>True if the update succeeded</returns>
-		private async Task<bool> TryUpdateAsync(UgsMetadataDocument Document, UpdateDefinition<UgsMetadataDocument> Update)
+		private async Task<bool> TryUpdateAsync(UgsMetadataDocument document, UpdateDefinition<UgsMetadataDocument> update)
 		{
-			int NewUpdateIndex = Document.UpdateIndex + 1;
-			Update = Update.Set(x => x.UpdateIndex, NewUpdateIndex);
+			int newUpdateIndex = document.UpdateIndex + 1;
+			update = update.Set(x => x.UpdateIndex, newUpdateIndex);
 
-			long NewUpdateTicks = DateTime.UtcNow.Ticks;
-			Update = Update.Set(x => x.UpdateTicks, NewUpdateTicks);
+			long newUpdateTicks = DateTime.UtcNow.Ticks;
+			update = update.Set(x => x.UpdateTicks, newUpdateTicks);
 
-			UpdateResult Result = await Collection.UpdateOneAsync(x => x.Change == Document.Change && x.UpdateIndex == Document.UpdateIndex, Update);
-			if (Result.ModifiedCount > 0)
+			UpdateResult result = await _collection.UpdateOneAsync(x => x.Change == document.Change && x.UpdateIndex == document.UpdateIndex, update);
+			if (result.ModifiedCount > 0)
 			{
-				Document.UpdateIndex = NewUpdateIndex;
-				Document.UpdateTicks = NewUpdateTicks;
+				document.UpdateIndex = newUpdateIndex;
+				document.UpdateTicks = newUpdateTicks;
 				return true;
 			}
 			return false;
diff --git a/Engine/Source/Programs/Horde/Horde.Build/Ugs/UgsMetadata.cs b/Engine/Source/Programs/Horde/Horde.Build/Ugs/UgsMetadata.cs
--- a/Engine/Source/Programs/Horde/Horde.Build/Ugs/UgsMetadata.cs
+++ b/Engine/Source/Programs/Horde/Horde.Build/Ugs/UgsMetadata.cs
@@ -1,12 +1,8 @@
 // Copyright Epic Games, Inc. All Rights Reserved.
 
-using Amazon.S3.Model;
-using Horde.Build.Api;
-using MongoDB.Bson;
 using System;
 using System.Collections.Generic;
-using System.Linq;
-using System.Threading.Tasks;
+using Horde.Build.Api;
 
 namespace Horde.Build.Models
 {
diff --git a/Engine/Source/Programs/Horde/Horde.Build/Users/IAvatar.cs b/Engine/Source/Programs/Horde/Horde.Build/Users/IAvatar.cs
--- a/Engine/Source/Programs/Horde/Horde.Build/Users/IAvatar.cs
+++ b/Engine/Source/Programs/Horde/Horde.Build/Users/IAvatar.cs
@@ -1,10 +1,5 @@
 // Copyright Epic Games, Inc. All Rights Reserved.
 
-using System;
-using System.Collections.Generic;
-using System.Linq;
-using System.Threading.Tasks;
-
 namespace Horde.Build.Models
 {
 	/// <summary>
diff --git a/Engine/Source/Programs/Horde/Horde.Build/Users/IAvatarService.cs b/Engine/Source/Programs/Horde/Horde.Build/Users/IAvatarService.cs
--- a/Engine/Source/Programs/Horde/Horde.Build/Users/IAvatarService.cs
+++ b/Engine/Source/Programs/Horde/Horde.Build/Users/IAvatarService.cs
@@ -1,14 +1,7 @@
 // Copyright Epic Games, Inc. All Rights Reserved.
 
-using Horde.Build.Models;
-using Microsoft.Extensions.Logging;
-using Microsoft.Extensions.Options;
-using Newtonsoft.Json.Linq;
-using System;
-using System.Collections.Generic;
-using System.Linq;
-using System.Net.Http;
 using System.Threading.Tasks;
+using Horde.Build.Models;
 
 namespace Horde.Build.Services
 {
@@ -20,9 +13,9 @@ namespace Horde.Build.Services
 		/// <summary>
 		/// Gets a users's avatar
 		/// </summary>
-		/// <param name="User"></param>
+		/// <param name="user"></param>
 		/// <returns></returns>
-		Task<IAvatar?> GetAvatarAsync(IUser User);
+		Task<IAvatar?> GetAvatarAsync(IUser user);
 	}
 }
 
diff --git a/Engine/Source/Programs/Horde/Horde.Build/Users/IServiceAccount.cs b/Engine/Source/Programs/Horde/Horde.Build/Users/IServiceAccount.cs
--- a/Engine/Source/Programs/Horde/Horde.Build/Users/IServiceAccount.cs
+++ b/Engine/Source/Programs/Horde/Horde.Build/Users/IServiceAccount.cs
@@ -1,7 +1,7 @@
 // Copyright Epic Games, Inc. All Rights Reserved.
 
-using MongoDB.Bson;
 using System.Collections.Generic;
+using MongoDB.Bson;
 
 namespace Horde.Build.Models
 {
@@ -34,10 +34,10 @@ namespace Horde.Build.Models
 		/// <summary>
 		/// Add a claim
 		/// </summary>
-		/// <param name="Type">Type of claim</param>
-		/// <param name="Value">Value of claim</param>
+		/// <param name="type">Type of claim</param>
+		/// <param name="value">Value of claim</param>
 		/// <returns></returns>
-		public void AddClaim(string Type, string Value);
+		public void AddClaim(string type, string value);
 		
 		/// <summary>
 		/// Get list of claims
diff --git a/Engine/Source/Programs/Horde/Horde.Build/Users/IUser.cs b/Engine/Source/Programs/Horde/Horde.Build/Users/IUser.cs
--- a/Engine/Source/Programs/Horde/Horde.Build/Users/IUser.cs
+++ b/Engine/Source/Programs/Horde/Horde.Build/Users/IUser.cs
@@ -1,12 +1,8 @@
 // Copyright Epic Games, Inc. All Rights Reserved.
 
+using System.Collections.Generic;
 using Horde.Build.Utilities;
 using MongoDB.Bson;
-using System;
-using System.Collections.Generic;
-using System.Linq;
-using System.Security.Claims;
-using System.Threading.Tasks;
 
 namespace Horde.Build.Models
 {
diff --git a/Engine/Source/Programs/Horde/Horde.Build/Users/IUserClaim.cs b/Engine/Source/Programs/Horde/Horde.Build/Users/IUserClaim.cs
--- a/Engine/Source/Programs/Horde/Horde.Build/Users/IUserClaim.cs
+++ b/Engine/Source/Programs/Horde/Horde.Build/Users/IUserClaim.cs
@@ -1,10 +1,6 @@
 // Copyright Epic Games, Inc. All Rights Reserved.
 
-using System;
-using System.Collections.Generic;
-using System.Linq;
 using System.Security.Claims;
-using System.Threading.Tasks;
 
 namespace Horde.Build.Models
 {
@@ -38,37 +34,37 @@ namespace Horde.Build.Models
 		/// <summary>
 		/// Constructor
 		/// </summary>
-		/// <param name="Type"></param>
-		/// <param name="Value"></param>
-		public UserClaim(string Type, string Value)
+		/// <param name="type"></param>
+		/// <param name="value"></param>
+		public UserClaim(string type, string value)
 		{
-			this.Type = Type;
-			this.Value = Value;
+			Type = type;
+			Value = value;
 		}
 
 		/// <summary>
 		/// Constructor
 		/// </summary>
-		/// <param name="Other">Claim to construct from</param>
-		public UserClaim(IUserClaim Other)
-			: this(Other.Type, Other.Value)
+		/// <param name="other">Claim to construct from</param>
+		public UserClaim(IUserClaim other)
+			: this(other.Type, other.Value)
 		{
 		}
 
 		/// <summary>
 		/// Constructs a UserClaim from a Claim object
 		/// </summary>
-		/// <param name="Claim">Claim object</param>
+		/// <param name="claim">Claim object</param>
 		/// <returns></returns>
-		public static UserClaim FromClaim(Claim Claim)
+		public static UserClaim FromClaim(Claim claim)
 		{
-			return new UserClaim(Claim.Type, Claim.Value);
+			return new UserClaim(claim.Type, claim.Value);
 		}
 
 		/// <summary>
 		/// Conversion operator from NET claims
 		/// </summary>
-		/// <param name="Claim"></param>
-		public static implicit operator UserClaim(Claim Claim) => FromClaim(Claim);
+		/// <param name="claim"></param>
+		public static implicit operator UserClaim(Claim claim) => FromClaim(claim);
 	}
 }
diff --git a/Engine/Source/Programs/Horde/Horde.Build/Users/IUserCollection.cs b/Engine/Source/Programs/Horde/Horde.Build/Users/IUserCollection.cs
--- a/Engine/Source/Programs/Horde/Horde.Build/Users/IUserCollection.cs
+++ b/Engine/Source/Programs/Horde/Horde.Build/Users/IUserCollection.cs
@@ -1,14 +1,11 @@
 // Copyright Epic Games, Inc. All Rights Reserved.
 
-using Horde.Build.Models;
-using Horde.Build.Services;
-using Horde.Build.Utilities;
-using MongoDB.Bson;
-using System;
 using System.Collections.Generic;
-using System.Linq;
 using System.Security.Claims;
 using System.Threading.Tasks;
+using Horde.Build.Models;
+using Horde.Build.Utilities;
+using MongoDB.Bson;
 
 namespace Horde.Build.Collections
 {
@@ -23,75 +20,75 @@ namespace Horde.Build.Collections
 		/// <summary>
 		/// Gets a user by unique id
 		/// </summary>
-		/// <param name="Id">Id of the user</param>
+		/// <param name="id">Id of the user</param>
 		/// <returns>The user information</returns>
-		Task<IUser?> GetUserAsync(UserId Id);
+		Task<IUser?> GetUserAsync(UserId id);
 
 		/// <summary>
 		/// Gets a cached user by unique id
 		/// </summary>
-		/// <param name="Id">Id of the user</param>
+		/// <param name="id">Id of the user</param>
 		/// <returns>The user information</returns>
-		ValueTask<IUser?> GetCachedUserAsync(UserId? Id);
+		ValueTask<IUser?> GetCachedUserAsync(UserId? id);
 
 		/// <summary>
 		/// Gets a user by unique id
 		/// </summary>
-		/// <param name="Ids">Ids of the users</param>
-		/// <param name="NameRegex">Name regex to match for the users</param>
-		/// <param name="Index">Maximum number of results</param>
-		/// <param name="Count">Number of results to return</param>
+		/// <param name="ids">Ids of the users</param>
+		/// <param name="nameRegex">Name regex to match for the users</param>
+		/// <param name="index">Maximum number of results</param>
+		/// <param name="count">Number of results to return</param>
 		/// <returns>The user information</returns>
-		Task<List<IUser>> FindUsersAsync(IEnumerable<UserId>? Ids = null, string? NameRegex = null, int? Index = null, int? Count = null);
+		Task<List<IUser>> FindUsersAsync(IEnumerable<UserId>? ids = null, string? nameRegex = null, int? index = null, int? count = null);
 
 		/// <summary>
 		/// Gets a user by login
 		/// </summary>
-		/// <param name="Login">Login for the user</param>
+		/// <param name="login">Login for the user</param>
 		/// <returns>The user information</returns>
-		Task<IUser?> FindUserByLoginAsync(string Login);
+		Task<IUser?> FindUserByLoginAsync(string login);
 
 		/// <summary>
 		/// Find or add a user with the given claims. Claims will be updated if the user exists.
 		/// </summary>
-		/// <param name="Login">Login id of the user</param>
-		/// <param name="Name">Full name of the user</param>
-		/// <param name="Email">Email address of the user</param>
+		/// <param name="login">Login id of the user</param>
+		/// <param name="name">Full name of the user</param>
+		/// <param name="email">Email address of the user</param>
 		/// <returns>The user document</returns>
-		Task<IUser> FindOrAddUserByLoginAsync(string Login, string? Name = null, string? Email = null);
+		Task<IUser> FindOrAddUserByLoginAsync(string login, string? name = null, string? email = null);
 
 		/// <summary>
 		/// Gets the claims for a user
 		/// </summary>
-		/// <param name="UserId"></param>
+		/// <param name="userId"></param>
 		/// <returns></returns>
-		Task<IUserClaims> GetClaimsAsync(UserId UserId);
+		Task<IUserClaims> GetClaimsAsync(UserId userId);
 
 		/// <summary>
 		/// Update the claims for a user
 		/// </summary>
-		/// <param name="UserId"></param>
-		/// <param name="Claims"></param>
+		/// <param name="userId"></param>
+		/// <param name="claims"></param>
 		/// <returns></returns>
-		Task UpdateClaimsAsync(UserId UserId, IEnumerable<IUserClaim> Claims);
+		Task UpdateClaimsAsync(UserId userId, IEnumerable<IUserClaim> claims);
 
 		/// <summary>
 		/// Get settings for a user
 		/// </summary>
-		/// <param name="UserId"></param>
+		/// <param name="userId"></param>
 		/// <returns></returns>
-		Task<IUserSettings> GetSettingsAsync(UserId UserId);
+		Task<IUserSettings> GetSettingsAsync(UserId userId);
 
 		/// <summary>
 		/// Update a user
 		/// </summary>
-		/// <param name="UserId">The user to update</param>
-		/// <param name="EnableExperimentalFeatures"></param>
-		/// <param name="DashboardSettings">Opaque settings object for the dashboard</param>
-		/// <param name="AddPinnedJobIds"></param>
-		/// <param name="RemovePinnedJobIds"></param>
+		/// <param name="userId">The user to update</param>
+		/// <param name="enableExperimentalFeatures"></param>
+		/// <param name="dashboardSettings">Opaque settings object for the dashboard</param>
+		/// <param name="addPinnedJobIds"></param>
+		/// <param name="removePinnedJobIds"></param>
 		/// <returns>Updated user object</returns>
-		Task UpdateSettingsAsync(UserId UserId, bool? EnableExperimentalFeatures = null, BsonValue? DashboardSettings = null, IEnumerable<JobId>? AddPinnedJobIds = null, IEnumerable<JobId>? RemovePinnedJobIds = null);
+		Task UpdateSettingsAsync(UserId userId, bool? enableExperimentalFeatures = null, BsonValue? dashboardSettings = null, IEnumerable<JobId>? addPinnedJobIds = null, IEnumerable<JobId>? removePinnedJobIds = null);
 	}
 
 	/// <summary>
@@ -102,15 +99,15 @@ namespace Horde.Build.Collections
 		/// <summary>
 		/// Gets a particular user info from the collection
 		/// </summary>
-		/// <param name="UserCollection"></param>
-		/// <param name="Principal"></param>
+		/// <param name="userCollection"></param>
+		/// <param name="principal"></param>
 		/// <returns></returns>
-		public static Task<IUser?> GetUserAsync(this IUserCollection UserCollection, ClaimsPrincipal Principal)
+		public static Task<IUser?> GetUserAsync(this IUserCollection userCollection, ClaimsPrincipal principal)
 		{
-			UserId? UserId = Principal.GetUserId();
-			if (UserId != null)
+			UserId? userId = principal.GetUserId();
+			if (userId != null)
 			{
-				return UserCollection.GetUserAsync(UserId.Value);
+				return userCollection.GetUserAsync(userId.Value);
 			}
 			else
 			{
diff --git a/Engine/Source/Programs/Horde/Horde.Build/Users/Impl/IServiceAccountCollection.cs b/Engine/Source/Programs/Horde/Horde.Build/Users/Impl/IServiceAccountCollection.cs
--- a/Engine/Source/Programs/Horde/Horde.Build/Users/Impl/IServiceAccountCollection.cs
+++ b/Engine/Source/Programs/Horde/Horde.Build/Users/Impl/IServiceAccountCollection.cs
@@ -1,9 +1,9 @@
 // Copyright Epic Games, Inc. All Rights Reserved.
 
-using Horde.Build.Models;
-using MongoDB.Bson;
 using System.Collections.Generic;
 using System.Threading.Tasks;
+using Horde.Build.Models;
+using MongoDB.Bson;
 
 namespace Horde.Build.Collections
 {
@@ -15,41 +15,41 @@ namespace Horde.Build.Collections
 		/// <summary>
 		/// Adds a new service account to the collection
 		/// </summary>
-		/// <param name="SecretToken">Secret token used for identifying this service account</param>
-		/// <param name="Claims">Dictionary of claims to assign</param>
-		/// <param name="Description">Description of the account</param>
-		Task<IServiceAccount> AddAsync(string SecretToken, List<string> Claims, string Description);
+		/// <param name="secretToken">Secret token used for identifying this service account</param>
+		/// <param name="claims">Dictionary of claims to assign</param>
+		/// <param name="description">Description of the account</param>
+		Task<IServiceAccount> AddAsync(string secretToken, List<string> claims, string description);
 
 		/// <summary>
 		/// Get service account via ID
 		/// </summary>
-		/// <param name="Id">The unique service account id</param>
+		/// <param name="id">The unique service account id</param>
 		/// <returns>The service account</returns>
-		Task<IServiceAccount?> GetAsync(ObjectId Id);
+		Task<IServiceAccount?> GetAsync(ObjectId id);
 		
 		/// <summary>
 		/// Get service account via secret token
 		/// </summary>
-		/// <param name="SecretToken">Secret token to use for searching</param>
+		/// <param name="secretToken">Secret token to use for searching</param>
 		/// <returns>The service account</returns>
-		Task<IServiceAccount?> GetBySecretTokenAsync(string SecretToken);
+		Task<IServiceAccount?> GetBySecretTokenAsync(string secretToken);
 
 		/// <summary>
 		/// Update a service account from the collection
 		/// </summary>
-		/// <param name="Id">The service account</param>
-		/// <param name="SecretToken">If set, secret token will be set</param>
-		/// <param name="Claims">If set, claims will be set</param>
-		/// <param name="Enabled">If set, enabled flag will be set</param>
-		/// <param name="Description">If set, description will be set</param>
+		/// <param name="id">The service account</param>
+		/// <param name="secretToken">If set, secret token will be set</param>
+		/// <param name="claims">If set, claims will be set</param>
+		/// <param name="enabled">If set, enabled flag will be set</param>
+		/// <param name="description">If set, description will be set</param>
 		/// <returns>Async task</returns>
-		Task UpdateAsync(ObjectId Id, string? SecretToken, List<string>? Claims, bool? Enabled, string? Description);
+		Task UpdateAsync(ObjectId id, string? secretToken, List<string>? claims, bool? enabled, string? description);
 
 		/// <summary>
 		/// Delete a service account from the collection
 		/// </summary>
-		/// <param name="Id">The service account</param>
+		/// <param name="id">The service account</param>
 		/// <returns>Async task</returns>
-		Task DeleteAsync(ObjectId Id);
+		Task DeleteAsync(ObjectId id);
 	}
 }
diff --git a/Engine/Source/Programs/Horde/Horde.Build/Users/Impl/ServiceAccountCollection.cs b/Engine/Source/Programs/Horde/Horde.Build/Users/Impl/ServiceAccountCollection.cs
--- a/Engine/Source/Programs/Horde/Horde.Build/Users/Impl/ServiceAccountCollection.cs
+++ b/Engine/Source/Programs/Horde/Horde.Build/Users/Impl/ServiceAccountCollection.cs
@@ -1,14 +1,14 @@
 // Copyright Epic Games, Inc. All Rights Reserved.
 
 using System;
+using System.Collections.Generic;
+using System.Linq;
+using System.Threading.Tasks;
 using Horde.Build.Models;
 using Horde.Build.Services;
 using MongoDB.Bson;
 using MongoDB.Bson.Serialization.Attributes;
 using MongoDB.Driver;
-using System.Collections.Generic;
-using System.Linq;
-using System.Threading.Tasks;
 
 namespace Horde.Build.Collections.Impl
 {
@@ -41,19 +41,19 @@ namespace Horde.Build.Collections.Impl
 			{
 			}
 
-			public ServiceAccountDocument(ObjectId Id, string SecretToken, List<string> Claims, bool Enabled, string Description)
+			public ServiceAccountDocument(ObjectId id, string secretToken, List<string> claims, bool enabled, string description)
 			{
-				this.Id = Id;
-				this.SecretToken = SecretToken;
-				this.Claims = Claims;
-				this.Enabled = Enabled;
-				this.Description = Description;
+				Id = id;
+				SecretToken = secretToken;
+				Claims = claims;
+				Enabled = enabled;
+				Description = description;
 			}
 			
 			/// <inheritdoc/>
-			public void AddClaim(string Type, string Value)
+			public void AddClaim(string type, string value)
 			{
-				Claims.Add(Type + ClaimSeparator + Value);
+				Claims.Add(type + ClaimSeparator + value);
 			}
 
 			/// <inheritdoc/>
@@ -61,23 +61,32 @@ namespace Horde.Build.Collections.Impl
 			{
 				return Claims.Select(x =>
 				{
-					string[] Split = x.Split(ClaimSeparator);
-					return (Split[0], Split[1]);
+					string[] split = x.Split(ClaimSeparator);
+					return (split[0], split[1]);
 				}).ToList();
 			}
 
 			protected bool Equals(ServiceAccountDocument other)
 			{
-				bool AreClaimsEqual = !Claims.Except(other.Claims).Any();
+				bool areClaimsEqual = !Claims.Except(other.Claims).Any();
 
-                return Id.Equals(other.Id) && SecretToken == other.SecretToken && AreClaimsEqual && Enabled == other.Enabled && Description == other.Description;
+                return Id.Equals(other.Id) && SecretToken == other.SecretToken && areClaimsEqual && Enabled == other.Enabled && Description == other.Description;
 			}
 
 			public override bool Equals(object? obj)
 			{
-				if (ReferenceEquals(null, obj)) return false;
-				if (ReferenceEquals(this, obj)) return true;
-				if (obj.GetType() != this.GetType()) return false;
+				if (obj is null)
+				{
+					return false;
+				}
+				if (ReferenceEquals(this, obj))
+				{
+					return true;
+				}
+				if (obj.GetType() != GetType())
+				{
+					return false;
+				}
 				return Equals((ServiceAccountDocument) obj);
 			}
 
@@ -90,60 +99,72 @@ namespace Horde.Build.Collections.Impl
 		/// <summary>
 		/// Collection of session documents
 		/// </summary>
-		private readonly IMongoCollection<ServiceAccountDocument> ServiceAccounts;
+		private readonly IMongoCollection<ServiceAccountDocument> _serviceAccounts;
 
 		/// <summary>
 		/// Constructor
 		/// </summary>
-		/// <param name="DatabaseService">The database service</param>
-		public ServiceAccountCollection(DatabaseService DatabaseService)
+		/// <param name="databaseService">The database service</param>
+		public ServiceAccountCollection(DatabaseService databaseService)
 		{
-			ServiceAccounts = DatabaseService.GetCollection<ServiceAccountDocument>("ServiceAccounts");
+			_serviceAccounts = databaseService.GetCollection<ServiceAccountDocument>("ServiceAccounts");
 
-			if (!DatabaseService.ReadOnlyMode)
+			if (!databaseService.ReadOnlyMode)
 			{
-				ServiceAccounts.Indexes.CreateOne(new CreateIndexModel<ServiceAccountDocument>(Builders<ServiceAccountDocument>.IndexKeys.Ascending(x => x.SecretToken)));
+				_serviceAccounts.Indexes.CreateOne(new CreateIndexModel<ServiceAccountDocument>(Builders<ServiceAccountDocument>.IndexKeys.Ascending(x => x.SecretToken)));
 			}
 		}
 
 		/// <inheritdoc/>
-		public async Task<IServiceAccount> AddAsync(string SecretToken, List<string> Claims, string Description)
+		public async Task<IServiceAccount> AddAsync(string secretToken, List<string> claims, string description)
 		{
-			ServiceAccountDocument NewSession = new ServiceAccountDocument(ObjectId.GenerateNewId(), SecretToken, Claims, true, Description);
-			await ServiceAccounts.InsertOneAsync(NewSession);
-			return NewSession;
+			ServiceAccountDocument newSession = new ServiceAccountDocument(ObjectId.GenerateNewId(), secretToken, claims, true, description);
+			await _serviceAccounts.InsertOneAsync(newSession);
+			return newSession;
 		}
 
 		/// <inheritdoc/>
-		public async Task<IServiceAccount?> GetAsync(ObjectId Id)
+		public async Task<IServiceAccount?> GetAsync(ObjectId id)
 		{
-			return await ServiceAccounts.Find(x => x.Id == Id).FirstOrDefaultAsync();
+			return await _serviceAccounts.Find(x => x.Id == id).FirstOrDefaultAsync();
 		}
 
 		/// <inheritdoc/>
-		public async Task<IServiceAccount?> GetBySecretTokenAsync(string SecretToken)
+		public async Task<IServiceAccount?> GetBySecretTokenAsync(string secretToken)
 		{
-			return await ServiceAccounts.Find(x => x.SecretToken == SecretToken).FirstOrDefaultAsync();
+			return await _serviceAccounts.Find(x => x.SecretToken == secretToken).FirstOrDefaultAsync();
 		}
 
 		/// <inheritdoc/>
-		public Task UpdateAsync(ObjectId Id, string? SecretToken, List<string>? Claims, bool? Enabled, string? Description)
+		public Task UpdateAsync(ObjectId id, string? secretToken, List<string>? claims, bool? enabled, string? description)
 		{
-			UpdateDefinitionBuilder<ServiceAccountDocument> Update = Builders<ServiceAccountDocument>.Update;
-			List<UpdateDefinition<ServiceAccountDocument>> Updates = new List<UpdateDefinition<ServiceAccountDocument>>();
-			
-			if (SecretToken != null) Updates.Add(Update.Set(x => x.SecretToken, SecretToken));
-			if (Claims != null) Updates.Add(Update.Set(x => x.Claims, Claims));
-			if (Enabled != null) Updates.Add(Update.Set(x => x.Enabled, Enabled));
-			if (Description != null) Updates.Add(Update.Set(x => x.Description, Description));
+			UpdateDefinitionBuilder<ServiceAccountDocument> update = Builders<ServiceAccountDocument>.Update;
+			List<UpdateDefinition<ServiceAccountDocument>> updates = new List<UpdateDefinition<ServiceAccountDocument>>();
+
+			if (secretToken != null)
+			{
+				updates.Add(update.Set(x => x.SecretToken, secretToken));
+			}
+			if (claims != null)
+			{
+				updates.Add(update.Set(x => x.Claims, claims));
+			}
+			if (enabled != null)
+			{
+				updates.Add(update.Set(x => x.Enabled, enabled));
+			}
+			if (description != null)
+			{
+				updates.Add(update.Set(x => x.Description, description));
+			}
 
-			return ServiceAccounts.FindOneAndUpdateAsync(x => x.Id == Id, Update.Combine(Updates));
+			return _serviceAccounts.FindOneAndUpdateAsync(x => x.Id == id, update.Combine(updates));
 		}
 
 		/// <inheritdoc/>
-		public Task DeleteAsync(ObjectId SessionId)
+		public Task DeleteAsync(ObjectId sessionId)
 		{
-			return ServiceAccounts.DeleteOneAsync(x => x.Id == SessionId);
+			return _serviceAccounts.DeleteOneAsync(x => x.Id == sessionId);
 		}
 	}
 }
diff --git a/Engine/Source/Programs/Horde/Horde.Build/Users/Impl/UserCollection.cs b/Engine/Source/Programs/Horde/Horde.Build/Users/Impl/UserCollection.cs
--- a/Engine/Source/Programs/Horde/Horde.Build/Users/Impl/UserCollection.cs
+++ b/Engine/Source/Programs/Horde/Horde.Build/Users/Impl/UserCollection.cs
@@ -1,5 +1,10 @@
 // Copyright Epic Games, Inc. All Rights Reserved.
 
+using System;
+using System.Collections.Generic;
+using System.Linq;
+using System.Security.Claims;
+using System.Threading.Tasks;
 using Horde.Build.Models;
 using Horde.Build.Services;
 using Horde.Build.Utilities;
@@ -7,11 +12,6 @@ using MongoDB.Bson;
 using MongoDB.Bson.Serialization;
 using MongoDB.Bson.Serialization.Attributes;
 using MongoDB.Driver;
-using System;
-using System.Collections.Generic;
-using System.Linq;
-using System.Security.Claims;
-using System.Threading.Tasks;
 
 namespace Horde.Build.Collections.Impl
 {
@@ -54,24 +54,24 @@ namespace Horde.Build.Collections.Impl
 
 			private ClaimDocument()
 			{
-				this.Type = null!;
-				this.Value = null!;
+				Type = null!;
+				Value = null!;
 			}
 
-			public ClaimDocument(string Type, string Value)
+			public ClaimDocument(string type, string value)
 			{
-				this.Type = Type;
-				this.Value = Value;
+				Type = type;
+				Value = value;
 			}
 
-			public ClaimDocument(IUserClaim Other)
+			public ClaimDocument(IUserClaim other)
 			{
-				this.Type = Other.Type;
-				this.Value = Other.Value;
+				Type = other.Type;
+				Value = other.Value;
 			}
 		}
 
-		IMongoCollection<UserDocument> Users;
+		readonly IMongoCollection<UserDocument> _users;
 
 		/// <summary>
 		/// Static constructor
@@ -85,106 +85,106 @@ namespace Horde.Build.Collections.Impl
 		/// <summary>
 		/// Constructor
 		/// </summary>
-		/// <param name="DatabaseService"></param>
-		public UserCollectionV1(DatabaseService DatabaseService)
+		/// <param name="databaseService"></param>
+		public UserCollectionV1(DatabaseService databaseService)
 		{
-			Users = DatabaseService.GetCollection<UserDocument>("Users");
+			_users = databaseService.GetCollection<UserDocument>("Users");
 
-			if (!DatabaseService.ReadOnlyMode)
+			if (!databaseService.ReadOnlyMode)
 			{
-				Users.Indexes.CreateOne(new CreateIndexModel<UserDocument>(Builders<UserDocument>.IndexKeys.Ascending(x => x.PrimaryClaim), new CreateIndexOptions { Unique = true }));
+				_users.Indexes.CreateOne(new CreateIndexModel<UserDocument>(Builders<UserDocument>.IndexKeys.Ascending(x => x.PrimaryClaim), new CreateIndexOptions { Unique = true }));
 			}
 		}
 
 		/// <inheritdoc/>
-		public async Task<IUser?> GetUserAsync(UserId Id)
+		public async Task<IUser?> GetUserAsync(UserId id)
 		{
-			return await Users.Find(x => x.Id == Id).FirstOrDefaultAsync();
+			return await _users.Find(x => x.Id == id).FirstOrDefaultAsync();
 		}
 
 		/// <inheritdoc/>
-		public async ValueTask<IUser?> GetCachedUserAsync(UserId? Id)
+		public async ValueTask<IUser?> GetCachedUserAsync(UserId? id)
 		{
-			if (Id == null)
+			if (id == null)
 			{
 				return null;
 			}
 			else
 			{
-				return await GetUserAsync(Id.Value);
+				return await GetUserAsync(id.Value);
 			}
 		}
 
 		/// <inheritdoc/>
-		public async Task<List<IUser>> FindUsersAsync(IEnumerable<UserId>? Ids, string? NameRegex, int? Index, int? Count)
+		public async Task<List<IUser>> FindUsersAsync(IEnumerable<UserId>? ids, string? nameRegex, int? index, int? count)
 		{
-			FilterDefinition<UserDocument> Filter = Builders<UserDocument>.Filter.In(x => x.Id, Ids);
-			return await Users.Find(Filter).Range(Index, Count).ToListAsync<UserDocument, IUser>();
+			FilterDefinition<UserDocument> filter = Builders<UserDocument>.Filter.In(x => x.Id, ids);
+			return await _users.Find(filter).Range(index, count).ToListAsync<UserDocument, IUser>();
 		}
 
 		/// <inheritdoc/>
-		public async Task<IUser?> FindUserByLoginAsync(string Login)
+		public async Task<IUser?> FindUserByLoginAsync(string login)
 		{
-			ClaimDocument PrimaryClaim = new ClaimDocument(ClaimTypes.Name, Login);
-			return await Users.Find(x => x.PrimaryClaim == PrimaryClaim).FirstOrDefaultAsync();
+			ClaimDocument primaryClaim = new ClaimDocument(ClaimTypes.Name, login);
+			return await _users.Find(x => x.PrimaryClaim == primaryClaim).FirstOrDefaultAsync();
 		}
 
 		/// <inheritdoc/>
-		public async Task<IUser> FindOrAddUserByLoginAsync(string Login, string? Name, string? Email)
+		public async Task<IUser> FindOrAddUserByLoginAsync(string login, string? name, string? email)
 		{
-			ClaimDocument NewPrimaryClaim = new ClaimDocument(ClaimTypes.Name, Login);
-			UpdateDefinition<UserDocument> Update = Builders<UserDocument>.Update.SetOnInsert(x => x.Id, UserId.GenerateNewId());
-			return await Users.FindOneAndUpdateAsync<UserDocument>(x => x.PrimaryClaim == NewPrimaryClaim, Update, new FindOneAndUpdateOptions<UserDocument> { IsUpsert = true, ReturnDocument = ReturnDocument.After });
+			ClaimDocument newPrimaryClaim = new ClaimDocument(ClaimTypes.Name, login);
+			UpdateDefinition<UserDocument> update = Builders<UserDocument>.Update.SetOnInsert(x => x.Id, UserId.GenerateNewId());
+			return await _users.FindOneAndUpdateAsync<UserDocument>(x => x.PrimaryClaim == newPrimaryClaim, update, new FindOneAndUpdateOptions<UserDocument> { IsUpsert = true, ReturnDocument = ReturnDocument.After });
 		}
 
 		/// <inheritdoc/>
-		public async Task<IUserClaims> GetClaimsAsync(UserId UserId)
+		public async Task<IUserClaims> GetClaimsAsync(UserId userId)
 		{
-			return await Users.Find(x => x.Id == UserId).FirstOrDefaultAsync() ?? new UserDocument { Id = UserId };
+			return await _users.Find(x => x.Id == userId).FirstOrDefaultAsync() ?? new UserDocument { Id = userId };
 		}
 
 		/// <inheritdoc/>
-		public async Task UpdateClaimsAsync(UserId UserId, IEnumerable<IUserClaim> Claims)
+		public async Task UpdateClaimsAsync(UserId userId, IEnumerable<IUserClaim> claims)
 		{
-			List<ClaimDocument> NewClaims = Claims.Select(x => new ClaimDocument(x)).ToList();
-			await Users.FindOneAndUpdateAsync(x => x.Id == UserId, Builders<UserDocument>.Update.Set(x => x.Claims, NewClaims));
+			List<ClaimDocument> newClaims = claims.Select(x => new ClaimDocument(x)).ToList();
+			await _users.FindOneAndUpdateAsync(x => x.Id == userId, Builders<UserDocument>.Update.Set(x => x.Claims, newClaims));
 		}
 
 		/// <inheritdoc/>
-		public async Task<IUserSettings> GetSettingsAsync(UserId UserId)
+		public async Task<IUserSettings> GetSettingsAsync(UserId userId)
 		{
-			return await Users.Find(x => x.Id == UserId).FirstOrDefaultAsync() ?? new UserDocument { Id = UserId };
+			return await _users.Find(x => x.Id == userId).FirstOrDefaultAsync() ?? new UserDocument { Id = userId };
 		}
 
 		/// <inheritdoc/>
-		public async Task UpdateSettingsAsync(UserId UserId, bool? EnableExperimentalFeatures, BsonValue? DashboardSettings = null, IEnumerable<JobId>? AddPinnedJobIds = null, IEnumerable<JobId>? RemovePinnedJobIds = null)
+		public async Task UpdateSettingsAsync(UserId userId, bool? enableExperimentalFeatures, BsonValue? dashboardSettings = null, IEnumerable<JobId>? addPinnedJobIds = null, IEnumerable<JobId>? removePinnedJobIds = null)
 		{
-			if (AddPinnedJobIds != null)
+			if (addPinnedJobIds != null)
 			{
-				foreach (JobId PinnedJobId in AddPinnedJobIds)
+				foreach (JobId pinnedJobId in addPinnedJobIds)
 				{
-					FilterDefinition<UserDocument> Filter = Builders<UserDocument>.Filter.Eq(x => x.Id, UserId) & Builders<UserDocument>.Filter.AnyNin<JobId>(x => x.PinnedJobIds, new[] { PinnedJobId });
-					UpdateDefinition<UserDocument> Update = Builders<UserDocument>.Update.PushEach(x => x.PinnedJobIds, new[] { PinnedJobId }, -50);
-					await Users.UpdateOneAsync(Filter, Update);
+					FilterDefinition<UserDocument> filter = Builders<UserDocument>.Filter.Eq(x => x.Id, userId) & Builders<UserDocument>.Filter.AnyNin<JobId>(x => x.PinnedJobIds, new[] { pinnedJobId });
+					UpdateDefinition<UserDocument> update = Builders<UserDocument>.Update.PushEach(x => x.PinnedJobIds, new[] { pinnedJobId }, -50);
+					await _users.UpdateOneAsync(filter, update);
 				}
 			}
 
-			List<UpdateDefinition<UserDocument>> Updates = new List<UpdateDefinition<UserDocument>>();
-			if (EnableExperimentalFeatures != null)
+			List<UpdateDefinition<UserDocument>> updates = new List<UpdateDefinition<UserDocument>>();
+			if (enableExperimentalFeatures != null)
 			{
-				Updates.Add(Builders<UserDocument>.Update.SetOrUnsetNull(x => x.EnableExperimentalFeatures, EnableExperimentalFeatures));
+				updates.Add(Builders<UserDocument>.Update.SetOrUnsetNull(x => x.EnableExperimentalFeatures, enableExperimentalFeatures));
 			}
-			if (DashboardSettings != null)
+			if (dashboardSettings != null)
 			{
-				Updates.Add(Builders<UserDocument>.Update.Set(x => x.DashboardSettings, DashboardSettings));
+				updates.Add(Builders<UserDocument>.Update.Set(x => x.DashboardSettings, dashboardSettings));
 			}
-			if (RemovePinnedJobIds != null && RemovePinnedJobIds.Any())
+			if (removePinnedJobIds != null && removePinnedJobIds.Any())
 			{
-				Updates.Add(Builders<UserDocument>.Update.PullAll(x => x.PinnedJobIds, RemovePinnedJobIds));
+				updates.Add(Builders<UserDocument>.Update.PullAll(x => x.PinnedJobIds, removePinnedJobIds));
 			}
-			if (Updates.Count > 0)
+			if (updates.Count > 0)
 			{
-				await Users.UpdateOneAsync<UserDocument>(x => x.Id == UserId, Builders<UserDocument>.Update.Combine(Updates));
+				await _users.UpdateOneAsync<UserDocument>(x => x.Id == userId, Builders<UserDocument>.Update.Combine(updates));
 			}
 		}
 
@@ -194,15 +194,15 @@ namespace Horde.Build.Collections.Impl
 		/// <returns></returns>
 		public async IAsyncEnumerable<(IUser, IUserClaims, IUserSettings)> EnumerateDocumentsAsync()
 		{
-			using (IAsyncCursor<UserDocument> Cursor = await Users.Find(FilterDefinition<UserDocument>.Empty).ToCursorAsync())
+			using (IAsyncCursor<UserDocument> cursor = await _users.Find(FilterDefinition<UserDocument>.Empty).ToCursorAsync())
 			{
-				while (await Cursor.MoveNextAsync())
+				while (await cursor.MoveNextAsync())
 				{
-					foreach (UserDocument Document in Cursor.Current)
+					foreach (UserDocument document in cursor.Current)
 					{
-						if (Document.Claims.Count > 0)
+						if (document.Claims.Count > 0)
 						{
-							yield return (Document, Document, Document);
+							yield return (document, document, document);
 						}
 					}
 				}
diff --git a/Engine/Source/Programs/Horde/Horde.Build/Users/Impl/UserCollectionV2.cs b/Engine/Source/Programs/Horde/Horde.Build/Users/Impl/UserCollectionV2.cs
--- a/Engine/Source/Programs/Horde/Horde.Build/Users/Impl/UserCollectionV2.cs
+++ b/Engine/Source/Programs/Horde/Horde.Build/Users/Impl/UserCollectionV2.cs
@@ -1,5 +1,9 @@
 // Copyright Epic Games, Inc. All Rights Reserved.
 
+using System;
+using System.Collections.Generic;
+using System.Linq;
+using System.Threading.Tasks;
 using Horde.Build.Models;
 using Horde.Build.Services;
 using Horde.Build.Utilities;
@@ -9,11 +13,6 @@ using MongoDB.Bson;
 using MongoDB.Bson.Serialization;
 using MongoDB.Bson.Serialization.Attributes;
 using MongoDB.Driver;
-using System;
-using System.Collections.Generic;
-using System.Linq;
-using System.Security.Claims;
-using System.Threading.Tasks;
 
 namespace Horde.Build.Collections.Impl
 {
@@ -46,19 +45,19 @@ namespace Horde.Build.Collections.Impl
 				LoginUpper = null!;
 			}
 
-			public UserDocument(IUser Other)
-				: this(Other.Id, Other.Name, Other.Login, Other.Email)
+			public UserDocument(IUser other)
+				: this(other.Id, other.Name, other.Login, other.Email)
 			{
 			}
 
-			public UserDocument(UserId Id, string Name, string Login, string? Email)
+			public UserDocument(UserId id, string name, string login, string? email)
 			{
-				this.Id = Id;
-				this.Name = Name;
-				this.Login = Login;
-				this.LoginUpper = Login.ToUpperInvariant();
-				this.Email = Email;
-				this.EmailUpper = Email?.ToUpperInvariant();
+				Id = id;
+				Name = name;
+				Login = login;
+				LoginUpper = login.ToUpperInvariant();
+				Email = email;
+				EmailUpper = email?.ToUpperInvariant();
 			}
 		}
 
@@ -75,15 +74,15 @@ namespace Horde.Build.Collections.Impl
 			{
 			}
 
-			public UserClaimsDocument(UserId Id)
+			public UserClaimsDocument(UserId id)
 			{
-				this.Id = Id;
+				Id = id;
 			}
 
-			public UserClaimsDocument(IUserClaims Other)
-				: this(Other.UserId)
+			public UserClaimsDocument(IUserClaims other)
+				: this(other.UserId)
 			{
-				this.Claims.AddRange(Other.Claims.Select(x => new UserClaim(x)));
+				Claims.AddRange(other.Claims.Select(x => new UserClaim(x)));
 			}
 		}
 
@@ -105,17 +104,17 @@ namespace Horde.Build.Collections.Impl
 			{
 			}
 
-			public UserSettingsDocument(UserId Id)
+			public UserSettingsDocument(UserId id)
 			{
-				this.Id = Id;
+				Id = id;
 			}
 
-			public UserSettingsDocument(IUserSettings Other)
-				: this(Other.UserId)
+			public UserSettingsDocument(IUserSettings other)
+				: this(other.UserId)
 			{
-				this.EnableExperimentalFeatures = Other.EnableExperimentalFeatures;
-				this.DashboardSettings = Other.DashboardSettings;
-				this.PinnedJobIds = new List<JobId>(Other.PinnedJobIds);
+				EnableExperimentalFeatures = other.EnableExperimentalFeatures;
+				DashboardSettings = other.DashboardSettings;
+				PinnedJobIds = new List<JobId>(other.PinnedJobIds);
 			}
 		}
 
@@ -126,30 +125,28 @@ namespace Horde.Build.Collections.Impl
 
 			private ClaimDocument()
 			{
-				this.Type = null!;
-				this.Value = null!;
+				Type = null!;
+				Value = null!;
 			}
 
-			public ClaimDocument(string Type, string Value)
+			public ClaimDocument(string type, string value)
 			{
-				this.Type = Type;
-				this.Value = Value;
+				Type = type;
+				Value = value;
 			}
 
-			public ClaimDocument(IUserClaim Other)
+			public ClaimDocument(IUserClaim other)
 			{
-				this.Type = Other.Type;
-				this.Value = Other.Value;
+				Type = other.Type;
+				Value = other.Value;
 			}
 		}
 
-		IMongoCollection<UserDocument> Users;
-		IMongoCollection<UserClaimsDocument> UserClaims;
-		IMongoCollection<UserSettingsDocument> UserSettings;
-
-		ILogger<UserCollectionV2> Logger;
-
-		MemoryCache UserCache;
+		readonly IMongoCollection<UserDocument> _users;
+		readonly IMongoCollection<UserClaimsDocument> _userClaims;
+		readonly IMongoCollection<UserSettingsDocument> _userSettings;
+		readonly ILogger<UserCollectionV2> _logger;
+		readonly MemoryCache _userCache;
 
 		/// <summary>
 		/// Static constructor
@@ -166,205 +163,205 @@ namespace Horde.Build.Collections.Impl
 		/// <summary>
 		/// Constructor
 		/// </summary>
-		/// <param name="DatabaseService"></param>
-		/// <param name="Logger"></param>
-		public UserCollectionV2(DatabaseService DatabaseService, ILogger<UserCollectionV2> Logger)
+		/// <param name="databaseService"></param>
+		/// <param name="logger"></param>
+		public UserCollectionV2(DatabaseService databaseService, ILogger<UserCollectionV2> logger)
 		{
-			this.Logger = Logger;
+			_logger = logger;
 
-			Users = DatabaseService.GetCollection<UserDocument>("UsersV2");
-			UserClaims = DatabaseService.GetCollection<UserClaimsDocument>("UserClaimsV2");
-			UserSettings = DatabaseService.GetCollection<UserSettingsDocument>("UserSettingsV2");
+			_users = databaseService.GetCollection<UserDocument>("UsersV2");
+			_userClaims = databaseService.GetCollection<UserClaimsDocument>("UserClaimsV2");
+			_userSettings = databaseService.GetCollection<UserSettingsDocument>("UserSettingsV2");
 
-			if (!DatabaseService.ReadOnlyMode)
+			if (!databaseService.ReadOnlyMode)
 			{
-				Users.Indexes.CreateOne(new CreateIndexModel<UserDocument>(Builders<UserDocument>.IndexKeys.Ascending(x => x.LoginUpper), new CreateIndexOptions { Unique = true }));
-				Users.Indexes.CreateOne(new CreateIndexModel<UserDocument>(Builders<UserDocument>.IndexKeys.Ascending(x => x.EmailUpper)));
+				_users.Indexes.CreateOne(new CreateIndexModel<UserDocument>(Builders<UserDocument>.IndexKeys.Ascending(x => x.LoginUpper), new CreateIndexOptions { Unique = true }));
+				_users.Indexes.CreateOne(new CreateIndexModel<UserDocument>(Builders<UserDocument>.IndexKeys.Ascending(x => x.EmailUpper)));
 			}
 
-			MemoryCacheOptions Options = new MemoryCacheOptions();
-			UserCache = new MemoryCache(Options);
+			MemoryCacheOptions options = new MemoryCacheOptions();
+			_userCache = new MemoryCache(options);
 		}
 
 		/// <inheritdoc/>
 		public void Dispose()
 		{
-			UserCache.Dispose();
+			_userCache.Dispose();
 		}
 
 		/// <inheritdoc/>
-		public async Task<IUser?> GetUserAsync(UserId Id)
+		public async Task<IUser?> GetUserAsync(UserId id)
 		{
-			IUser? User = await Users.Find(x => x.Id == Id).FirstOrDefaultAsync();
-			using (ICacheEntry Entry = UserCache.CreateEntry(Id))
+			IUser? user = await _users.Find(x => x.Id == id).FirstOrDefaultAsync();
+			using (ICacheEntry entry = _userCache.CreateEntry(id))
 			{
-				Entry.SetValue(User);
-				Entry.SetSlidingExpiration(TimeSpan.FromMinutes(30.0));
+				entry.SetValue(user);
+				entry.SetSlidingExpiration(TimeSpan.FromMinutes(30.0));
 			}
-			return User;
+			return user;
 		}
 
 		/// <inheritdoc/>
-		public async ValueTask<IUser?> GetCachedUserAsync(UserId? Id)
+		public async ValueTask<IUser?> GetCachedUserAsync(UserId? id)
 		{
-			IUser? User;
-			if(Id == null)
+			IUser? user;
+			if(id == null)
 			{
 				return null;
 			}
-			else if (UserCache.TryGetValue(Id.Value, out User))
+			else if (_userCache.TryGetValue(id.Value, out user))
 			{
-				return User;
+				return user;
 			}
 			else
 			{
-				return await GetUserAsync(Id.Value);
+				return await GetUserAsync(id.Value);
 			}
 		}
 
 		/// <inheritdoc/>
-		public async Task<List<IUser>> FindUsersAsync(IEnumerable<UserId>? Ids, string? NameRegex = null, int? Index = null, int? Count = null)
+		public async Task<List<IUser>> FindUsersAsync(IEnumerable<UserId>? ids, string? nameRegex = null, int? index = null, int? count = null)
 		{
-			FilterDefinition<UserDocument> Filter = FilterDefinition<UserDocument>.Empty;
-			if (Ids != null)
+			FilterDefinition<UserDocument> filter = FilterDefinition<UserDocument>.Empty;
+			if (ids != null)
 			{
-				Filter &= Builders<UserDocument>.Filter.In(x => x.Id, Ids);
+				filter &= Builders<UserDocument>.Filter.In(x => x.Id, ids);
 			}
 			
-			if (NameRegex != null)
+			if (nameRegex != null)
 			{
-				BsonRegularExpression Regex = new BsonRegularExpression(NameRegex, "i");
-				Filter &= Builders<UserDocument>.Filter.Regex(x => x.Name, Regex);
+				BsonRegularExpression regex = new BsonRegularExpression(nameRegex, "i");
+				filter &= Builders<UserDocument>.Filter.Regex(x => x.Name, regex);
 			}
 
-			Filter &= Builders<UserDocument>.Filter.Ne(x => x.Hidden, true);
+			filter &= Builders<UserDocument>.Filter.Ne(x => x.Hidden, true);
 
-			return await Users.Find(Filter).Range(Index, Count ?? 100).ToListAsync<UserDocument, IUser>();
+			return await _users.Find(filter).Range(index, count ?? 100).ToListAsync<UserDocument, IUser>();
 		}
 
 		/// <inheritdoc/>
-		public async Task<IUser?> FindUserByLoginAsync(string Login)
+		public async Task<IUser?> FindUserByLoginAsync(string login)
 		{
-			string LoginUpper = Login.ToUpperInvariant();
-			return await Users.Find(x => x.LoginUpper == LoginUpper).FirstOrDefaultAsync();
+			string loginUpper = login.ToUpperInvariant();
+			return await _users.Find(x => x.LoginUpper == loginUpper).FirstOrDefaultAsync();
 		}
 
 		/// <inheritdoc/>
-		public async Task<IUser> FindOrAddUserByLoginAsync(string Login, string? Name, string? Email)
+		public async Task<IUser> FindOrAddUserByLoginAsync(string login, string? name, string? email)
 		{
-			if (Login.Contains(' ', StringComparison.Ordinal))
+			if (login.Contains(' ', StringComparison.Ordinal))
 			{
-				Logger.LogWarning("Potentially invalid login name: {Login} (from {Trace})", Login, Environment.StackTrace);
+				_logger.LogWarning("Potentially invalid login name: {Login} (from {Trace})", login, Environment.StackTrace);
 			}
 
-			UserId NewUserId = UserId.GenerateNewId();
-			UpdateDefinition<UserDocument> Update = Builders<UserDocument>.Update.SetOnInsert(x => x.Id, NewUserId).SetOnInsert(x => x.Login, Login).Unset(x => x.Hidden);
+			UserId newUserId = UserId.GenerateNewId();
+			UpdateDefinition<UserDocument> update = Builders<UserDocument>.Update.SetOnInsert(x => x.Id, newUserId).SetOnInsert(x => x.Login, login).Unset(x => x.Hidden);
 
-			if (Name == null)
+			if (name == null)
 			{
-				Update = Update.SetOnInsert(x => x.Name, Login);
+				update = update.SetOnInsert(x => x.Name, login);
 			}
 			else
 			{
-				Update = Update.Set(x => x.Name, Name);
+				update = update.Set(x => x.Name, name);
 			}
 
-			if (Email != null)
+			if (email != null)
 			{
-				Update = Update.Set(x => x.Email, Email).Set(x => x.EmailUpper, Email.ToUpperInvariant());
+				update = update.Set(x => x.Email, email).Set(x => x.EmailUpper, email.ToUpperInvariant());
 			}
 
-			string LoginUpper = Login.ToUpperInvariant();
+			string loginUpper = login.ToUpperInvariant();
 
-			IUser User = await Users.FindOneAndUpdateAsync<UserDocument>(x => x.LoginUpper == LoginUpper, Update, new FindOneAndUpdateOptions<UserDocument, UserDocument> { IsUpsert = true, ReturnDocument = ReturnDocument.After });
-			if (User.Id == NewUserId)
+			IUser user = await _users.FindOneAndUpdateAsync<UserDocument>(x => x.LoginUpper == loginUpper, update, new FindOneAndUpdateOptions<UserDocument, UserDocument> { IsUpsert = true, ReturnDocument = ReturnDocument.After });
+			if (user.Id == newUserId)
 			{
-				Logger.LogInformation("Added new user {Name} ({UserId}, {Login}, {Email})", User.Name, User.Id, User.Login, User.Email);
+				_logger.LogInformation("Added new user {Name} ({UserId}, {Login}, {Email})", user.Name, user.Id, user.Login, user.Email);
 			}
 
-			return User;
+			return user;
 		}
 
 		/// <inheritdoc/>
-		public async Task<IUserClaims> GetClaimsAsync(UserId UserId)
+		public async Task<IUserClaims> GetClaimsAsync(UserId userId)
 		{
-			IUserClaims? Claims = await UserClaims.Find(x => x.Id == UserId).FirstOrDefaultAsync();
-			if(Claims == null)
+			IUserClaims? claims = await _userClaims.Find(x => x.Id == userId).FirstOrDefaultAsync();
+			if(claims == null)
 			{
-				Claims = new UserClaimsDocument(UserId);
+				claims = new UserClaimsDocument(userId);
 			}
-			return Claims;
+			return claims;
 		}
 
 		/// <inheritdoc/>
-		public async Task UpdateClaimsAsync(UserId UserId, IEnumerable<IUserClaim> Claims)
+		public async Task UpdateClaimsAsync(UserId userId, IEnumerable<IUserClaim> claims)
 		{
-			UserClaimsDocument NewDocument = new UserClaimsDocument(UserId);
-			NewDocument.Claims.AddRange(Claims.Select(x => new UserClaim(x)));
-			await UserClaims.ReplaceOneAsync(x => x.Id == UserId, NewDocument, new ReplaceOptions { IsUpsert = true });
+			UserClaimsDocument newDocument = new UserClaimsDocument(userId);
+			newDocument.Claims.AddRange(claims.Select(x => new UserClaim(x)));
+			await _userClaims.ReplaceOneAsync(x => x.Id == userId, newDocument, new ReplaceOptions { IsUpsert = true });
 		}
 
 		/// <inheritdoc/>
-		public async Task<IUserSettings> GetSettingsAsync(UserId UserId)
+		public async Task<IUserSettings> GetSettingsAsync(UserId userId)
 		{
-			IUserSettings? Settings = await UserSettings.Find(x => x.Id == UserId).FirstOrDefaultAsync();
-			if (Settings == null)
+			IUserSettings? settings = await _userSettings.Find(x => x.Id == userId).FirstOrDefaultAsync();
+			if (settings == null)
 			{
-				Settings = new UserSettingsDocument(UserId);
+				settings = new UserSettingsDocument(userId);
 			}
-			return Settings;
+			return settings;
 		}
 
 		/// <inheritdoc/>
-		public async Task UpdateSettingsAsync(UserId UserId, bool? EnableExperimentalFeatures = null, BsonValue? DashboardSettings = null, IEnumerable<JobId>? AddPinnedJobIds = null, IEnumerable<JobId>? RemovePinnedJobIds = null)
+		public async Task UpdateSettingsAsync(UserId userId, bool? enableExperimentalFeatures = null, BsonValue? dashboardSettings = null, IEnumerable<JobId>? addPinnedJobIds = null, IEnumerable<JobId>? removePinnedJobIds = null)
 		{
-			List<UpdateDefinition<UserSettingsDocument>> Updates = new List<UpdateDefinition<UserSettingsDocument>>();
-			if (EnableExperimentalFeatures != null)
+			List<UpdateDefinition<UserSettingsDocument>> updates = new List<UpdateDefinition<UserSettingsDocument>>();
+			if (enableExperimentalFeatures != null)
 			{
-				Updates.Add(Builders<UserSettingsDocument>.Update.SetOrUnsetNull(x => x.EnableExperimentalFeatures, EnableExperimentalFeatures));
+				updates.Add(Builders<UserSettingsDocument>.Update.SetOrUnsetNull(x => x.EnableExperimentalFeatures, enableExperimentalFeatures));
 			}
-			if (DashboardSettings != null)
+			if (dashboardSettings != null)
 			{
-				Updates.Add(Builders<UserSettingsDocument>.Update.Set(x => x.DashboardSettings, DashboardSettings));
+				updates.Add(Builders<UserSettingsDocument>.Update.Set(x => x.DashboardSettings, dashboardSettings));
 			}
-			if (AddPinnedJobIds != null && AddPinnedJobIds.Any())
+			if (addPinnedJobIds != null && addPinnedJobIds.Any())
 			{
-				Updates.Add(Builders<UserSettingsDocument>.Update.AddToSetEach(x => x.PinnedJobIds, AddPinnedJobIds));
+				updates.Add(Builders<UserSettingsDocument>.Update.AddToSetEach(x => x.PinnedJobIds, addPinnedJobIds));
 			}
-			if (RemovePinnedJobIds != null && RemovePinnedJobIds.Any())
+			if (removePinnedJobIds != null && removePinnedJobIds.Any())
 			{
-				Updates.Add(Builders<UserSettingsDocument>.Update.PullAll(x => x.PinnedJobIds, RemovePinnedJobIds));
+				updates.Add(Builders<UserSettingsDocument>.Update.PullAll(x => x.PinnedJobIds, removePinnedJobIds));
 			}
-			if (Updates.Count > 0)
+			if (updates.Count > 0)
 			{
-				await UserSettings.UpdateOneAsync<UserSettingsDocument>(x => x.Id == UserId, Builders<UserSettingsDocument>.Update.Combine(Updates), new UpdateOptions { IsUpsert = true });
+				await _userSettings.UpdateOneAsync<UserSettingsDocument>(x => x.Id == userId, Builders<UserSettingsDocument>.Update.Combine(updates), new UpdateOptions { IsUpsert = true });
 			}
 		}
 
 		/// <summary>
 		/// Upgrade from V1 collection
 		/// </summary>
-		/// <param name="UserCollectionV1"></param>
+		/// <param name="userCollectionV1"></param>
 		/// <returns></returns>
-		public async Task ResaveDocumentsAsync(UserCollectionV1 UserCollectionV1)
+		public async Task ResaveDocumentsAsync(UserCollectionV1 userCollectionV1)
 		{
-			await foreach ((IUser User, IUserClaims Claims, IUserSettings Settings) in UserCollectionV1.EnumerateDocumentsAsync())
+			await foreach ((IUser user, IUserClaims claims, IUserSettings settings) in userCollectionV1.EnumerateDocumentsAsync())
 			{
 				try
 				{
-					await Users.ReplaceOneAsync(x => x.Id == User.Id, new UserDocument(User), new ReplaceOptions { IsUpsert = true });
-					await UserClaims.ReplaceOneAsync(x => x.Id == User.Id, new UserClaimsDocument(Claims), new ReplaceOptions { IsUpsert = true });
-					await UserSettings.ReplaceOneAsync(x => x.Id == User.Id, new UserSettingsDocument(Settings), new ReplaceOptions { IsUpsert = true });
-					Logger.LogDebug("Updated user {UserId}", User.Id);
+					await _users.ReplaceOneAsync(x => x.Id == user.Id, new UserDocument(user), new ReplaceOptions { IsUpsert = true });
+					await _userClaims.ReplaceOneAsync(x => x.Id == user.Id, new UserClaimsDocument(claims), new ReplaceOptions { IsUpsert = true });
+					await _userSettings.ReplaceOneAsync(x => x.Id == user.Id, new UserSettingsDocument(settings), new ReplaceOptions { IsUpsert = true });
+					_logger.LogDebug("Updated user {UserId}", user.Id);
 				}
-				catch (MongoWriteException Ex)
+				catch (MongoWriteException ex)
 				{
-					Logger.LogWarning(Ex, "Unable to resave user {UserId}", User.Id);
+					_logger.LogWarning(ex, "Unable to resave user {UserId}", user.Id);
 				}
 
-				if(Settings.PinnedJobIds.Count > 0)
+				if(settings.PinnedJobIds.Count > 0)
 				{
-					await UpdateSettingsAsync(User.Id, AddPinnedJobIds: Settings.PinnedJobIds);
+					await UpdateSettingsAsync(user.Id, addPinnedJobIds: settings.PinnedJobIds);
 				}
 			}
 		}
