Horde: Fix coding conventions within Horde.Build.

diff --git a/Engine/Source/Programs/Horde/Horde.Build/Utilities/AgentUtilities.cs b/Engine/Source/Programs/Horde/Horde.Build/Utilities/AgentUtilities.cs
--- a/Engine/Source/Programs/Horde/Horde.Build/Utilities/AgentUtilities.cs
+++ b/Engine/Source/Programs/Horde/Horde.Build/Utilities/AgentUtilities.cs
@@ -2,13 +2,10 @@
 
 using System;
 using System.Collections.Generic;
+using System.Diagnostics;
 using System.IO;
 using System.IO.Compression;
-using System.Security.Cryptography;
-using System.Security.Cryptography.X509Certificates;
-using System.Text.RegularExpressions;
 using System.Text.Json;
-using System.Diagnostics;
 
 #nullable disable
 
@@ -19,26 +16,26 @@ namespace Horde.Build.Utilities
 		/// <summary>
 		/// Reads the version number from an archive
 		/// </summary>
-		/// <param name="Data">The archive data</param>
+		/// <param name="data">The archive data</param>
 		/// <returns></returns>
-		public static string ReadVersion(byte[] Data)
+		public static string ReadVersion(byte[] data)
 		{
-			MemoryStream InputStream = new MemoryStream(Data);
-			using (ZipArchive InputArchive = new ZipArchive(InputStream, ZipArchiveMode.Read, true))
+			MemoryStream inputStream = new MemoryStream(data);
+			using (ZipArchive inputArchive = new ZipArchive(inputStream, ZipArchiveMode.Read, true))
 			{
-				foreach (ZipArchiveEntry InputEntry in InputArchive.Entries)
+				foreach (ZipArchiveEntry inputEntry in inputArchive.Entries)
 				{
-					if (InputEntry.FullName.Equals("HordeAgent.dll", StringComparison.OrdinalIgnoreCase))
+					if (inputEntry.FullName.Equals("HordeAgent.dll", StringComparison.OrdinalIgnoreCase))
 					{
-						string TempFile = Path.GetTempFileName();
+						string tempFile = Path.GetTempFileName();
 						try
 						{
-							InputEntry.ExtractToFile(TempFile, true);
-							return FileVersionInfo.GetVersionInfo(TempFile).ProductVersion;
+							inputEntry.ExtractToFile(tempFile, true);
+							return FileVersionInfo.GetVersionInfo(tempFile).ProductVersion;
 						}
 						finally
 						{
-							File.Delete(TempFile);
+							File.Delete(tempFile);
 						}
 					}
 				}
@@ -46,48 +43,48 @@ namespace Horde.Build.Utilities
 			throw new Exception("Unable to find HordeAgent.dll in archive");
 		}
 
-
 		/// <summary>
 		/// Updates the agent app settings within the archive data
 		/// </summary>
-		/// <param name="Data">Data for the zip archive</param>
-		/// <param name="Settings">The settings to update</param>
+		/// <param name="data">Data for the zip archive</param>
+		/// <param name="settings">The settings to update</param>
 		/// <returns>New agent app data</returns>
-		public static byte[] UpdateAppSettings(byte[] Data, Dictionary<string, object> Settings)
+		public static byte[] UpdateAppSettings(byte[] data, Dictionary<string, object> settings)
 		{
 			bool bWrittenClientId = false;
 
-			MemoryStream OutputStream = new MemoryStream();
-			using (ZipArchive OutputArchive = new ZipArchive(OutputStream, ZipArchiveMode.Create, true))
+			MemoryStream outputStream = new MemoryStream();
+			using (ZipArchive outputArchive = new ZipArchive(outputStream, ZipArchiveMode.Create, true))
 			{
-				MemoryStream InputStream = new MemoryStream(Data);
-				using (ZipArchive InputArchive = new ZipArchive(InputStream, ZipArchiveMode.Read, true))
+				MemoryStream inputStream = new MemoryStream(data);
+				using (ZipArchive inputArchive = new ZipArchive(inputStream, ZipArchiveMode.Read, true))
 				{
-					foreach (ZipArchiveEntry InputEntry in InputArchive.Entries)
+					foreach (ZipArchiveEntry inputEntry in inputArchive.Entries)
 					{
-						ZipArchiveEntry OutputEntry = OutputArchive.CreateEntry(InputEntry.FullName);
+						ZipArchiveEntry outputEntry = outputArchive.CreateEntry(inputEntry.FullName);
 
-						using System.IO.Stream InputEntryStream = InputEntry.Open();
-						using System.IO.Stream OutputEntryStream = OutputEntry.Open();
+						using System.IO.Stream inputEntryStream = inputEntry.Open();
+						using System.IO.Stream outputEntryStream = outputEntry.Open();
 
-						if (InputEntry.FullName.Equals("appsettings.json", StringComparison.OrdinalIgnoreCase))
+						if (inputEntry.FullName.Equals("appsettings.json", StringComparison.OrdinalIgnoreCase))
 						{
-							using MemoryStream MemoryStream = new MemoryStream();
-							InputEntryStream.CopyTo(MemoryStream);
+							using MemoryStream memoryStream = new MemoryStream();
+							inputEntryStream.CopyTo(memoryStream);
 
-							Dictionary<string, Dictionary<string, object>> Document = JsonSerializer.Deserialize<Dictionary<string, Dictionary<string, object>>>(MemoryStream.ToArray());
-							foreach (KeyValuePair<string, object> Pair in Settings) {								
-								Document["Horde"][Pair.Key] = Pair.Value;
+							Dictionary<string, Dictionary<string, object>> document = JsonSerializer.Deserialize<Dictionary<string, Dictionary<string, object>>>(memoryStream.ToArray());
+							foreach (KeyValuePair<string, object> pair in settings) 
+							{								
+								document["Horde"][pair.Key] = pair.Value;
 							}
 
-							using Utf8JsonWriter Writer = new Utf8JsonWriter(OutputEntryStream, new JsonWriterOptions { Indented = true });
-							JsonSerializer.Serialize<Dictionary<string, Dictionary<string, object>>>(Writer, Document, new JsonSerializerOptions { WriteIndented = true });
+							using Utf8JsonWriter writer = new Utf8JsonWriter(outputEntryStream, new JsonWriterOptions { Indented = true });
+							JsonSerializer.Serialize<Dictionary<string, Dictionary<string, object>>>(writer, document, new JsonSerializerOptions { WriteIndented = true });
 
 							bWrittenClientId = true;
 						}
 						else
 						{
-							InputEntryStream.CopyTo(OutputEntryStream);
+							inputEntryStream.CopyTo(outputEntryStream);
 						}
 					}
 				}
@@ -98,8 +95,7 @@ namespace Horde.Build.Utilities
 				throw new InvalidDataException("Missing appsettings.json file from zip archive");
 			}
 
-			return OutputStream.ToArray();
+			return outputStream.ToArray();
 		}
-
 	}
 }
diff --git a/Engine/Source/Programs/Horde/Horde.Build/Utilities/ArtifactChunkStream.cs b/Engine/Source/Programs/Horde/Horde.Build/Utilities/ArtifactChunkStream.cs
--- a/Engine/Source/Programs/Horde/Horde.Build/Utilities/ArtifactChunkStream.cs
+++ b/Engine/Source/Programs/Horde/Horde.Build/Utilities/ArtifactChunkStream.cs
@@ -1,13 +1,11 @@
 // Copyright Epic Games, Inc. All Rights Reserved.
 
-using Grpc.Core;
-using HordeCommon.Rpc;
 using System;
-using System.Collections.Generic;
 using System.IO;
-using System.Linq;
 using System.Threading;
 using System.Threading.Tasks;
+using Grpc.Core;
+using HordeCommon.Rpc;
 
 namespace Horde.Build.Utilities
 {
@@ -19,48 +17,48 @@ namespace Horde.Build.Utilities
 		/// <summary>
 		/// The grpc client reader. Should probably be templatized
 		/// </summary>
-		IAsyncStreamReader<UploadArtifactRequest> Reader;
+		readonly IAsyncStreamReader<UploadArtifactRequest> _reader;
 
 		/// <summary>
 		/// Position within the stream
 		/// </summary>
-		long StreamPosition;
+		long _streamPosition;
 
 		/// <summary>
 		/// The length of the stream
 		/// </summary>
-		long StreamLength;
+		readonly long _streamLength;
 
 		/// <summary>
 		/// The current request being read from
 		/// </summary>
-		UploadArtifactRequest? Request;
+		UploadArtifactRequest? _request;
 
 		/// <summary>
 		/// Position within the current request
 		/// </summary>
-		int RequestPos;
+		int _requestPos;
 
 		/// <summary>
 		/// Constructor
 		/// </summary>
-		/// <param name="Reader">the grpc reader</param>
-		/// <param name="Length">filesize reported by the client</param>
-		public ArtifactChunkStream(IAsyncStreamReader<UploadArtifactRequest> Reader, long Length)
+		/// <param name="reader">the grpc reader</param>
+		/// <param name="length">filesize reported by the client</param>
+		public ArtifactChunkStream(IAsyncStreamReader<UploadArtifactRequest> reader, long length)
 		{
-			this.Reader = Reader;
-			this.StreamLength = Length;
+			_reader = reader;
+			_streamLength = length;
 		}
 
 		/// <inheritdoc/>
 		public override long Position
 		{
-			get { return StreamPosition; }
-			set { throw new NotSupportedException(); }
+			get => _streamPosition;
+			set => throw new NotSupportedException();
 		}
 
 		/// <inheritdoc/>
-		public override long Length => StreamLength;
+		public override long Length => _streamLength;
 
 		/// <inheritdoc/>
 		public override void SetLength(long value) => throw new NotSupportedException();
@@ -69,7 +67,7 @@ namespace Horde.Build.Utilities
 		public override bool CanSeek => false;
 
 		/// <inheritdoc/>
-		public override long Seek(long Offset, SeekOrigin Origin) => throw new NotSupportedException();
+		public override long Seek(long offset, SeekOrigin origin) => throw new NotSupportedException();
 
 		/// <inheritdoc/>
 		public override bool CanWrite => false;
@@ -81,45 +79,45 @@ namespace Horde.Build.Utilities
 		public override bool CanRead => true;
 
 		/// <inheritdoc/>
-		public override int Read(byte[] Buffer, int Offset, int Count)
+		public override int Read(byte[] buffer, int offset, int count)
 		{
-			return ReadAsync(Buffer, Offset, Count, CancellationToken.None).Result;
+			return ReadAsync(buffer, offset, count, CancellationToken.None).Result;
 		}
 
 		/// <inheritdoc/>
-		public override async Task<int> ReadAsync(byte[] Buffer, int Offset, int Count, CancellationToken CancellationToken)
+		public override async Task<int> ReadAsync(byte[] buffer, int offset, int count, CancellationToken cancellationToken)
 		{
-			return await ReadAsync(Buffer.AsMemory(Offset, Count), CancellationToken);
+			return await ReadAsync(buffer.AsMemory(offset, count), cancellationToken);
 		}
 
 		/// <inheritdoc/>
-		public override async ValueTask<int> ReadAsync(Memory<byte> Buffer, CancellationToken CancellationToken)
+		public override async ValueTask<int> ReadAsync(Memory<byte> buffer, CancellationToken cancellationToken)
 		{
-			int BytesRead = 0;
-			while (BytesRead < Buffer.Length && StreamPosition < StreamLength)
+			int bytesRead = 0;
+			while (bytesRead < buffer.Length && _streamPosition < _streamLength)
 			{
-				if (Request == null || RequestPos == Request.Data.Length)
+				if (_request == null || _requestPos == _request.Data.Length)
 				{
 					// Read the next request object
-					if (!await Reader.MoveNext())
+					if (!await _reader.MoveNext())
 					{
 						throw new EndOfStreamException("Unexpected end of stream while reading artifact");
 					}
 
-					Request = Reader.Current;
-					RequestPos = 0;
+					_request = _reader.Current;
+					_requestPos = 0;
 				}
 				else
 				{
 					// Copy data from the current request object
-					int NumBytesToCopy = Math.Min(Buffer.Length - BytesRead, Request.Data.Length - RequestPos);
-					Request.Data.Span.Slice(RequestPos, NumBytesToCopy).CopyTo(Buffer.Slice(BytesRead, NumBytesToCopy).Span);
-					RequestPos += NumBytesToCopy;
-					BytesRead += NumBytesToCopy;
-					StreamPosition += NumBytesToCopy;
+					int numBytesToCopy = Math.Min(buffer.Length - bytesRead, _request.Data.Length - _requestPos);
+					_request.Data.Span.Slice(_requestPos, numBytesToCopy).CopyTo(buffer.Slice(bytesRead, numBytesToCopy).Span);
+					_requestPos += numBytesToCopy;
+					bytesRead += numBytesToCopy;
+					_streamPosition += numBytesToCopy;
 				}
 			}
-			return BytesRead;
+			return bytesRead;
 		}
 
 		/// <inheritdoc/>
diff --git a/Engine/Source/Programs/Horde/Horde.Build/Utilities/AsyncCachedValue.cs b/Engine/Source/Programs/Horde/Horde.Build/Utilities/AsyncCachedValue.cs
--- a/Engine/Source/Programs/Horde/Horde.Build/Utilities/AsyncCachedValue.cs
+++ b/Engine/Source/Programs/Horde/Horde.Build/Utilities/AsyncCachedValue.cs
@@ -1,9 +1,7 @@
 // Copyright Epic Games, Inc. All Rights Reserved.
 
 using System;
-using System.Collections.Generic;
 using System.Diagnostics;
-using System.Linq;
 using System.Threading;
 using System.Threading.Tasks;
 
@@ -18,54 +16,54 @@ namespace Horde.Build.Utilities
 		class State
 		{
 			public readonly T Value;
-			Stopwatch Timer;
-			public Task<State>? Next;
+			readonly Stopwatch _timer;
+			public Task<State>? _next;
 
-			public TimeSpan Elapsed => Timer.Elapsed;
+			public TimeSpan Elapsed => _timer.Elapsed;
 
-			public State(T Value)
+			public State(T value)
 			{
-				this.Value = Value;
-				this.Timer = Stopwatch.StartNew();
+				Value = value;
+				_timer = Stopwatch.StartNew();
 			}
 		}
 
 		/// <summary>
 		/// The current state
 		/// </summary>
-		Task<State>? Current = null;
+		Task<State>? _current = null;
 
 		/// <summary>
 		/// Generator for the new value
 		/// </summary>
-		Func<Task<T>> Generator;
+		readonly Func<Task<T>> _generator;
 
 		/// <summary>
 		/// Time at which to start to refresh the value
 		/// </summary>
-		TimeSpan MinRefreshTime;
+		readonly TimeSpan _minRefreshTime;
 
 		/// <summary>
 		/// Time at which to wait for the value to refresh
 		/// </summary>
-		TimeSpan MaxRefreshTime;
+		readonly TimeSpan _maxRefreshTime;
 
 		/// <summary>
 		/// Default constructor
 		/// </summary>
-		public AsyncCachedValue(Func<Task<T>> Generator, TimeSpan RefreshTime)
-			: this(Generator, RefreshTime * 0.75, RefreshTime)
+		public AsyncCachedValue(Func<Task<T>> generator, TimeSpan refreshTime)
+			: this(generator, refreshTime * 0.75, refreshTime)
 		{
 		}
 
 		/// <summary>
 		/// Default constructor
 		/// </summary>
-		public AsyncCachedValue(Func<Task<T>> Generator, TimeSpan MinRefreshTime, TimeSpan MaxRefreshTime)
+		public AsyncCachedValue(Func<Task<T>> generator, TimeSpan minRefreshTime, TimeSpan maxRefreshTime)
 		{
-			this.Generator = Generator;
-			this.MinRefreshTime = MinRefreshTime;
-			this.MaxRefreshTime = MaxRefreshTime;
+			_generator = generator;
+			_minRefreshTime = minRefreshTime;
+			_maxRefreshTime = maxRefreshTime;
 		}
 
 		/// <summary>
@@ -74,56 +72,56 @@ namespace Horde.Build.Utilities
 		/// <returns>The cached value, if valid</returns>
 		public Task<T> GetAsync()
 		{
-			return GetAsync(MaxRefreshTime);
+			return GetAsync(_maxRefreshTime);
 		}
 
-		Task<State> CreateOrGetStateTask(ref Task<State>? StateTask)
+		Task<State> CreateOrGetStateTask(ref Task<State>? stateTask)
 		{
 			for(; ;)
 			{
-				Task<State>? CurrentStateTask = StateTask;
-				if (CurrentStateTask != null)
+				Task<State>? currentStateTask = stateTask;
+				if (currentStateTask != null)
 				{
-					return CurrentStateTask;
+					return currentStateTask;
 				}
 
-				Task<Task<State>> InnerNewStateTask = new Task<Task<State>>(() => CreateState());
-				if (Interlocked.CompareExchange(ref StateTask, InnerNewStateTask.Unwrap(), null) == null)
+				Task<Task<State>> innerNewStateTask = new Task<Task<State>>(() => CreateState());
+				if (Interlocked.CompareExchange(ref stateTask, innerNewStateTask.Unwrap(), null) == null)
 				{
-					InnerNewStateTask.Start();
+					innerNewStateTask.Start();
 				}
 			}
 		}
 
 		async Task<State> CreateState()
 		{
-			return new State(await Generator());
+			return new State(await _generator());
 		}
 
 		/// <summary>
 		/// Tries to get the current value
 		/// </summary>
 		/// <returns>The cached value, if valid</returns>
-		public async Task<T> GetAsync(TimeSpan MaxAge)
+		public async Task<T> GetAsync(TimeSpan maxAge)
 		{
-			Task<State> StateTask = CreateOrGetStateTask(ref Current);
+			Task<State> stateTask = CreateOrGetStateTask(ref _current);
 
-			State State = await StateTask;
-			if (State.Next != null && State.Next.IsCompleted)
+			State state = await stateTask;
+			if (state._next != null && state._next.IsCompleted)
 			{
-				_ = Interlocked.CompareExchange(ref Current, State.Next, StateTask);
-				State = await State.Next;
+				_ = Interlocked.CompareExchange(ref _current, state._next, stateTask);
+				state = await state._next;
 			}
-			if (State.Elapsed > MaxAge)
+			if (state.Elapsed > maxAge)
 			{
-				State = await CreateOrGetStateTask(ref State.Next);
+				state = await CreateOrGetStateTask(ref state._next);
 			}
-			if (State.Elapsed > MinRefreshTime)
+			if (state.Elapsed > _minRefreshTime)
 			{
-				_ = CreateOrGetStateTask(ref State.Next);
+				_ = CreateOrGetStateTask(ref state._next);
 			}
 
-			return State.Value;
+			return state.Value;
 		}
 	}
 }
diff --git a/Engine/Source/Programs/Horde/Horde.Build/Utilities/Aws.cs b/Engine/Source/Programs/Horde/Horde.Build/Utilities/Aws.cs
--- a/Engine/Source/Programs/Horde/Horde.Build/Utilities/Aws.cs
+++ b/Engine/Source/Programs/Horde/Horde.Build/Utilities/Aws.cs
@@ -15,35 +15,38 @@ namespace Horde.Build.Utilities
 		///
 		/// The AWS SDK provided functionality was proven too much work, so this is a simplified version.
 		/// </summary>
-		/// <param name="ProfileName">Name of the AWS profile</param>
-		/// <param name="CredentialsFilePath">Override for the credentials file path, set to null for default path</param>
+		/// <param name="profileName">Name of the AWS profile</param>
+		/// <param name="credentialsFilePath">Override for the credentials file path, set to null for default path</param>
 		/// <returns>Credentials as a tuple</returns>
 		/// <exception cref="Exception"></exception>
-		public static (string AccessKey, string SecretAccessKey, string SecretToken) ReadAwsCredentials(string ProfileName, string? CredentialsFilePath = null)
+		public static (string AccessKey, string SecretAccessKey, string SecretToken) ReadAwsCredentials(string profileName, string? credentialsFilePath = null)
 		{
-			CredentialsFilePath ??= Path.Join(Environment.GetFolderPath(Environment.SpecialFolder.UserProfile), ".aws", "credentials");
-			string[] Lines = File.ReadAllLines(CredentialsFilePath);
+			credentialsFilePath ??= Path.Join(Environment.GetFolderPath(Environment.SpecialFolder.UserProfile), ".aws", "credentials");
+			string[] lines = File.ReadAllLines(credentialsFilePath);
 		    
-			string ReadValue(string Line, string ExpectedKey)
+			static string ReadValue(string line, string expectedKey)
 			{
-				if (Line.StartsWith(ExpectedKey, StringComparison.Ordinal)) return Line.Split("=")[1].Trim();
-				throw new Exception($"Unable to read key/value on line {Line} for key {ExpectedKey}");
+				if (line.StartsWith(expectedKey, StringComparison.Ordinal))
+				{
+					return line.Split("=")[1].Trim();
+				}
+				throw new Exception($"Unable to read key/value on line {line} for key {expectedKey}");
 			}
 
-			for (int i = 0; i < Lines.Length; i++)
+			for (int i = 0; i < lines.Length; i++)
 			{
-				string Line = Lines[i];
-				if (Line == $"[{ProfileName}]")
+				string line = lines[i];
+				if (line == $"[{profileName}]")
 				{
-					string AccessKey = ReadValue(Lines[i + 1], "aws_access_key_id");
-					string SecretAccessKey = ReadValue(Lines[i + 2], "aws_secret_access_key");
-					string SessionToken = ReadValue(Lines[i + 3], "aws_session_token");
+					string accessKey = ReadValue(lines[i + 1], "aws_access_key_id");
+					string secretAccessKey = ReadValue(lines[i + 2], "aws_secret_access_key");
+					string sessionToken = ReadValue(lines[i + 3], "aws_session_token");
 
-					return (AccessKey, SecretAccessKey, SessionToken);
+					return (accessKey, secretAccessKey, sessionToken);
 				}
 			}
 
-			throw new Exception($"Unable to find profile {ProfileName} in file {CredentialsFilePath}");
+			throw new Exception($"Unable to find profile {profileName} in file {credentialsFilePath}");
 		}
 	}
 }
\ No newline at end of file
diff --git a/Engine/Source/Programs/Horde/Horde.Build/Utilities/BackgroundTask.cs b/Engine/Source/Programs/Horde/Horde.Build/Utilities/BackgroundTask.cs
--- a/Engine/Source/Programs/Horde/Horde.Build/Utilities/BackgroundTask.cs
+++ b/Engine/Source/Programs/Horde/Horde.Build/Utilities/BackgroundTask.cs
@@ -11,28 +11,28 @@ namespace Horde.Build.Utilities
 	/// </summary>
 	public sealed class BackgroundTask : IDisposable
 	{
-		Func<CancellationToken, Task> RunTask;
-		Task Task = Task.CompletedTask;
-		CancellationTokenSource CancellationTokenSource = new CancellationTokenSource();
+		readonly Func<CancellationToken, Task> _runTask;
+		Task _task = Task.CompletedTask;
+		readonly CancellationTokenSource _cancellationTokenSource = new CancellationTokenSource();
 
 		/// <summary>
 		/// Constructor
 		/// </summary>
-		/// <param name="RunTask"></param>
-		public BackgroundTask(Func<CancellationToken, Task> RunTask)
+		/// <param name="runTask"></param>
+		public BackgroundTask(Func<CancellationToken, Task> runTask)
 		{
-			this.RunTask = RunTask;
+			_runTask = runTask;
 		}
 
 		/// <inheritdoc/>
 		public void Dispose()
 		{
-			if (!Task.IsCompleted)
+			if (!_task.IsCompleted)
 			{
-				CancellationTokenSource.Cancel();
+				_cancellationTokenSource.Cancel();
 				StopAsync().Wait();
 			}
-			CancellationTokenSource.Dispose();
+			_cancellationTokenSource.Dispose();
 		}
 
 		/// <summary>
@@ -40,11 +40,11 @@ namespace Horde.Build.Utilities
 		/// </summary>
 		public void Start()
 		{
-			if (!Task.IsCompleted)
+			if (!_task.IsCompleted)
 			{
 				throw new InvalidOperationException("Background task is already running");
 			}
-			Task = Task.Run(() => RunTask(CancellationTokenSource.Token), CancellationTokenSource.Token);
+			_task = Task.Run(() => _runTask(_cancellationTokenSource.Token), _cancellationTokenSource.Token);
 		}
 
 		/// <summary>
@@ -55,8 +55,8 @@ namespace Horde.Build.Utilities
 		{
 			try
 			{
-				CancellationTokenSource.Cancel();
-				await Task;
+				_cancellationTokenSource.Cancel();
+				await _task;
 			}
 			catch (OperationCanceledException)
 			{
diff --git a/Engine/Source/Programs/Horde/Horde.Build/Utilities/BaseDatabaseIndexes.cs b/Engine/Source/Programs/Horde/Horde.Build/Utilities/BaseDatabaseIndexes.cs
--- a/Engine/Source/Programs/Horde/Horde.Build/Utilities/BaseDatabaseIndexes.cs
+++ b/Engine/Source/Programs/Horde/Horde.Build/Utilities/BaseDatabaseIndexes.cs
@@ -13,41 +13,40 @@ namespace Horde.Build.Utilities
 		/// <summary>
 		/// If database is in read-only mode (i.e no indexes can be created/modified)
 		/// </summary>
-		private bool DatabaseReadOnlyMode;
+		private readonly bool _databaseReadOnlyMode;
 
 		/// <summary>
 		/// Constructor
 		/// </summary>
-		/// <param name="DatabaseReadOnlyMode"></param>
-		protected BaseDatabaseIndexes(bool DatabaseReadOnlyMode)
+		/// <param name="databaseReadOnlyMode"></param>
+		protected BaseDatabaseIndexes(bool databaseReadOnlyMode)
 		{
-			this.DatabaseReadOnlyMode = DatabaseReadOnlyMode;
+			_databaseReadOnlyMode = databaseReadOnlyMode;
 		}
 
 		/// <summary>
 		/// Create or get an index
 		/// </summary>
-		/// <param name="Collection">Name of collection</param>
-		/// <param name="Name">Name of index (used for checking if index exists)</param>
-		/// <param name="IndexDef">Type-specific index definition</param>
+		/// <param name="collection">Name of collection</param>
+		/// <param name="name">Name of index (used for checking if index exists)</param>
+		/// <param name="indexDef">Type-specific index definition</param>
 		/// <typeparam name="T">Document type</typeparam>
 		/// <returns></returns>
 		/// <exception cref="ArgumentException">If index cannot be found and database is in read-only mode</exception>
-		protected string CreateOrGetIndex<T>(IMongoCollection<T> Collection, string Name, IndexKeysDefinition<T> IndexDef)
+		protected string CreateOrGetIndex<T>(IMongoCollection<T> collection, string name, IndexKeysDefinition<T> indexDef)
 		{
-			if (DatabaseReadOnlyMode)
+			if (_databaseReadOnlyMode)
 			{
-				if (Collection.Indexes.List().ToList().Find(x => x.TryGetString("name", out string? n) && n == Name) == null)
+				if (collection.Indexes.List().ToList().Find(x => x.TryGetString("name", out string? n) && n == name) == null)
 				{
-					throw new ArgumentException($"Index with name {Name} not found (running in database read-only mode)");
+					throw new ArgumentException($"Index with name {name} not found (running in database read-only mode)");
 				}
-				return Name;
+				return name;
 			}
 
-			CreateIndexOptions IndexOptions = new() { Name = Name, Background = true };
-			Collection.Indexes.CreateOne(new CreateIndexModel<T>(IndexDef, IndexOptions));
-			return Name;
+			CreateIndexOptions indexOptions = new() { Name = name, Background = true };
+			collection.Indexes.CreateOne(new CreateIndexModel<T>(indexDef, indexOptions));
+			return name;
 		}
 	}
-
 }
diff --git a/Engine/Source/Programs/Horde/Horde.Build/Utilities/BlockKit/ActionsBlock.cs b/Engine/Source/Programs/Horde/Horde.Build/Utilities/BlockKit/ActionsBlock.cs
--- a/Engine/Source/Programs/Horde/Horde.Build/Utilities/BlockKit/ActionsBlock.cs
+++ b/Engine/Source/Programs/Horde/Horde.Build/Utilities/BlockKit/ActionsBlock.cs
@@ -18,32 +18,32 @@ namespace Horde.Build.Utilities.Slack.BlockKit
 		public List<ActionElement> Elements { get; } = new List<ActionElement>();
 
 		/// <inheritdoc/>
-		public override void Write(Utf8JsonWriter Writer, JsonSerializerOptions Options)
+		public override void Write(Utf8JsonWriter writer, JsonSerializerOptions options)
 		{
-			Writer.WriteStartObject();
-			Writer.WriteString("type", "actions");
-			Writer.WriteStartArray("elements");
-			foreach (ActionElement Element in Elements)
+			writer.WriteStartObject();
+			writer.WriteString("type", "actions");
+			writer.WriteStartArray("elements");
+			foreach (ActionElement element in Elements)
 			{
-				Element.Write(Writer, Options);
+				element.Write(writer, options);
 			}
-			Writer.WriteEndArray();
-			Writer.WriteEndObject();
+			writer.WriteEndArray();
+			writer.WriteEndObject();
 		}
 
 		/// <summary>
 		/// Helper method to add a new <see cref="ActionButton"/> to the block.
 		/// </summary>
-		/// <param name="Text">The button text.</param>
-		/// <param name="URL">The URL to navigate to when clicked, if any.</param>
-		/// <param name="Value">
+		/// <param name="text">The button text.</param>
+		/// <param name="url">The URL to navigate to when clicked, if any.</param>
+		/// <param name="value">
 		/// The value, if any, that should be passed in the interaction payload if the button is clicked.
 		/// </param>
-		/// <param name="ActionId">An optional identifier for handling interaction logic.</param>
-		/// <param name="Style">An option button style to apply to the button.</param>
-		public void AddButton(string Text, Uri? URL = null, string? Value = null, string? ActionId = null, ActionButton.ButtonStyle Style = ActionButton.ButtonStyle.Default)
+		/// <param name="actionId">An optional identifier for handling interaction logic.</param>
+		/// <param name="style">An option button style to apply to the button.</param>
+		public void AddButton(string text, Uri? url = null, string? value = null, string? actionId = null, ActionButton.ButtonStyle style = ActionButton.ButtonStyle.Default)
 		{
-			Elements.Add(new ActionButton(Text, URL, Value, ActionId, Style));
+			Elements.Add(new ActionButton(text, url, value, actionId, style));
 		}
 	}
 
@@ -55,9 +55,9 @@ namespace Horde.Build.Utilities.Slack.BlockKit
 		/// <summary>
 		/// Writes the Json for an Action block element to the given <see cref="Utf8JsonWriter"/>.
 		/// </summary>
-		/// <param name="Writer"></param>
-		/// <param name="Options"></param>
-		public abstract void Write(Utf8JsonWriter Writer, JsonSerializerOptions Options);
+		/// <param name="writer"></param>
+		/// <param name="options"></param>
+		public abstract void Write(Utf8JsonWriter writer, JsonSerializerOptions options);
 	}
 
 	/// <summary>
@@ -89,20 +89,20 @@ namespace Horde.Build.Utilities.Slack.BlockKit
 		/// <summary>
 		/// Construct a new Button action element.
 		/// </summary>
-		/// <param name="Text">The button text.</param>
-		/// <param name="URL">The URL to navigate to when clicked, if any.</param>
-		/// <param name="Value">
+		/// <param name="text">The button text.</param>
+		/// <param name="url">The URL to navigate to when clicked, if any.</param>
+		/// <param name="value">
 		/// The value, if any, that should be passed in the interaction payload if the button is clicked.
 		/// </param>
-		/// <param name="ActionId">An optional identifier for handling interaction logic.</param>
-		/// <param name="Style">An option button style to apply to the button.</param>
-		public ActionButton(string Text, Uri? URL = null, string? Value = null, string? ActionId = null, ButtonStyle Style = ButtonStyle.Default)
+		/// <param name="actionId">An optional identifier for handling interaction logic.</param>
+		/// <param name="style">An option button style to apply to the button.</param>
+		public ActionButton(string text, Uri? url = null, string? value = null, string? actionId = null, ButtonStyle style = ButtonStyle.Default)
 		{
-			this.Text = new TextObject(Text, IsMarkdown: false);
-			this.URL = URL;
-			this.Value = Value;
-			this.ActionId = ActionId;
-			this.Style = Style;
+			Text = new TextObject(text, isMarkdown: false);
+			Url = url;
+			Value = value;
+			ActionId = actionId;
+			Style = style;
 		}
 
 		/// <summary>
@@ -113,7 +113,7 @@ namespace Horde.Build.Utilities.Slack.BlockKit
 		/// <summary>
 		/// The URL to navigate to when clicked, if any.
 		/// </summary>
-		public Uri? URL { get; set; }
+		public Uri? Url { get; set; }
 
 		/// <summary>
 		/// The value, if any, that should be passed in the interaction payload if the button is clicked.
@@ -131,29 +131,29 @@ namespace Horde.Build.Utilities.Slack.BlockKit
 		public ButtonStyle Style { get; set; }
 
 		/// <inheritdoc/>
-		public override void Write(Utf8JsonWriter Writer, JsonSerializerOptions Options)
+		public override void Write(Utf8JsonWriter writer, JsonSerializerOptions options)
 		{
-			Writer.WriteStartObject();
-			Writer.WriteString("type", "button");
-			Text.Write(Writer);
-			if (URL != null)
+			writer.WriteStartObject();
+			writer.WriteString("type", "button");
+			Text.Write(writer);
+			if (Url != null)
 			{
-				Writer.WriteString("url", URL.AbsoluteUri);
+				writer.WriteString("url", Url.AbsoluteUri);
 			}
-			if (!string.IsNullOrWhiteSpace(Value))
+			if (!String.IsNullOrWhiteSpace(Value))
 			{
-				Writer.WriteString("value", Value);
+				writer.WriteString("value", Value);
 			}
-			if (!string.IsNullOrWhiteSpace(ActionId))
+			if (!String.IsNullOrWhiteSpace(ActionId))
 			{
-				Writer.WriteString("action_id", ActionId);
+				writer.WriteString("action_id", ActionId);
 			}
 			if (Style != ButtonStyle.Default)
 			{
-				Writer.WriteString("style", Style.ToString().ToLower(CultureInfo.CurrentCulture));
+				writer.WriteString("style", Style.ToString().ToLower(CultureInfo.CurrentCulture));
 			}
 
-			Writer.WriteEndObject();
+			writer.WriteEndObject();
 		}
 	}
 }
diff --git a/Engine/Source/Programs/Horde/Horde.Build/Utilities/BlockKit/BlockBase.cs b/Engine/Source/Programs/Horde/Horde.Build/Utilities/BlockKit/BlockBase.cs
--- a/Engine/Source/Programs/Horde/Horde.Build/Utilities/BlockKit/BlockBase.cs
+++ b/Engine/Source/Programs/Horde/Horde.Build/Utilities/BlockKit/BlockBase.cs
@@ -20,9 +20,9 @@ namespace Horde.Build.Utilities.Slack.BlockKit
 		/// <summary>
 		/// Writes the Json for a BlockKit block to the given <see cref="Utf8JsonWriter"/>.
 		/// </summary>
-		public override void Write(Utf8JsonWriter Writer, BlockBase Value, JsonSerializerOptions Options)
+		public override void Write(Utf8JsonWriter writer, BlockBase value, JsonSerializerOptions options)
 		{
-			Value.Write(Writer, Options);
+			value.Write(writer, options);
 		}
 	}
 
@@ -35,8 +35,8 @@ namespace Horde.Build.Utilities.Slack.BlockKit
 		/// <summary>
 		/// Write a BlockKit block to a <see cref="Utf8JsonWriter"/>
 		/// </summary>
-		/// <param name="Writer"></param>
-		/// <param name="Options"></param>
-		public abstract void Write(Utf8JsonWriter Writer, JsonSerializerOptions Options);
+		/// <param name="writer"></param>
+		/// <param name="options"></param>
+		public abstract void Write(Utf8JsonWriter writer, JsonSerializerOptions options);
 	}
 }
diff --git a/Engine/Source/Programs/Horde/Horde.Build/Utilities/BlockKit/BlockKitAttachment.cs b/Engine/Source/Programs/Horde/Horde.Build/Utilities/BlockKit/BlockKitAttachment.cs
--- a/Engine/Source/Programs/Horde/Horde.Build/Utilities/BlockKit/BlockKitAttachment.cs
+++ b/Engine/Source/Programs/Horde/Horde.Build/Utilities/BlockKit/BlockKitAttachment.cs
@@ -22,9 +22,9 @@ namespace Horde.Build.Utilities.Slack.BlockKit
 		/// <summary>
 		/// Writes the Json for a BlockKitAttachment to the given <see cref="Utf8JsonWriter"/>.
 		/// </summary>
-		public override void Write(Utf8JsonWriter Writer, BlockKitAttachment Value, JsonSerializerOptions Options)
+		public override void Write(Utf8JsonWriter writer, BlockKitAttachment value, JsonSerializerOptions options)
 		{
-			Value.Write(Writer, Options);
+			value.Write(writer, options);
 		}
 	}
 
@@ -75,23 +75,23 @@ namespace Horde.Build.Utilities.Slack.BlockKit
 		/// <summary>
 		/// Write a BlockKit block to a <see cref="Utf8JsonWriter"/>
 		/// </summary>
-		/// <param name="Writer"></param>
-		/// <param name="Options"></param>
-		public void Write(Utf8JsonWriter Writer, JsonSerializerOptions Options)
+		/// <param name="writer"></param>
+		/// <param name="options"></param>
+		public void Write(Utf8JsonWriter writer, JsonSerializerOptions options)
 		{
-			Writer.WriteStartObject();
+			writer.WriteStartObject();
 			if (FallbackText != null)
 			{
-				Writer.WriteString("fallback", FallbackText);
+				writer.WriteString("fallback", FallbackText);
 			}
-			Writer.WriteString("color", $"#{Color.ToArgb() & 0xFFFFFF:x6}");
-			Writer.WriteStartArray("blocks");
-			foreach (BlockBase Block in Blocks)
+			writer.WriteString("color", $"#{Color.ToArgb() & 0xFFFFFF:x6}");
+			writer.WriteStartArray("blocks");
+			foreach (BlockBase block in Blocks)
 			{
-				Block.Write(Writer, Options);
+				block.Write(writer, options);
 			}
-			Writer.WriteEndArray();
-			Writer.WriteEndObject();
+			writer.WriteEndArray();
+			writer.WriteEndObject();
 		}
 	}
 }
diff --git a/Engine/Source/Programs/Horde/Horde.Build/Utilities/BlockKit/DividerBlock.cs b/Engine/Source/Programs/Horde/Horde.Build/Utilities/BlockKit/DividerBlock.cs
--- a/Engine/Source/Programs/Horde/Horde.Build/Utilities/BlockKit/DividerBlock.cs
+++ b/Engine/Source/Programs/Horde/Horde.Build/Utilities/BlockKit/DividerBlock.cs
@@ -10,11 +10,11 @@ namespace Horde.Build.Utilities.Slack.BlockKit
 	public class DividerBlock : BlockBase
 	{
 		/// <inheritdoc/>
-		public override void Write(Utf8JsonWriter Writer, JsonSerializerOptions Options)
+		public override void Write(Utf8JsonWriter writer, JsonSerializerOptions options)
 		{
-			Writer.WriteStartObject();
-			Writer.WriteString("type", "divider");
-			Writer.WriteEndObject();
+			writer.WriteStartObject();
+			writer.WriteString("type", "divider");
+			writer.WriteEndObject();
 		}
 	}
 }
diff --git a/Engine/Source/Programs/Horde/Horde.Build/Utilities/BlockKit/HeaderBlock.cs b/Engine/Source/Programs/Horde/Horde.Build/Utilities/BlockKit/HeaderBlock.cs
--- a/Engine/Source/Programs/Horde/Horde.Build/Utilities/BlockKit/HeaderBlock.cs
+++ b/Engine/Source/Programs/Horde/Horde.Build/Utilities/BlockKit/HeaderBlock.cs
@@ -12,12 +12,12 @@ namespace Horde.Build.Utilities.Slack.BlockKit
 		/// <summary>
 		/// Constructor
 		/// </summary>
-		/// <param name="Text">Any text to initiale the Section with</param>
-		/// <param name="IsMarkdown">If true, text is markdown, else text is considered plain text</param>
-		/// <param name="Emoji">If true, allows escaping of emojis with colons.</param>
-		public HeaderBlock(string Text = "", bool IsMarkdown = true, bool Emoji = false)
+		/// <param name="text">Any text to initiale the Section with</param>
+		/// <param name="isMarkdown">If true, text is markdown, else text is considered plain text</param>
+		/// <param name="emoji">If true, allows escaping of emojis with colons.</param>
+		public HeaderBlock(string text = "", bool isMarkdown = true, bool emoji = false)
 		{
-			this.Text = new TextObject(Text, IsMarkdown, Emoji);
+			Text = new TextObject(text, isMarkdown, emoji);
 		}
 
 		/// <summary>
@@ -26,12 +26,12 @@ namespace Horde.Build.Utilities.Slack.BlockKit
 		public TextObject Text { get; }
 
 		/// <inheritdoc/>
-		public override void Write(Utf8JsonWriter Writer, JsonSerializerOptions Options)
+		public override void Write(Utf8JsonWriter writer, JsonSerializerOptions options)
 		{
-			Writer.WriteStartObject();
-			Writer.WriteString("type", "header");
-			Text.Write(Writer);
-			Writer.WriteEndObject();
+			writer.WriteStartObject();
+			writer.WriteString("type", "header");
+			Text.Write(writer);
+			writer.WriteEndObject();
 		}
 	}
 }
diff --git a/Engine/Source/Programs/Horde/Horde.Build/Utilities/BlockKit/SectionBlock.cs b/Engine/Source/Programs/Horde/Horde.Build/Utilities/BlockKit/SectionBlock.cs
--- a/Engine/Source/Programs/Horde/Horde.Build/Utilities/BlockKit/SectionBlock.cs
+++ b/Engine/Source/Programs/Horde/Horde.Build/Utilities/BlockKit/SectionBlock.cs
@@ -12,12 +12,12 @@ namespace Horde.Build.Utilities.Slack.BlockKit
 		/// <summary>
 		/// Constructor
 		/// </summary>
-		/// <param name="Text">Any text to initiale the Section with</param>
-		/// <param name="IsMarkdown">If true, text is markdown, else text is considered plain text</param>
-		/// <param name="Emoji">If true, allows escaping of emoji with colon.</param>
-		public SectionBlock(string Text = "", bool IsMarkdown = true, bool Emoji = false)
+		/// <param name="text">Any text to initiale the Section with</param>
+		/// <param name="isMarkdown">If true, text is markdown, else text is considered plain text</param>
+		/// <param name="emoji">If true, allows escaping of emoji with colon.</param>
+		public SectionBlock(string text = "", bool isMarkdown = true, bool emoji = false)
 		{
-			this.Text = new TextObject(Text, IsMarkdown, Emoji);
+			Text = new TextObject(text, isMarkdown, emoji);
 		}
 
 		/// <summary>
@@ -26,12 +26,12 @@ namespace Horde.Build.Utilities.Slack.BlockKit
 		public TextObject Text { get; }
 
 		/// <inheritdoc/>
-		public override void Write(Utf8JsonWriter Writer, JsonSerializerOptions Options)
+		public override void Write(Utf8JsonWriter writer, JsonSerializerOptions options)
 		{
-			Writer.WriteStartObject();
-			Writer.WriteString("type", "section");
-			Text.Write(Writer);
-			Writer.WriteEndObject();
+			writer.WriteStartObject();
+			writer.WriteString("type", "section");
+			Text.Write(writer);
+			writer.WriteEndObject();
 		}
 	}
 }
diff --git a/Engine/Source/Programs/Horde/Horde.Build/Utilities/BlockKit/SlackResponses.cs b/Engine/Source/Programs/Horde/Horde.Build/Utilities/BlockKit/SlackResponses.cs
--- a/Engine/Source/Programs/Horde/Horde.Build/Utilities/BlockKit/SlackResponses.cs
+++ b/Engine/Source/Programs/Horde/Horde.Build/Utilities/BlockKit/SlackResponses.cs
@@ -1,6 +1,5 @@
 // Copyright Epic Games, Inc. All Rights Reserved.
 
-using System.Drawing;
 using Horde.Build.Utilities.Slack.BlockKit;
 
 namespace Horde.Build.Utilities.BlockKit
@@ -14,17 +13,17 @@ namespace Horde.Build.Utilities.BlockKit
 		/// <summary>
 		/// Creates a response to send to Slack about receiving bad data for a request.
 		/// </summary>
-		/// <param name="Message">A message explaining what was wrong with the data.</param>
+		/// <param name="message">A message explaining what was wrong with the data.</param>
 		/// <returns>A <see cref="BlockKitMessage"/> to be sent to Slack.</returns>
-		public static BlockKitMessage CreateBadDataResponse(string Message)
+		public static BlockKitMessage CreateBadDataResponse(string message)
 		{
-			BlockKitMessage BlockKitMessage = new BlockKitMessage();
-			BlockKitAttachment Attachment = new BlockKitAttachment();
-			Attachment.Color = BlockKitAttachmentColors.Error;
-			Attachment.Blocks.Add(new SectionBlock("Failed to process request due to bad data from Slack:"));
-			Attachment.Blocks.Add(new SectionBlock(Message));
-			BlockKitMessage.Attachments.Add(Attachment);
-			return BlockKitMessage;
+			BlockKitMessage blockKitMessage = new BlockKitMessage();
+			BlockKitAttachment attachment = new BlockKitAttachment();
+			attachment.Color = BlockKitAttachmentColors.Error;
+			attachment.Blocks.Add(new SectionBlock("Failed to process request due to bad data from Slack:"));
+			attachment.Blocks.Add(new SectionBlock(message));
+			blockKitMessage.Attachments.Add(attachment);
+			return blockKitMessage;
 		}
 
 		/// <summary>
@@ -32,14 +31,14 @@ namespace Horde.Build.Utilities.BlockKit
 		/// being found in the DB.
 		/// </summary>
 		/// <returns>A <see cref="BlockKitMessage"/> to be sent to Slack.</returns>
-		public static BlockKitMessage CreateIssueNotFoundResponse(int IssueId)
+		public static BlockKitMessage CreateIssueNotFoundResponse(int issueId)
 		{
-			BlockKitMessage BlockKitMessage = new BlockKitMessage();
-			BlockKitAttachment Attachment = new BlockKitAttachment();
-			Attachment.Color = BlockKitAttachmentColors.Error;
-			Attachment.Blocks.Add(new SectionBlock($"Failed to assign issue as issue {IssueId} could not be found."));
-			BlockKitMessage.Attachments.Add(Attachment);
-			return BlockKitMessage;
+			BlockKitMessage blockKitMessage = new BlockKitMessage();
+			BlockKitAttachment attachment = new BlockKitAttachment();
+			attachment.Color = BlockKitAttachmentColors.Error;
+			attachment.Blocks.Add(new SectionBlock($"Failed to assign issue as issue {issueId} could not be found."));
+			blockKitMessage.Attachments.Add(attachment);
+			return blockKitMessage;
 		}
 	}
 }
diff --git a/Engine/Source/Programs/Horde/Horde.Build/Utilities/BlockKit/TextObject.cs b/Engine/Source/Programs/Horde/Horde.Build/Utilities/BlockKit/TextObject.cs
--- a/Engine/Source/Programs/Horde/Horde.Build/Utilities/BlockKit/TextObject.cs
+++ b/Engine/Source/Programs/Horde/Horde.Build/Utilities/BlockKit/TextObject.cs
@@ -1,6 +1,5 @@
 // Copyright Epic Games, Inc. All Rights Reserved.
 
-using System;
 using System.Text.Json;
 
 namespace Horde.Build.Utilities.Slack.BlockKit
@@ -13,15 +12,15 @@ namespace Horde.Build.Utilities.Slack.BlockKit
 		/// <summary>
 		/// Constructor
 		/// </summary>
-		/// <param name="Text">Any text to initialize with</param>
-		/// <param name="IsMarkdown">If true, text is markdown, else text is considered plain text</param>
-		/// <param name="Emoji">If true, allows escaping of emojis with a colon</param>
-		public TextObject(string Text = "", bool IsMarkdown = true, bool Emoji = false)
+		/// <param name="text">Any text to initialize with</param>
+		/// <param name="isMarkdown">If true, text is markdown, else text is considered plain text</param>
+		/// <param name="emoji">If true, allows escaping of emojis with a colon</param>
+		public TextObject(string text = "", bool isMarkdown = true, bool emoji = false)
 		{
-			this.Text = Text;
-			this.IsMarkdown = IsMarkdown;
+			Text = text;
+			IsMarkdown = isMarkdown;
 			// force emoji to false if markdown is true as this invalidates the schema
-			this.Emoji = this.IsMarkdown ? false : Emoji;
+			Emoji = !IsMarkdown && emoji;
 		}
 
 		/// <summary>
@@ -42,26 +41,26 @@ namespace Horde.Build.Utilities.Slack.BlockKit
 		/// <summary>
 		/// Writes the TextObject JSON to a <see cref="Utf8JsonWriter"/>
 		/// </summary>
-		/// <param name="Writer"></param>
-		public void Write(Utf8JsonWriter Writer)
+		/// <param name="writer"></param>
+		public void Write(Utf8JsonWriter writer)
 		{
-			Writer.WriteStartObject("text");
-			Writer.WriteString("type", IsMarkdown ? "mrkdwn" : "plain_text");
+			writer.WriteStartObject("text");
+			writer.WriteString("type", IsMarkdown ? "mrkdwn" : "plain_text");
 
 			const int MaxLength = 2048;
 
-			string TruncatedText = Text;
-			if (TruncatedText.Length > MaxLength)
+			string truncatedText = Text;
+			if (truncatedText.Length > MaxLength)
 			{
-				TruncatedText = TruncatedText.Substring(0, MaxLength) + "...\n(message truncated)";
+				truncatedText = truncatedText.Substring(0, MaxLength) + "...\n(message truncated)";
 			}
 
-			Writer.WriteString("text", TruncatedText);
+			writer.WriteString("text", truncatedText);
 			if (Emoji)
 			{
-				Writer.WriteBoolean("emoji", Emoji);
+				writer.WriteBoolean("emoji", Emoji);
 			}
-			Writer.WriteEndObject();
+			writer.WriteEndObject();
 		}
 	}
 }
diff --git a/Engine/Source/Programs/Horde/Horde.Build/Utilities/BloomFilter.cs b/Engine/Source/Programs/Horde/Horde.Build/Utilities/BloomFilter.cs
--- a/Engine/Source/Programs/Horde/Horde.Build/Utilities/BloomFilter.cs
+++ b/Engine/Source/Programs/Horde/Horde.Build/Utilities/BloomFilter.cs
@@ -2,8 +2,6 @@
 
 using System;
 using System.Collections.Generic;
-using System.Linq;
-using System.Threading.Tasks;
 
 namespace Horde.Build.Utilities
 {
@@ -16,28 +14,28 @@ namespace Horde.Build.Utilities
 		/// <summary>
 		/// Enumerates the hash codes for a given block of text
 		/// </summary>
-		/// <param name="Text">The text to parse</param>
+		/// <param name="text">The text to parse</param>
 		/// <returns></returns>
-		public static IEnumerable<int> GetHashCodes(ReadOnlyMemory<byte> Text)
+		public static IEnumerable<int> GetHashCodes(ReadOnlyMemory<byte> text)
 		{
-			if(Text.Length >= NgramLength)
+			if(text.Length >= NgramLength)
 			{
-				int Value = 0;
-				for(int Idx = 0; Idx < NgramLength - 1; Idx++)
+				int value = 0;
+				for(int idx = 0; idx < NgramLength - 1; idx++)
 				{
-					Value = (Value << 8) | ToLowerUtf8(Text.Span[Idx]);
+					value = (value << 8) | ToLowerUtf8(text.Span[idx]);
 				}
-				for(int Idx = NgramLength - 1; Idx < Text.Length; Idx++)
+				for(int idx = NgramLength - 1; idx < text.Length; idx++)
 				{
-					Value = ((Value << 8) | ToLowerUtf8(Text.Span[Idx])) & NgramMask;
+					value = ((value << 8) | ToLowerUtf8(text.Span[idx])) & NgramMask;
 
-					int HashValue = Value;
-					yield return HashValue;
+					int hashValue = value;
+					yield return hashValue;
 
-					for(int HashIdx = 1; HashIdx < NumHashes; HashIdx++)
+					for(int hashIdx = 1; hashIdx < NumHashes; hashIdx++)
 					{
-						HashValue = Scramble(HashValue);
-						yield return HashValue;
+						hashValue = Scramble(hashValue);
+						yield return hashValue;
 					}
 				}
 			}
@@ -46,31 +44,31 @@ namespace Horde.Build.Utilities
 		/// <summary>
 		/// Converts a character into a format for hashing
 		/// </summary>
-		/// <param name="Value">The input byte</param>
+		/// <param name="value">The input byte</param>
 		/// <returns>The value to include in the trigram</returns>
-		static byte ToLowerUtf8(byte Value)
+		static byte ToLowerUtf8(byte value)
 		{
-			if (Value >= 'A' && Value <= 'Z')
+			if (value >= 'A' && value <= 'Z')
 			{
-				return (byte)(Value + 'a' - 'A');
+				return (byte)(value + 'a' - 'A');
 			}
 			else
 			{
-				return Value;
+				return value;
 			}
 		}
 
 		/// <summary>
 		/// Scramble the given number using a pseudo-RNG
 		/// </summary>
-		/// <param name="Value">Initial value</param>
+		/// <param name="value">Initial value</param>
 		/// <returns>The scrambled value</returns>
-		private static int Scramble(int Value)
+		private static int Scramble(int value)
 		{
-			Value ^= Value << 13;
-			Value ^= Value >> 17;
-			Value ^= Value << 5;
-			return Value;
+			value ^= value << 13;
+			value ^= value >> 17;
+			value ^= value << 5;
+			return value;
 		}
 	}
 
@@ -87,24 +85,24 @@ namespace Horde.Build.Utilities
 		/// <summary>
 		/// Constructor
 		/// </summary>
-		/// <param name="Data">Data to construct from</param>
-		public ReadOnlyBloomFilter(ReadOnlyMemory<byte> Data)
+		/// <param name="data">Data to construct from</param>
+		public ReadOnlyBloomFilter(ReadOnlyMemory<byte> data)
 		{
-			this.Data = Data;
+			Data = data;
 		}
 
 		/// <summary>
 		/// Determines if the filter may contain a value
 		/// </summary>
-		/// <param name="HashCodes">Sequence of hash codes</param>
+		/// <param name="hashCodes">Sequence of hash codes</param>
 		/// <returns>True if the value is in the set</returns>
-		public bool Contains(IEnumerable<int> HashCodes)
+		public bool Contains(IEnumerable<int> hashCodes)
 		{
-			ReadOnlySpan<byte> Span = Data.Span;
-			foreach (int HashCode in HashCodes)
+			ReadOnlySpan<byte> span = Data.Span;
+			foreach (int hashCode in hashCodes)
 			{
-				int Index = HashCode & ((Span.Length << 3) - 1);
-				if ((Span[Index >> 3] & (1 << (Index & 7))) == 0)
+				int index = hashCode & ((span.Length << 3) - 1);
+				if ((span[index >> 3] & (1 << (index & 7))) == 0)
 				{
 					return false;
 				}
@@ -118,21 +116,21 @@ namespace Horde.Build.Utilities
 		/// <returns>Value between 0 and 1</returns>
 		public double CalculateLoadFactor()
 		{
-			int SetCount = 0;
+			int setCount = 0;
 
-			ReadOnlySpan<byte> Span = Data.Span;
-			for (int Idx = 0; Idx < Span.Length; Idx++)
+			ReadOnlySpan<byte> span = Data.Span;
+			for (int idx = 0; idx < span.Length; idx++)
 			{
-				for (int Bit = 0; Bit < 8; Bit++)
+				for (int bit = 0; bit < 8; bit++)
 				{
-					if ((Span[Idx] & (1 << Bit)) != 0)
+					if ((span[idx] & (1 << bit)) != 0)
 					{
-						SetCount++;
+						setCount++;
 					}
 				}
 			}
 
-			return (double)SetCount / (Span.Length * 8);
+			return (double)setCount / (span.Length * 8);
 		}
 	}
 
@@ -149,26 +147,26 @@ namespace Horde.Build.Utilities
 		/// <summary>
 		/// Constructs a new filter of the given size
 		/// </summary>
-		/// <param name="MinSize">The minimum size of the filter. Will be rounded up to the next power of two.</param>
-		public BloomFilter(int MinSize)
+		/// <param name="minSize">The minimum size of the filter. Will be rounded up to the next power of two.</param>
+		public BloomFilter(int minSize)
 		{
-			int Size = MinSize;
-			while ((Size & (Size - 1)) != 0)
+			int size = minSize;
+			while ((size & (size - 1)) != 0)
 			{
-				Size |= (Size - 1);
-				Size++;
+				size |= (size - 1);
+				size++;
 			}
-			Data = new byte[Size];
+			Data = new byte[size];
 		}
 
 		/// <summary>
 		/// Constructs a filter from raw data
 		/// </summary>
-		/// <param name="Data">The data to construct from</param>
-		public BloomFilter(byte[] Data)
+		/// <param name="data">The data to construct from</param>
+		public BloomFilter(byte[] data)
 		{
-			this.Data = Data;
-			if ((Data.Length & (Data.Length - 1)) != 0)
+			Data = data;
+			if ((data.Length & (data.Length - 1)) != 0)
 			{
 				throw new ArgumentException("Array for bloom filter must be a power of 2 in size.");
 			}
@@ -177,27 +175,27 @@ namespace Horde.Build.Utilities
 		/// <summary>
 		/// Add a value to the filter
 		/// </summary>
-		/// <param name="HashCodes">Sequence of hash codes</param>
-		public void Add(IEnumerable<int> HashCodes)
+		/// <param name="hashCodes">Sequence of hash codes</param>
+		public void Add(IEnumerable<int> hashCodes)
 		{
-			foreach(int HashCode in HashCodes)
+			foreach(int hashCode in hashCodes)
 			{
-				int Index = HashCode & ((Data.Length << 3) - 1);
-				Data[Index >> 3] = (byte)(Data[Index >> 3] | (1 << (Index & 7)));
+				int index = hashCode & ((Data.Length << 3) - 1);
+				Data[index >> 3] = (byte)(Data[index >> 3] | (1 << (index & 7)));
 			}
 		}
 
 		/// <summary>
 		/// Determines if the filter may contain a value
 		/// </summary>
-		/// <param name="HashCodes">Sequence of hash codes</param>
+		/// <param name="hashCodes">Sequence of hash codes</param>
 		/// <returns>True if the value is in the set</returns>
-		public bool Contains(IEnumerable<int> HashCodes)
+		public bool Contains(IEnumerable<int> hashCodes)
 		{
-			foreach (int HashCode in HashCodes)
+			foreach (int hashCode in hashCodes)
 			{
-				int Index = HashCode & ((Data.Length << 3) - 1);
-				if ((Data[Index >> 3] & (1 << (Index & 7))) == 0)
+				int index = hashCode & ((Data.Length << 3) - 1);
+				if ((Data[index >> 3] & (1 << (index & 7))) == 0)
 				{
 					return false;
 				}
@@ -211,20 +209,20 @@ namespace Horde.Build.Utilities
 		/// <returns>Value between 0 and 1</returns>
 		public double CalculateLoadFactor()
 		{
-			int SetCount = 0;
-			int AllCount = 0;
-			for (int Idx = 0; Idx < Data.Length; Idx++)
+			int setCount = 0;
+			int allCount = 0;
+			for (int idx = 0; idx < Data.Length; idx++)
 			{
-				for (int Bit = 0; Bit < 8; Bit++)
+				for (int bit = 0; bit < 8; bit++)
 				{
-					if ((Data[Idx] & (1 << Bit)) != 0)
+					if ((Data[idx] & (1 << bit)) != 0)
 					{
-						SetCount++;
+						setCount++;
 					}
-					AllCount++;
+					allCount++;
 				}
 			}
-			return (double)SetCount / AllCount;
+			return (double)setCount / allCount;
 		}
 	}
 }
diff --git a/Engine/Source/Programs/Horde/Horde.Build/Utilities/BsonDocumentExtensions.cs b/Engine/Source/Programs/Horde/Horde.Build/Utilities/BsonDocumentExtensions.cs
--- a/Engine/Source/Programs/Horde/Horde.Build/Utilities/BsonDocumentExtensions.cs
+++ b/Engine/Source/Programs/Horde/Horde.Build/Utilities/BsonDocumentExtensions.cs
@@ -1,11 +1,8 @@
 // Copyright Epic Games, Inc. All Rights Reserved.
 
-using MongoDB.Bson;
 using System;
-using System.Collections.Generic;
 using System.Diagnostics.CodeAnalysis;
-using System.Linq;
-using System.Threading.Tasks;
+using MongoDB.Bson;
 
 namespace Horde.Build.Utilities
 {
@@ -17,48 +14,48 @@ namespace Horde.Build.Utilities
 		/// <summary>
 		/// Gets a property value from a document or subdocument, indicated with dotted notation
 		/// </summary>
-		/// <param name="Document">Document to get a property for</param>
-		/// <param name="Name">Name of the property</param>
-		/// <param name="Type">Expected type of the property</param>
-		/// <param name="OutValue">Receives the property value</param>
+		/// <param name="document">Document to get a property for</param>
+		/// <param name="name">Name of the property</param>
+		/// <param name="type">Expected type of the property</param>
+		/// <param name="outValue">Receives the property value</param>
 		/// <returns>True if the property exists and was of the correct type</returns>
-		public static bool TryGetPropertyValue(this BsonDocument Document, string Name, BsonType Type, [NotNullWhen(true)] out BsonValue? OutValue)
+		public static bool TryGetPropertyValue(this BsonDocument document, string name, BsonType type, [NotNullWhen(true)] out BsonValue? outValue)
 		{
-			int DotIdx = Name.IndexOf('.', StringComparison.Ordinal);
-			if (DotIdx == -1)
+			int dotIdx = name.IndexOf('.', StringComparison.Ordinal);
+			if (dotIdx == -1)
 			{
-				return TryGetDirectPropertyValue(Document, Name, Type, out OutValue);
+				return TryGetDirectPropertyValue(document, name, type, out outValue);
 			}
 
-			BsonValue? DocValue;
-			if (TryGetDirectPropertyValue(Document, Name.Substring(0, DotIdx), BsonType.Document, out DocValue))
+			BsonValue? docValue;
+			if (TryGetDirectPropertyValue(document, name.Substring(0, dotIdx), BsonType.Document, out docValue))
 			{
-				return TryGetPropertyValue(DocValue.AsBsonDocument, Name.Substring(DotIdx + 1), Type, out OutValue);
+				return TryGetPropertyValue(docValue.AsBsonDocument, name.Substring(dotIdx + 1), type, out outValue);
 			}
 
-			OutValue = null;
+			outValue = null;
 			return false;
 		}
 
 		/// <summary>
 		/// Gets a property value that's an immediate child of the document
 		/// </summary>
-		/// <param name="Document">Document to get a property for</param>
-		/// <param name="Name">Name of the property</param>
-		/// <param name="Type">Expected type of the property</param>
-		/// <param name="OutValue">Receives the property value</param>
+		/// <param name="document">Document to get a property for</param>
+		/// <param name="name">Name of the property</param>
+		/// <param name="type">Expected type of the property</param>
+		/// <param name="outValue">Receives the property value</param>
 		/// <returns>True if the property exists and was of the correct type</returns>
-		private static bool TryGetDirectPropertyValue(this BsonDocument Document, string Name, BsonType Type, [NotNullWhen(true)] out BsonValue? OutValue)
+		private static bool TryGetDirectPropertyValue(this BsonDocument document, string name, BsonType type, [NotNullWhen(true)] out BsonValue? outValue)
 		{
-			BsonValue Value;
-			if (Document.TryGetValue(Name, out Value) && Value.BsonType == Type)
+			BsonValue value;
+			if (document.TryGetValue(name, out value) && value.BsonType == type)
 			{
-				OutValue = Value;
+				outValue = value;
 				return true;
 			}
 			else
 			{
-				OutValue = null;
+				outValue = null;
 				return false;
 			}
 		}
@@ -66,21 +63,21 @@ namespace Horde.Build.Utilities
 		/// <summary>
 		/// Gets an int32 value from the document
 		/// </summary>
-		/// <param name="Document">Document to get a property for</param>
-		/// <param name="Name">Name of the property</param>
-		/// <param name="OutValue">Receives the property value</param>
+		/// <param name="document">Document to get a property for</param>
+		/// <param name="name">Name of the property</param>
+		/// <param name="outValue">Receives the property value</param>
 		/// <returns>True if the property was retrieved</returns>
-		public static bool TryGetInt32(this BsonDocument Document, string Name, out int OutValue)
+		public static bool TryGetInt32(this BsonDocument document, string name, out int outValue)
 		{
-			BsonValue? Value;
-			if (Document.TryGetPropertyValue(Name, BsonType.Int32, out Value))
+			BsonValue? value;
+			if (document.TryGetPropertyValue(name, BsonType.Int32, out value))
 			{
-				OutValue = Value.AsInt32;
+				outValue = value.AsInt32;
 				return true;
 			}
 			else
 			{
-				OutValue = 0;
+				outValue = 0;
 				return false;
 			}
 		}
@@ -88,21 +85,21 @@ namespace Horde.Build.Utilities
 		/// <summary>
 		/// Gets a string value from the document
 		/// </summary>
-		/// <param name="Document">Document to get a property for</param>
-		/// <param name="Name">Name of the property</param>
-		/// <param name="OutValue">Receives the property value</param>
+		/// <param name="document">Document to get a property for</param>
+		/// <param name="name">Name of the property</param>
+		/// <param name="outValue">Receives the property value</param>
 		/// <returns>True if the property was retrieved</returns>
-		public static bool TryGetString(this BsonDocument Document, string Name, [NotNullWhen(true)] out string? OutValue)
+		public static bool TryGetString(this BsonDocument document, string name, [NotNullWhen(true)] out string? outValue)
 		{
-			BsonValue? Value;
-			if (Document.TryGetPropertyValue(Name, BsonType.String, out Value))
+			BsonValue? value;
+			if (document.TryGetPropertyValue(name, BsonType.String, out value))
 			{
-				OutValue = Value.AsString;
+				outValue = value.AsString;
 				return true;
 			}
 			else
 			{
-				OutValue = null;
+				outValue = null;
 				return false;
 			}
 		}
diff --git a/Engine/Source/Programs/Horde/Horde.Build/Utilities/BuildGraphExtensions.cs b/Engine/Source/Programs/Horde/Horde.Build/Utilities/BuildGraphExtensions.cs
--- a/Engine/Source/Programs/Horde/Horde.Build/Utilities/BuildGraphExtensions.cs
+++ b/Engine/Source/Programs/Horde/Horde.Build/Utilities/BuildGraphExtensions.cs
@@ -1,15 +1,15 @@
 // Copyright Epic Games, Inc. All Rights Reserved.
 
-using EpicGames.BuildGraph;
-using EpicGames.Perforce;
-using Horde.Build.Collections;
-using Horde.Build.Models;
 using System;
 using System.Collections.Generic;
 using System.Globalization;
 using System.IO;
 using System.Linq;
 using System.Threading.Tasks;
+using EpicGames.BuildGraph;
+using EpicGames.Perforce;
+using Horde.Build.Collections;
+using Horde.Build.Models;
 
 namespace Horde.Build.Utilities
 {
@@ -20,35 +20,35 @@ namespace Horde.Build.Utilities
 	{
 		class PerforceScriptContext : IBgScriptReaderContext
 		{
-			IPerforceConnection PerforceConnection;
-			string Stream;
-			int ChangeNumber;
+			readonly IPerforceConnection _perforceConnection;
+			readonly string _stream;
+			readonly int _changeNumber;
 
-			public PerforceScriptContext(IPerforceConnection Perforce, string Stream, int ChangeNumber)
+			public PerforceScriptContext(IPerforceConnection perforce, string stream, int changeNumber)
 			{
-				this.PerforceConnection = Perforce;
-				this.Stream = Stream;
-				this.ChangeNumber = ChangeNumber;
+				_perforceConnection = perforce;
+				_stream = stream;
+				_changeNumber = changeNumber;
 			}
 
-			public async Task<bool> ExistsAsync(string Path)
+			public async Task<bool> ExistsAsync(string path)
 			{
-				List<FStatRecord> Records = await PerforceConnection.FStatAsync($"{Stream}/{Path}@{ChangeNumber}").ToListAsync();
-				return Records.Count > 0;
+				List<FStatRecord> records = await _perforceConnection.FStatAsync($"{_stream}/{path}@{_changeNumber}").ToListAsync();
+				return records.Count > 0;
 			}
 
-			public object GetNativePath(string Path)
+			public object GetNativePath(string path)
 			{
-				return $"{Stream}/{Path}";
+				return $"{_stream}/{path}";
 			}
 
-			public async Task<byte[]?> ReadAsync(string Path)
+			public async Task<byte[]?> ReadAsync(string path)
 			{
-				string TempFileName = System.IO.Path.GetTempFileName();
+				string tempFileName = System.IO.Path.GetTempFileName();
 				try
 				{
-					await PerforceConnection.PrintAsync(TempFileName, $"{Stream}/{Path}");
-					return await File.ReadAllBytesAsync(TempFileName);
+					await _perforceConnection.PrintAsync(tempFileName, $"{_stream}/{path}");
+					return await File.ReadAllBytesAsync(tempFileName);
 				}
 				catch
 				{
@@ -58,14 +58,14 @@ namespace Horde.Build.Utilities
 				{
 					try
 					{
-						File.SetAttributes(TempFileName, FileAttributes.Normal);
-						File.Delete(TempFileName);
+						File.SetAttributes(tempFileName, FileAttributes.Normal);
+						File.Delete(tempFileName);
 					}
 					catch { }
 				}
 			}
 
-			public Task<string[]> FindAsync(string Path)
+			public Task<string[]> FindAsync(string path)
 			{
 				throw new NotImplementedException();
 			}
@@ -74,133 +74,133 @@ namespace Horde.Build.Utilities
 		/// <summary>
 		/// Create a graph from a script in Perforce
 		/// </summary>
-		/// <param name="GraphCollection"></param>
-		/// <param name="PerforceConnection"></param>
-		/// <param name="Stream"></param>
-		/// <param name="ChangeNumber"></param>
-		/// <param name="CodeChangeNumber"></param>
-		/// <param name="Arguments"></param>
+		/// <param name="graphCollection"></param>
+		/// <param name="perforceConnection"></param>
+		/// <param name="stream"></param>
+		/// <param name="changeNumber"></param>
+		/// <param name="codeChangeNumber"></param>
+		/// <param name="arguments"></param>
 		/// <returns></returns>
 		/// <exception cref="NotImplementedException"></exception>
-		public static async Task<IGraph> AddAsync(this IGraphCollection GraphCollection, IPerforceConnection PerforceConnection, IStream Stream, int ChangeNumber, int CodeChangeNumber, List<string> Arguments)
+		public static async Task<IGraph> AddAsync(this IGraphCollection graphCollection, IPerforceConnection perforceConnection, IStream stream, int changeNumber, int codeChangeNumber, List<string> arguments)
 		{
-			PerforceScriptContext Context = new PerforceScriptContext(PerforceConnection, Stream.Name, ChangeNumber);
+			PerforceScriptContext context = new PerforceScriptContext(perforceConnection, stream.Name, changeNumber);
 
-			BgScriptSchema? Schema = null!;
+			BgScriptSchema? schema = null!;
 
-			Dictionary<string, string> Properties = new Dictionary<string, string>();
-			Properties["HostPlatform"] = "";
-			Properties["RootDir"] = "";
-			Properties["IsBuildMachine"] = "true";
-			Properties["Branch"] = Stream.Name;
-			Properties["Change"] = ChangeNumber.ToString(CultureInfo.InvariantCulture);
-			Properties["EscapedBranch"] = Stream.Name.Replace('/', '+');
-			Properties["CodeChange"] = CodeChangeNumber.ToString(CultureInfo.InvariantCulture);
+			Dictionary<string, string> properties = new Dictionary<string, string>();
+			properties["HostPlatform"] = "";
+			properties["RootDir"] = "";
+			properties["IsBuildMachine"] = "true";
+			properties["Branch"] = stream.Name;
+			properties["Change"] = changeNumber.ToString(CultureInfo.InvariantCulture);
+			properties["EscapedBranch"] = stream.Name.Replace('/', '+');
+			properties["CodeChange"] = codeChangeNumber.ToString(CultureInfo.InvariantCulture);
 
-			string? ScriptFile = null;
-			foreach (string Argument in Arguments)
+			string? scriptFile = null;
+			foreach (string argument in arguments)
 			{
 				const string ScriptPrefix = "-Script=";
-				if (Argument.StartsWith(ScriptPrefix, StringComparison.OrdinalIgnoreCase))
+				if (argument.StartsWith(ScriptPrefix, StringComparison.OrdinalIgnoreCase))
 				{
-					ScriptFile = Argument.Substring(ScriptPrefix.Length);
+					scriptFile = argument.Substring(ScriptPrefix.Length);
 				}
 			}
 
-			List<string> Targets = new List<string>();
-			foreach (string Argument in Arguments)
+			List<string> targets = new List<string>();
+			foreach (string argument in arguments)
 			{
 				const string TargetPrefix = "-Target=";
-				if (Argument.StartsWith(TargetPrefix, StringComparison.OrdinalIgnoreCase))
+				if (argument.StartsWith(TargetPrefix, StringComparison.OrdinalIgnoreCase))
 				{
-					Targets.AddRange(Argument.Substring(TargetPrefix.Length).Split(';'));
+					targets.AddRange(argument.Substring(TargetPrefix.Length).Split(';'));
 				}
 			}
 
-			Dictionary<string, string> Options = new Dictionary<string, string>();
-			foreach (string Argument in Arguments)
+			Dictionary<string, string> options = new Dictionary<string, string>();
+			foreach (string argument in arguments)
 			{
 				const string SetPrefix = "-Set:";
-				if (Argument.StartsWith(SetPrefix, StringComparison.OrdinalIgnoreCase))
+				if (argument.StartsWith(SetPrefix, StringComparison.OrdinalIgnoreCase))
 				{
-					int EqualsIdx = Argument.IndexOf('=', StringComparison.Ordinal);
-					if (EqualsIdx != -1)
+					int equalsIdx = argument.IndexOf('=', StringComparison.Ordinal);
+					if (equalsIdx != -1)
 					{
-						Options[Argument.Substring(SetPrefix.Length, EqualsIdx - SetPrefix.Length)] = Argument.Substring(EqualsIdx + 1);
+						options[argument.Substring(SetPrefix.Length, equalsIdx - SetPrefix.Length)] = argument.Substring(equalsIdx + 1);
 					}
 				}
 			}
 
-			if(ScriptFile == null)
+			if(scriptFile == null)
 			{
 				throw new NotImplementedException();
 			}
 
-			BgGraph? Script = await BgScriptReader.ReadAsync(Context, ScriptFile, Options, Properties, Schema, EpicGames.Core.Log.Logger);
-			if (Script == null)
+			BgGraph? script = await BgScriptReader.ReadAsync(context, scriptFile, options, properties, schema, EpicGames.Core.Log.Logger);
+			if (script == null)
 			{
 				throw new NotImplementedException();
 			}
 
-			HashSet<BgNode> TargetNodes = new HashSet<BgNode>();
-			foreach (string Target in Targets)
+			HashSet<BgNode> targetNodes = new HashSet<BgNode>();
+			foreach (string target in targets)
 			{
-				if (Script.NameToNode.TryGetValue(Target, out BgNode? Node))
+				if (script.NameToNode.TryGetValue(target, out BgNode? node))
 				{
-					TargetNodes.Add(Node);
+					targetNodes.Add(node);
 				}
-				else if (Script.NameToAgent.TryGetValue(Target, out BgAgent? Agent))
+				else if (script.NameToAgent.TryGetValue(target, out BgAgent? agent))
 				{
-					TargetNodes.UnionWith(Agent.Nodes);
+					targetNodes.UnionWith(agent.Nodes);
 				}
-				else if (Script.NameToAggregate.TryGetValue(Target, out BgAggregate? Aggregate))
+				else if (script.NameToAggregate.TryGetValue(target, out BgAggregate? aggregate))
 				{
-					TargetNodes.UnionWith(Aggregate.RequiredNodes);
+					targetNodes.UnionWith(aggregate.RequiredNodes);
 				}
 				else
 				{
 					throw new NotImplementedException();
 				}
 			}
-			Script.Select(TargetNodes);
+			script.Select(targetNodes);
 
-			List<NewGroup> NewGroups = new List<NewGroup>();
-			foreach (BgAgent Agent in Script.Agents)
+			List<NewGroup> newGroups = new List<NewGroup>();
+			foreach (BgAgent agent in script.Agents)
 			{
-				List<NewNode> NewNodes = new List<NewNode>();
-				foreach (BgNode Node in Agent.Nodes)
+				List<NewNode> newNodes = new List<NewNode>();
+				foreach (BgNode node in agent.Nodes)
 				{
-					const bool bAllowRetry = true;
-					NewNode NewNode = new NewNode(Node.Name, Node.InputDependencies.Select(x => x.Name).ToList(), Node.OrderDependencies.Select(x => x.Name).ToList(), HordeCommon.Priority.Normal, bAllowRetry, Node.bRunEarly, Node.bNotifyOnWarnings, new Dictionary<string, string>(), new Dictionary<string, string>());
-					NewNodes.Add(NewNode);
+					const bool BAllowRetry = true;
+					NewNode newNode = new NewNode(node.Name, node.InputDependencies.Select(x => x.Name).ToList(), node.OrderDependencies.Select(x => x.Name).ToList(), HordeCommon.Priority.Normal, BAllowRetry, node.bRunEarly, node.bNotifyOnWarnings, new Dictionary<string, string>(), new Dictionary<string, string>());
+					newNodes.Add(newNode);
 				}
 
-				string AgentType = Agent.PossibleTypes.FirstOrDefault(x => Stream.AgentTypes.ContainsKey(x)) ?? Agent.PossibleTypes.FirstOrDefault() ?? "Unknown";
-				NewGroups.Add(new NewGroup(AgentType, NewNodes));
+				string agentType = agent.PossibleTypes.FirstOrDefault(x => stream.AgentTypes.ContainsKey(x)) ?? agent.PossibleTypes.FirstOrDefault() ?? "Unknown";
+				newGroups.Add(new NewGroup(agentType, newNodes));
 			}
 
-			List<NewAggregate> NewAggregates = new List<NewAggregate>();
-			foreach (BgAggregate Aggregate in Script.NameToAggregate.Values)
+			List<NewAggregate> newAggregates = new List<NewAggregate>();
+			foreach (BgAggregate aggregate in script.NameToAggregate.Values)
 			{
-				NewAggregate NewAggregate = new NewAggregate(Aggregate.Name, Aggregate.RequiredNodes.Select(x => x.Name).ToList());
-				NewAggregates.Add(NewAggregate);
+				NewAggregate newAggregate = new NewAggregate(aggregate.Name, aggregate.RequiredNodes.Select(x => x.Name).ToList());
+				newAggregates.Add(newAggregate);
 			}
 
-			List<NewLabel> NewLabels = new List<NewLabel>();
-			foreach (BgLabel Label in Script.Labels)
+			List<NewLabel> newLabels = new List<NewLabel>();
+			foreach (BgLabel label in script.Labels)
 			{
-				NewLabel NewLabel = new NewLabel();
-				NewLabel.DashboardName = String.IsNullOrEmpty(Label.DashboardName) ? null : Label.DashboardName;
-				NewLabel.DashboardCategory = String.IsNullOrEmpty(Label.DashboardCategory) ? null : Label.DashboardCategory;
-				NewLabel.UgsName = String.IsNullOrEmpty(Label.UgsBadge) ? null : Label.UgsBadge;
-				NewLabel.UgsProject = String.IsNullOrEmpty(Label.UgsProject) ? null : Label.UgsProject;
-				NewLabel.Change = (Label.Change == BgLabelChange.Code) ? HordeCommon.LabelChange.Code : HordeCommon.LabelChange.Current;
-				NewLabel.RequiredNodes = Label.RequiredNodes.Select(x => x.Name).ToList();
-				NewLabel.IncludedNodes = Label.IncludedNodes.Select(x => x.Name).ToList();
-				NewLabels.Add(NewLabel);
+				NewLabel newLabel = new NewLabel();
+				newLabel.DashboardName = String.IsNullOrEmpty(label.DashboardName) ? null : label.DashboardName;
+				newLabel.DashboardCategory = String.IsNullOrEmpty(label.DashboardCategory) ? null : label.DashboardCategory;
+				newLabel.UgsName = String.IsNullOrEmpty(label.UgsBadge) ? null : label.UgsBadge;
+				newLabel.UgsProject = String.IsNullOrEmpty(label.UgsProject) ? null : label.UgsProject;
+				newLabel.Change = (label.Change == BgLabelChange.Code) ? HordeCommon.LabelChange.Code : HordeCommon.LabelChange.Current;
+				newLabel.RequiredNodes = label.RequiredNodes.Select(x => x.Name).ToList();
+				newLabel.IncludedNodes = label.IncludedNodes.Select(x => x.Name).ToList();
+				newLabels.Add(newLabel);
 			}
 
-			return await GraphCollection.AppendAsync(null, NewGroups, NewAggregates, NewLabels);
+			return await graphCollection.AppendAsync(null, newGroups, newAggregates, newLabels);
 		}
 	}
 }
diff --git a/Engine/Source/Programs/Horde/Horde.Build/Utilities/CaseInsensitiveStringSet.cs b/Engine/Source/Programs/Horde/Horde.Build/Utilities/CaseInsensitiveStringSet.cs
--- a/Engine/Source/Programs/Horde/Horde.Build/Utilities/CaseInsensitiveStringSet.cs
+++ b/Engine/Source/Programs/Horde/Horde.Build/Utilities/CaseInsensitiveStringSet.cs
@@ -1,11 +1,7 @@
 // Copyright Epic Games, Inc. All Rights Reserved.
 
-using MongoDB.Bson.Serialization;
-using MongoDB.Bson.Serialization.Serializers;
 using System;
 using System.Collections.Generic;
-using System.Linq;
-using System.Threading.Tasks;
 
 namespace Horde.Build.Utilities
 {
@@ -25,8 +21,8 @@ namespace Horde.Build.Utilities
 		/// <summary>
 		/// Constructor
 		/// </summary>
-		public CaseInsensitiveStringSet(IEnumerable<string> Items)
-			: base(Items, StringComparer.OrdinalIgnoreCase)
+		public CaseInsensitiveStringSet(IEnumerable<string> items)
+			: base(items, StringComparer.OrdinalIgnoreCase)
 		{
 		}
 	}
diff --git a/Engine/Source/Programs/Horde/Horde.Build/Utilities/Clock.cs b/Engine/Source/Programs/Horde/Horde.Build/Utilities/Clock.cs
--- a/Engine/Source/Programs/Horde/Horde.Build/Utilities/Clock.cs
+++ b/Engine/Source/Programs/Horde/Horde.Build/Utilities/Clock.cs
@@ -1,15 +1,14 @@
 // Copyright Epic Games, Inc. All Rights Reserved.
 
-using EpicGames.Core;
-using EpicGames.Redis.Utility;
-using Horde.Build.Services;
-using Microsoft.Extensions.Logging;
-using StackExchange.Redis;
 using System;
 using System.Collections.Generic;
 using System.Diagnostics;
 using System.Threading;
 using System.Threading.Tasks;
+using EpicGames.Redis.Utility;
+using Horde.Build.Services;
+using Microsoft.Extensions.Logging;
+using StackExchange.Redis;
 
 namespace HordeCommon
 {
@@ -58,21 +57,22 @@ namespace HordeCommon
 		/// <summary>
 		/// Create an event that will trigger after the given time
 		/// </summary>
-		/// <param name="Interval">Time after which the event will trigger</param>
-		/// <param name="TickAsync">Callback for the tick. Returns the time interval until the next tick, or null to cancel the tick.</param>
-		/// <param name="Logger">Logger for error messages</param>
+		/// <param name="name">Name of the event</param>
+		/// <param name="interval">Time after which the event will trigger</param>
+		/// <param name="tickAsync">Callback for the tick. Returns the time interval until the next tick, or null to cancel the tick.</param>
+		/// <param name="logger">Logger for error messages</param>
 		/// <returns>Handle to the event</returns>
-		ITicker AddTicker(TimeSpan Interval, Func<CancellationToken, ValueTask<TimeSpan?>> TickAsync, ILogger Logger);
+		ITicker AddTicker(string name, TimeSpan interval, Func<CancellationToken, ValueTask<TimeSpan?>> tickAsync, ILogger logger);
 
 		/// <summary>
 		/// Create a ticker shared between all server pods
 		/// </summary>
-		/// <param name="Name">Name of the event</param>
-		/// <param name="Interval">Time after which the event will trigger</param>
-		/// <param name="TickAsync">Callback for the tick. Returns the time interval until the next tick, or null to cancel the tick.</param>
-		/// <param name="Logger">Logger for error messages</param>
+		/// <param name="name">Name of the event</param>
+		/// <param name="interval">Time after which the event will trigger</param>
+		/// <param name="tickAsync">Callback for the tick. Returns the time interval until the next tick, or null to cancel the tick.</param>
+		/// <param name="logger">Logger for error messages</param>
 		/// <returns>New ticker instance</returns>
-		ITicker AddSharedTicker(string Name, TimeSpan Interval, Func<CancellationToken, ValueTask> TickAsync, ILogger Logger);
+		ITicker AddSharedTicker(string name, TimeSpan interval, Func<CancellationToken, ValueTask> tickAsync, ILogger logger);
 	}
 
 	/// <summary>
@@ -83,31 +83,53 @@ namespace HordeCommon
 		/// <summary>
 		/// Create an event that will trigger after the given time
 		/// </summary>
-		/// <param name="Clock">Clock to schedule the event on</param>
-		/// <param name="Interval">Interval for the callback</param>
-		/// <param name="TickAsync">Trigger callback</param>
-		/// <param name="Logger">Logger for any error messages</param>
+		/// <param name="clock">Clock to schedule the event on</param>
+		/// <param name="name">Name of the ticker</param>
+		/// <param name="interval">Interval for the callback</param>
+		/// <param name="tickAsync">Trigger callback</param>
+		/// <param name="logger">Logger for any error messages</param>
 		/// <returns>Handle to the event</returns>
-		public static ITicker AddTicker(this IClock Clock, TimeSpan Interval, Func<CancellationToken, ValueTask> TickAsync, ILogger Logger)
+		public static ITicker AddTicker(this IClock clock, string name, TimeSpan interval, Func<CancellationToken, ValueTask> tickAsync, ILogger logger)
 		{
-			Func<CancellationToken, ValueTask<TimeSpan?>> WrappedTrigger = async Token =>
+			async ValueTask<TimeSpan?> WrappedTrigger(CancellationToken token)
 			{
-				Stopwatch Timer = Stopwatch.StartNew();
-				await TickAsync(Token);
-				return Interval - Timer.Elapsed;
-			};
-			return Clock.AddTicker(Interval, WrappedTrigger, Logger);
+				Stopwatch timer = Stopwatch.StartNew();
+				await tickAsync(token);
+				return interval - timer.Elapsed;
+			}
+
+			return clock.AddTicker(name, interval, WrappedTrigger, logger);
 		}
 
+		/// <summary>
+		/// Create an event that will trigger after the given time
+		/// </summary>
+		/// <param name="clock">Clock to schedule the event on</param>
+		/// <param name="interval">Time after which the event will trigger</param>
+		/// <param name="tickAsync">Callback for the tick. Returns the time interval until the next tick, or null to cancel the tick.</param>
+		/// <param name="logger">Logger for error messages</param>
+		/// <returns>Handle to the event</returns>
+		public static ITicker AddTicker<T>(this IClock clock, TimeSpan interval, Func<CancellationToken, ValueTask<TimeSpan?>> tickAsync, ILogger logger) => clock.AddTicker(typeof(T).Name, interval, tickAsync, logger);
+
+		/// <summary>
+		/// Create an event that will trigger after the given time
+		/// </summary>
+		/// <param name="clock">Clock to schedule the event on</param>
+		/// <param name="interval">Interval for the callback</param>
+		/// <param name="tickAsync">Trigger callback</param>
+		/// <param name="logger">Logger for any error messages</param>
+		/// <returns>Handle to the event</returns>
+		public static ITicker AddTicker<T>(this IClock clock, TimeSpan interval, Func<CancellationToken, ValueTask> tickAsync, ILogger logger) => clock.AddTicker(typeof(T).Name, interval, tickAsync, logger);
+
 		/// <summary>
 		/// Create a ticker shared between all server pods
 		/// </summary>
-		/// <param name="Clock">Clock to schedule the event on</param>
-		/// <param name="Interval">Time after which the event will trigger</param>
-		/// <param name="TickAsync">Callback for the tick. Returns the time interval until the next tick, or null to cancel the tick.</param>
-		/// <param name="Logger">Logger for error messages</param>
+		/// <param name="clock">Clock to schedule the event on</param>
+		/// <param name="interval">Time after which the event will trigger</param>
+		/// <param name="tickAsync">Callback for the tick. Returns the time interval until the next tick, or null to cancel the tick.</param>
+		/// <param name="logger">Logger for error messages</param>
 		/// <returns>New ticker instance</returns>
-		public static ITicker AddSharedTicker<T>(this IClock Clock, TimeSpan Interval, Func<CancellationToken, ValueTask> TickAsync, ILogger Logger) => Clock.AddSharedTicker(typeof(T).Name, Interval, TickAsync, Logger);
+		public static ITicker AddSharedTicker<T>(this IClock clock, TimeSpan interval, Func<CancellationToken, ValueTask> tickAsync, ILogger logger) => clock.AddSharedTicker(typeof(T).Name, interval, tickAsync, logger);
 	}
 
 	/// <summary>
@@ -117,79 +139,80 @@ namespace HordeCommon
 	{
 		class TickerImpl : ITicker
 		{
-			[System.Diagnostics.CodeAnalysis.SuppressMessage("Usage", "CA2213:Disposable fields should be disposed", Justification = "<Pending>")]
-			CancellationTokenSource? CancellationSource;
-			Func<Task> TickFunc;
-			Task? BackgroundTask;
+			readonly string _name;
+			CancellationTokenSource? _cancellationSource;
+			readonly Func<Task> _tickFunc;
+			Task? _backgroundTask;
 
-			public TickerImpl(TimeSpan Delay, Func<CancellationToken, ValueTask<TimeSpan?>> TriggerAsync, ILogger Logger)
+			public TickerImpl(string name, TimeSpan delay, Func<CancellationToken, ValueTask<TimeSpan?>> triggerAsync, ILogger logger)
 			{
-				TickFunc = () => Run(Delay, TriggerAsync, Logger);
+				_name = name;
+				_tickFunc = () => Run(delay, triggerAsync, logger);
 			}
 
 			public async Task StartAsync()
 			{
 				await StopAsync();
-				CancellationSource = new CancellationTokenSource();
-				BackgroundTask = Task.Run(TickFunc);
+				_cancellationSource = new CancellationTokenSource();
+				_backgroundTask = Task.Run(_tickFunc);
 			}
 
 			public async Task StopAsync()
 			{
-				if (CancellationSource != null)
+				if (_cancellationSource != null)
 				{
-					CancellationSource.Cancel();
-					await BackgroundTask!;
-					CancellationSource.Dispose();
-					CancellationSource = null;
+					_cancellationSource.Cancel();
+					await _backgroundTask!;
+					_cancellationSource.Dispose();
+					_cancellationSource = null;
 				}
 			}
 
 			public void Dispose()
 			{
 				StopAsync().Wait();
-				if (CancellationSource != null)
+				if (_cancellationSource != null)
 				{
-					CancellationSource.Dispose();
+					_cancellationSource.Dispose();
 				}
 			}
 
-			public async Task Run(TimeSpan Delay, Func<CancellationToken, ValueTask<TimeSpan?>> TriggerAsync, ILogger Logger)
+			public async Task Run(TimeSpan delay, Func<CancellationToken, ValueTask<TimeSpan?>> triggerAsync, ILogger logger)
 			{
-				while (!CancellationSource!.IsCancellationRequested)
+				while (!_cancellationSource!.IsCancellationRequested)
 				{
 					try
 					{
-						if (Delay > TimeSpan.Zero)
+						if (delay > TimeSpan.Zero)
 						{
-							await Task.Delay(Delay, CancellationSource.Token);
+							await Task.Delay(delay, _cancellationSource.Token);
 						}
 
-						TimeSpan? NextDelay = await TriggerAsync(CancellationSource.Token);
-						if(NextDelay == null)
+						TimeSpan? nextDelay = await triggerAsync(_cancellationSource.Token);
+						if(nextDelay == null)
 						{
 							break;
 						}
 
-						Delay = NextDelay.Value;
+						delay = nextDelay.Value;
 					}
-					catch (OperationCanceledException) when (CancellationSource.IsCancellationRequested)
+					catch (OperationCanceledException) when (_cancellationSource.IsCancellationRequested)
 					{
 					}
-					catch (Exception Ex)
+					catch (Exception ex)
 					{
-						Logger.LogError(Ex, "Exception while executing scheduled event");
-						if (Delay < TimeSpan.Zero)
+						logger.LogError(ex, "Exception while executing scheduled event");
+						if (delay < TimeSpan.Zero)
 						{
-							Delay = TimeSpan.FromSeconds(5.0);
-							Logger.LogWarning("Delaying tick for 5 seconds");
+							delay = TimeSpan.FromSeconds(5.0);
+							logger.LogWarning("Delaying tick for 5 seconds");
 						}
 					}
 				}
 			}
 		}
 
-		RedisService Redis;
+		readonly RedisService _redis;
 
 		/// <inheritdoc/>
 		public DateTime UtcNow => DateTime.UtcNow;
@@ -197,31 +220,31 @@ namespace HordeCommon
 		/// <summary>
 		/// Constructor
 		/// </summary>
-		public Clock(RedisService Redis)
+		public Clock(RedisService redis)
 		{
-			this.Redis = Redis;
+			_redis = redis;
 		}
 
 		/// <inheritdoc/>
-		public ITicker AddTicker(TimeSpan Delay, Func<CancellationToken, ValueTask<TimeSpan?>> TickAsync, ILogger Logger)
+		public ITicker AddTicker(string name, TimeSpan delay, Func<CancellationToken, ValueTask<TimeSpan?>> tickAsync, ILogger logger)
 		{
-			return new TickerImpl(Delay, TickAsync, Logger);
+			return new TickerImpl(name, delay, tickAsync, logger);
 		}
 
 		/// <inheritdoc/>
-		public ITicker AddSharedTicker(string Name, TimeSpan Delay, Func<CancellationToken, ValueTask> TickAsync, ILogger Logger)
+		public ITicker AddSharedTicker(string name, TimeSpan delay, Func<CancellationToken, ValueTask> tickAsync, ILogger logger)
 		{
-			RedisKey Key = new RedisKey($"tick/{Name}");
-			return ClockExtensions.AddTicker(this, Delay / 4, Token => TriggerSharedAsync(Key, Delay, TickAsync, Token), Logger);
+			RedisKey key = new RedisKey($"tick/{name}");
+			return ClockExtensions.AddTicker(this, name, delay / 4, token => TriggerSharedAsync(key, delay, tickAsync, token), logger);
 		}
 
-		async ValueTask TriggerSharedAsync(RedisKey Key, TimeSpan Interval, Func<CancellationToken, ValueTask> TickAsync, CancellationToken CancellationToken)
+		async ValueTask TriggerSharedAsync(RedisKey key, TimeSpan interval, Func<CancellationToken, ValueTask> tickAsync, CancellationToken cancellationToken)
 		{
-			using (RedisLock Lock = new RedisLock(Redis.Database, Key))
+			using (RedisLock sharedLock = new RedisLock(_redis.Database, key))
 			{
-				if (await Lock.AcquireAsync(Interval, false))
+				if (await sharedLock.AcquireAsync(interval, false))
 				{
-					await TickAsync(CancellationToken);
+					await tickAsync(cancellationToken);
 				}
 			}
 		}
@@ -235,26 +258,28 @@ namespace HordeCommon
 	{
 		class TickerImpl : ITicker
 		{
-			FakeClock Outer { get; }
-			TimeSpan Interval;
+			readonly FakeClock _outer;
+			readonly string _name;
+			readonly TimeSpan _interval;
 			public DateTime? NextTime { get; set; }
 			public Func<CancellationToken, ValueTask<TimeSpan?>> TickAsync { get; }
 
-			public TickerImpl(FakeClock Outer, TimeSpan Interval, Func<CancellationToken, ValueTask<TimeSpan?>> TickAsync)
+			public TickerImpl(FakeClock outer, string name, TimeSpan interval, Func<CancellationToken, ValueTask<TimeSpan?>> tickAsync)
 			{
-				this.Outer = Outer;
-				this.Interval = Interval;
-				this.TickAsync = TickAsync;
+				_outer = outer;
+				_name = name;
+				_interval = interval;
+				TickAsync = tickAsync;
 
-				lock (Outer.Triggers)
+				lock (outer._triggers)
 				{
-					Outer.Triggers.Add(this);
+					outer._triggers.Add(this);
 				}
 			}
 
 			public Task StartAsync()
 			{
-				NextTime = Outer.UtcNow + Interval;
+				NextTime = _outer.UtcNow + _interval;
 				return Task.CompletedTask;
 			}
 
@@ -271,46 +296,58 @@ namespace HordeCommon
 
 			public ValueTask DisposeAsync()
 			{
-				lock (Outer.Triggers)
+				lock (_outer._triggers)
 				{
-					Outer.Triggers.Remove(this);
+					_outer._triggers.Remove(this);
 				}
 				return new ValueTask();
 			}
+
+			public override string ToString()
+			{
+				if (NextTime == null)
+				{
+					return $"{_name} (paused)";
+				}
+				else
+				{
+					return $"{_name} ({NextTime.Value})";
+				}
+			}
 		}
 
-		DateTime UtcNowPrivate;
-		List<TickerImpl> Triggers = new List<TickerImpl>();
+		DateTime _utcNowPrivate;
+		readonly List<TickerImpl> _triggers = new List<TickerImpl>();
 
 		/// <summary>
 		/// Constructor
 		/// </summary>
 		public FakeClock()
 		{
-			UtcNowPrivate = DateTime.UtcNow;
+			_utcNowPrivate = DateTime.UtcNow;
 		}
 
 		/// <summary>
 		/// Advance time by given amount
 		/// Useful for letting time progress during tests
 		/// </summary>
-		/// <param name="Period">Time span to advance</param>
-		public async Task AdvanceAsync(TimeSpan Period)
+		/// <param name="period">Time span to advance</param>
+		public async Task AdvanceAsync(TimeSpan period)
 		{
-			UtcNowPrivate = UtcNowPrivate.Add(Period);
+			_utcNowPrivate = _utcNowPrivate.Add(period);
 
-			for (int Idx = 0; Idx < Triggers.Count; Idx++)
+			for (int idx = 0; idx < _triggers.Count; idx++)
 			{
-				TickerImpl Trigger = Triggers[Idx];
-				while (Trigger.NextTime != null && UtcNowPrivate > Trigger.NextTime)
+				TickerImpl trigger = _triggers[idx];
+				while (trigger.NextTime != null && _utcNowPrivate > trigger.NextTime)
 				{
-					TimeSpan? Delay = await Trigger.TickAsync(CancellationToken.None);
-					if (Delay == null)
+					TimeSpan? delay = await trigger.TickAsync(CancellationToken.None);
+					if (delay == null)
 					{
-						Triggers.RemoveAt(Idx--);
+						_triggers.RemoveAt(idx--);
 						break;
 					}
-					Trigger.NextTime = UtcNowPrivate + Delay.Value;
+					trigger.NextTime = _utcNowPrivate + delay.Value;
 				}
 			}
 		}
@@ -318,20 +355,25 @@ namespace HordeCommon
 		/// <inheritdoc/>
 		public DateTime UtcNow
 		{ 
-			get => UtcNowPrivate;
-			set => UtcNowPrivate = value.ToUniversalTime(); 
+			get => _utcNowPrivate;
+			set => _utcNowPrivate = value.ToUniversalTime(); 
 		}
 
 		/// <inheritdoc/>
-		public ITicker AddTicker(TimeSpan Interval, Func<CancellationToken, ValueTask<TimeSpan?>> TickAsync, ILogger Logger)
+		public ITicker AddTicker(string name, TimeSpan interval, Func<CancellationToken, ValueTask<TimeSpan?>> tickAsync, ILogger logger)
 		{
-			return new TickerImpl(this, Interval, TickAsync);
+			return new TickerImpl(this, name, interval, tickAsync);
 		}
 
 		/// <inheritdoc/>
-		public ITicker AddSharedTicker(string Name, TimeSpan Interval, Func<CancellationToken, ValueTask> TickAsync, ILogger Logger)
+		public ITicker AddSharedTicker(string name, TimeSpan interval, Func<CancellationToken, ValueTask> tickAsync, ILogger logger)
 		{
-			return AddTicker(Interval, async Token => { await TickAsync(Token); return Interval; }, Logger);
+			async ValueTask<TimeSpan?> Tick(CancellationToken token)
+			{
+				await tickAsync(token);
+				return interval;
+			}
+			return new TickerImpl(this, name, interval, Tick);
 		}
 	}
 }
\ No newline at end of file
diff --git a/Engine/Source/Programs/Horde/Horde.Build/Utilities/CollectionExtensions.cs b/Engine/Source/Programs/Horde/Horde.Build/Utilities/CollectionExtensions.cs
--- a/Engine/Source/Programs/Horde/Horde.Build/Utilities/CollectionExtensions.cs
+++ b/Engine/Source/Programs/Horde/Horde.Build/Utilities/CollectionExtensions.cs
@@ -1,12 +1,6 @@
 // Copyright Epic Games, Inc. All Rights Reserved.
 
-using Microsoft.AspNetCore.Http.Features;
-using Microsoft.AspNetCore.Mvc.Formatters;
-using Microsoft.IdentityModel.Tokens;
-using System;
 using System.Collections.Generic;
-using System.Linq;
-using System.Threading.Tasks;
 
 namespace Horde.Build.Utilities
 {
@@ -20,13 +14,13 @@ namespace Horde.Build.Utilities
 		/// </summary>
 		/// <typeparam name="TKey">The key type</typeparam>
 		/// <typeparam name="TValue">The value type</typeparam>
-		/// <param name="Map">The map to update</param>
-		/// <param name="Sequence">Sequence of items to add</param>
-		public static void Add<TKey, TValue>(this Google.Protobuf.Collections.MapField<TKey, TValue> Map, IEnumerable<KeyValuePair<TKey, TValue>> Sequence)
+		/// <param name="map">The map to update</param>
+		/// <param name="sequence">Sequence of items to add</param>
+		public static void Add<TKey, TValue>(this Google.Protobuf.Collections.MapField<TKey, TValue> map, IEnumerable<KeyValuePair<TKey, TValue>> sequence)
 		{
-			foreach(KeyValuePair<TKey, TValue> Pair in Sequence)
+			foreach(KeyValuePair<TKey, TValue> pair in sequence)
 			{
-				Map.Add(Pair.Key, Pair.Value);
+				map.Add(pair.Key, pair.Value);
 			}
 		}
 	}
diff --git a/Engine/Source/Programs/Horde/Horde.Build/Utilities/CompressionExtensions.cs b/Engine/Source/Programs/Horde/Horde.Build/Utilities/CompressionExtensions.cs
--- a/Engine/Source/Programs/Horde/Horde.Build/Utilities/CompressionExtensions.cs
+++ b/Engine/Source/Programs/Horde/Horde.Build/Utilities/CompressionExtensions.cs
@@ -1,16 +1,13 @@
 // Copyright Epic Games, Inc. All Rights Reserved.
 
-using EpicGames.Core;
-using ICSharpCode.SharpZipLib.BZip2;
-using OpenTracing;
-using OpenTracing.Util;
 using System;
 using System.Buffers.Binary;
-using System.Collections.Generic;
 using System.Globalization;
 using System.IO;
-using System.Linq;
-using System.Threading.Tasks;
+using EpicGames.Core;
+using ICSharpCode.SharpZipLib.BZip2;
+using OpenTracing;
+using OpenTracing.Util;
 
 namespace Horde.Build.Utilities
 {
@@ -23,55 +20,55 @@ namespace Horde.Build.Utilities
 		/// Compress a block of data with bzip2
 		/// </summary>
 		/// <returns>The compressed data</returns>
-		public static byte[] CompressBzip2(this ReadOnlyMemory<byte> Memory)
+		public static byte[] CompressBzip2(this ReadOnlyMemory<byte> memory)
 		{
-			using IScope Scope = GlobalTracer.Instance.BuildSpan("CompressBzip2").StartActive();
-			Scope.Span.SetTag("DecompressedSize", Memory.Length.ToString(CultureInfo.InvariantCulture));
+			using IScope scope = GlobalTracer.Instance.BuildSpan("CompressBzip2").StartActive();
+			scope.Span.SetTag("DecompressedSize", memory.Length.ToString(CultureInfo.InvariantCulture));
 
-			byte[] CompressedData;
-			using (MemoryStream Stream = new MemoryStream())
+			byte[] compressedData;
+			using (MemoryStream stream = new MemoryStream())
 			{
-				byte[] DecompressedSize = new byte[4];
-				BinaryPrimitives.WriteInt32LittleEndian(DecompressedSize.AsSpan(), Memory.Length);
-				Stream.Write(DecompressedSize.AsSpan());
+				byte[] decompressedSize = new byte[4];
+				BinaryPrimitives.WriteInt32LittleEndian(decompressedSize.AsSpan(), memory.Length);
+				stream.Write(decompressedSize.AsSpan());
 
-				using (BZip2OutputStream CompressedStream = new BZip2OutputStream(Stream))
+				using (BZip2OutputStream compressedStream = new BZip2OutputStream(stream))
 				{
-					CompressedStream.Write(Memory.Span);
+					compressedStream.Write(memory.Span);
 				}
 
-				CompressedData = Stream.ToArray();
+				compressedData = stream.ToArray();
 			}
 
-			Scope.Span.SetTag("CompressedSize", CompressedData.Length.ToString(CultureInfo.InvariantCulture));
-			return CompressedData;
+			scope.Span.SetTag("CompressedSize", compressedData.Length.ToString(CultureInfo.InvariantCulture));
+			return compressedData;
 		}
 
 		/// <summary>
 		/// Decompress the data
 		/// </summary>
 		/// <returns>The decompressed data</returns>
-		public static byte[] DecompressBzip2(this ReadOnlyMemory<byte> Memory)
+		public static byte[] DecompressBzip2(this ReadOnlyMemory<byte> memory)
 		{
-			int DecompressedSize = BinaryPrimitives.ReadInt32LittleEndian(Memory.Span);
+			int decompressedSize = BinaryPrimitives.ReadInt32LittleEndian(memory.Span);
 
-			using IScope Scope = GlobalTracer.Instance.BuildSpan("DecompressBzip2").StartActive();
-			Scope.Span.SetTag("CompressedSize", Memory.Length.ToString(CultureInfo.InvariantCulture));
-			Scope.Span.SetTag("DecompressedSize", DecompressedSize.ToString(CultureInfo.InvariantCulture));
+			using IScope scope = GlobalTracer.Instance.BuildSpan("DecompressBzip2").StartActive();
+			scope.Span.SetTag("CompressedSize", memory.Length.ToString(CultureInfo.InvariantCulture));
+			scope.Span.SetTag("DecompressedSize", decompressedSize.ToString(CultureInfo.InvariantCulture));
 
-			byte[] Data = new byte[DecompressedSize];
-			using (ReadOnlyMemoryStream Stream = new ReadOnlyMemoryStream(Memory.Slice(4)))
+			byte[] data = new byte[decompressedSize];
+			using (ReadOnlyMemoryStream stream = new ReadOnlyMemoryStream(memory.Slice(4)))
 			{
-				using (BZip2InputStream DecompressedStream = new BZip2InputStream(Stream))
+				using (BZip2InputStream decompressedStream = new BZip2InputStream(stream))
 				{
-					int ReadSize = DecompressedStream.Read(Data.AsSpan());
-					if (ReadSize != Data.Length)
+					int readSize = decompressedStream.Read(data.AsSpan());
+					if (readSize != data.Length)
 					{
-						throw new InvalidDataException($"Compressed data is too short (expected {Data.Length} bytes, got {ReadSize} bytes)");
+						throw new InvalidDataException($"Compressed data is too short (expected {data.Length} bytes, got {readSize} bytes)");
 					}
 				}
 			}
-			return Data;
+			return data;
 		}
 	}
 }
diff --git a/Engine/Source/Programs/Horde/Horde.Build/Utilities/ContentHashSerializer.cs b/Engine/Source/Programs/Horde/Horde.Build/Utilities/ContentHashSerializer.cs
--- a/Engine/Source/Programs/Horde/Horde.Build/Utilities/ContentHashSerializer.cs
+++ b/Engine/Source/Programs/Horde/Horde.Build/Utilities/ContentHashSerializer.cs
@@ -1,11 +1,8 @@
 // Copyright Epic Games, Inc. All Rights Reserved.
 
+using System;
 using EpicGames.Core;
 using MongoDB.Bson.Serialization;
-using System;
-using System.Collections.Generic;
-using System.Linq;
-using System.Threading.Tasks;
 
 namespace Horde.Build.Utilities
 {
@@ -15,40 +12,37 @@ namespace Horde.Build.Utilities
 	public sealed class ContentHashSerializer : IBsonSerializer<ContentHash>
 	{
 		/// <inheritdoc/>
-		public Type ValueType
-		{
-			get { return typeof(ContentHash); }
-		}
+		public Type ValueType => typeof(ContentHash);
 
 		/// <inheritdoc/>
-		void IBsonSerializer.Serialize(BsonSerializationContext Context, BsonSerializationArgs Args, object Value)
+		void IBsonSerializer.Serialize(BsonSerializationContext context, BsonSerializationArgs args, object value)
 		{
-			Serialize(Context, Args, (ContentHash)Value);
+			Serialize(context, args, (ContentHash)value);
 		}
 
 		/// <inheritdoc/>
-		object IBsonSerializer.Deserialize(BsonDeserializationContext Context, BsonDeserializationArgs Args)
+		object IBsonSerializer.Deserialize(BsonDeserializationContext context, BsonDeserializationArgs args)
 		{
-			return ((IBsonSerializer<ContentHash>)this).Deserialize(Context, Args);
+			return ((IBsonSerializer<ContentHash>)this).Deserialize(context, args);
 		}
 
 		/// <inheritdoc/>
-		public ContentHash Deserialize(BsonDeserializationContext Context, BsonDeserializationArgs Args)
+		public ContentHash Deserialize(BsonDeserializationContext context, BsonDeserializationArgs args)
 		{
-			if (Context.Reader.CurrentBsonType == MongoDB.Bson.BsonType.ObjectId)
+			if (context.Reader.CurrentBsonType == MongoDB.Bson.BsonType.ObjectId)
 			{
-				return new ContentHash(Context.Reader.ReadObjectId().ToByteArray());
+				return new ContentHash(context.Reader.ReadObjectId().ToByteArray());
 			}
 			else
 			{
-				return ContentHash.Parse(Context.Reader.ReadString());
+				return ContentHash.Parse(context.Reader.ReadString());
 			}
 		}
 
 		/// <inheritdoc/>
-		public void Serialize(BsonSerializationContext Context, BsonSerializationArgs Args, ContentHash Value)
+		public void Serialize(BsonSerializationContext context, BsonSerializationArgs args, ContentHash value)
 		{
-			Context.Writer.WriteString(Value.ToString());
+			context.Writer.WriteString(value.ToString());
 		}
 	}
 }
