Horde: Fix coding conventions within Horde.Build.

diff --git a/Engine/Source/Programs/Horde/Horde.Build/Utilities/CustomFileCallbackResult.cs b/Engine/Source/Programs/Horde/Horde.Build/Utilities/CustomFileCallbackResult.cs
--- a/Engine/Source/Programs/Horde/Horde.Build/Utilities/CustomFileCallbackResult.cs
+++ b/Engine/Source/Programs/Horde/Horde.Build/Utilities/CustomFileCallbackResult.cs
@@ -1,58 +1,54 @@
 // Copyright Epic Games, Inc. All Rights Reserved.
 
-using Microsoft.AspNetCore.Mvc;
-using Microsoft.AspNetCore.Mvc.Infrastructure;
-using Microsoft.Extensions.DependencyInjection;
-using Microsoft.Extensions.Logging;
 using System;
-using System.Collections.Generic;
 using System.IO;
-using System.Linq;
-using System.Net.Http.Headers;
 using System.Net.Mime;
-using System.Threading;
 using System.Threading.Tasks;
+using Microsoft.AspNetCore.Mvc;
+using Microsoft.AspNetCore.Mvc.Infrastructure;
+using Microsoft.Extensions.DependencyInjection;
+using Microsoft.Extensions.Logging;
 
 namespace Horde.Build.Utilities
 {
-    /// <summary>
-    /// Class deriving from a FileResult that allows custom file types (used for zip file creation)
-    /// </summary>
-    public class CustomFileCallbackResult : FileResult
+	/// <summary>
+	/// Class deriving from a FileResult that allows custom file types (used for zip file creation)
+	/// </summary>
+	public class CustomFileCallbackResult : FileResult
     {
-        private Func<Stream, ActionContext, Task> Callback;
-		string FileName;
-		bool Inline;
+        private readonly Func<Stream, ActionContext, Task> _callback;
+		readonly string _fileName;
+		readonly bool _inline;
 
 		/// <summary>
 		/// Constructor
 		/// </summary>
-		/// <param name="FileName">Default filename for the downloaded file</param>
-		/// <param name="MimeType">Content type for the file</param>
-		/// <param name="Inline">Whether to display the file inline in the browser</param>
-		/// <param name="Callback">Callback used to write the data</param>
-		public CustomFileCallbackResult(string FileName, string MimeType, bool Inline, Func<Stream, ActionContext, Task> Callback)
-            : base(MimeType)
+		/// <param name="fileName">Default filename for the downloaded file</param>
+		/// <param name="mimeType">Content type for the file</param>
+		/// <param name="inline">Whether to display the file inline in the browser</param>
+		/// <param name="callback">Callback used to write the data</param>
+		public CustomFileCallbackResult(string fileName, string mimeType, bool inline, Func<Stream, ActionContext, Task> callback)
+            : base(mimeType)
         {
-			this.FileName = FileName;
-			this.Inline = Inline;
-            this.Callback = Callback;
+			_fileName = fileName;
+			_inline = inline;
+            _callback = callback;
         }
 
         /// <summary>
         /// Executes the action result
         /// </summary>
-        /// <param name="Context">The controller context</param>
+        /// <param name="context">The controller context</param>
         /// <returns></returns>
-        public override Task ExecuteResultAsync(ActionContext Context)
+        public override Task ExecuteResultAsync(ActionContext context)
         {
-			ContentDisposition ContentDisposition = new ContentDisposition();
-			ContentDisposition.Inline = Inline;
-			ContentDisposition.FileName = FileName;
-			Context.HttpContext.Response.Headers.Add("Content-Disposition", ContentDisposition.ToString());
+			ContentDisposition contentDisposition = new ContentDisposition();
+			contentDisposition.Inline = _inline;
+			contentDisposition.FileName = _fileName;
+			context.HttpContext.Response.Headers.Add("Content-Disposition", contentDisposition.ToString());
 
-			CustomFileCallbackResultExecutor Executor = new CustomFileCallbackResultExecutor(Context.HttpContext.RequestServices.GetRequiredService<ILoggerFactory>());
-            return Executor.ExecuteAsync(Context, this);
+			CustomFileCallbackResultExecutor executor = new CustomFileCallbackResultExecutor(context.HttpContext.RequestServices.GetRequiredService<ILoggerFactory>());
+            return executor.ExecuteAsync(context, this);
         }
 
         /// <summary>
@@ -63,22 +59,22 @@ namespace Horde.Build.Utilities
             /// <summary>
             /// Constructor
             /// </summary>
-            /// <param name="LoggerFactory">The logger</param>
-            public CustomFileCallbackResultExecutor(ILoggerFactory LoggerFactory)
-                : base(CreateLogger<CustomFileCallbackResultExecutor>(LoggerFactory))
+            /// <param name="loggerFactory">The logger</param>
+            public CustomFileCallbackResultExecutor(ILoggerFactory loggerFactory)
+                : base(CreateLogger<CustomFileCallbackResultExecutor>(loggerFactory))
             {
             }
 
             /// <summary>
             /// Executes a CustomFileResult callback
             /// </summary>
-            /// <param name="Context">The controller context</param>
-            /// <param name="Result">The custom file result</param>
+            /// <param name="context">The controller context</param>
+            /// <param name="result">The custom file result</param>
             /// <returns></returns>
-            public Task ExecuteAsync(ActionContext Context, CustomFileCallbackResult Result)
+            public Task ExecuteAsync(ActionContext context, CustomFileCallbackResult result)
             {
-                SetHeadersAndLog(Context, Result, null, false);
-                return Result.Callback(Context.HttpContext.Response.Body, Context);
+                SetHeadersAndLog(context, result, null, false);
+                return result._callback(context.HttpContext.Response.Body, context);
             }
         }
     }
diff --git a/Engine/Source/Programs/Horde/Horde.Build/Utilities/DateTimeOffsetStringSerializer.cs b/Engine/Source/Programs/Horde/Horde.Build/Utilities/DateTimeOffsetStringSerializer.cs
--- a/Engine/Source/Programs/Horde/Horde.Build/Utilities/DateTimeOffsetStringSerializer.cs
+++ b/Engine/Source/Programs/Horde/Horde.Build/Utilities/DateTimeOffsetStringSerializer.cs
@@ -1,14 +1,10 @@
 // Copyright Epic Games, Inc. All Rights Reserved.
 
+using System;
+using System.Globalization;
 using MongoDB.Bson;
 using MongoDB.Bson.Serialization;
 using MongoDB.Bson.Serialization.Serializers;
-using Newtonsoft.Json.Bson;
-using System;
-using System.Collections.Generic;
-using System.Globalization;
-using System.Linq;
-using System.Threading.Tasks;
 
 namespace Horde.Build.Utilities
 {
@@ -18,25 +14,25 @@ namespace Horde.Build.Utilities
 	class DateTimeOffsetStringSerializer : StructSerializerBase<DateTimeOffset>
 	{
 		/// <inheritdoc/>
-		public override DateTimeOffset Deserialize(BsonDeserializationContext Context, BsonDeserializationArgs Args)
+		public override DateTimeOffset Deserialize(BsonDeserializationContext context, BsonDeserializationArgs args)
 		{
-			switch (Context.Reader.CurrentBsonType)
+			switch (context.Reader.CurrentBsonType)
 			{
 				case BsonType.String:
-					string StringValue = Context.Reader.ReadString();
-					return DateTimeOffset.Parse(StringValue, DateTimeFormatInfo.InvariantInfo);
+					string stringValue = context.Reader.ReadString();
+					return DateTimeOffset.Parse(stringValue, DateTimeFormatInfo.InvariantInfo);
 				case BsonType.DateTime:
-					long Ticks = Context.Reader.ReadDateTime();
-					return DateTimeOffset.FromUnixTimeMilliseconds(Ticks);
+					long ticks = context.Reader.ReadDateTime();
+					return DateTimeOffset.FromUnixTimeMilliseconds(ticks);
 				default:
-					throw new FormatException($"Unable to deserialize a DateTimeOffset from a {Context.Reader.CurrentBsonType}");
+					throw new FormatException($"Unable to deserialize a DateTimeOffset from a {context.Reader.CurrentBsonType}");
 			}
 		}
 
 		/// <inheritdoc/>
-		public override void Serialize(BsonSerializationContext Context, BsonSerializationArgs Args, DateTimeOffset Value)
+		public override void Serialize(BsonSerializationContext context, BsonSerializationArgs args, DateTimeOffset value)
 		{
-			Context.Writer.WriteString(Value.ToString(DateTimeFormatInfo.InvariantInfo));
+			context.Writer.WriteString(value.ToString(DateTimeFormatInfo.InvariantInfo));
 		}
 	}
 }
diff --git a/Engine/Source/Programs/Horde/Horde.Build/Utilities/DefaultDiscriminatorConvention.cs b/Engine/Source/Programs/Horde/Horde.Build/Utilities/DefaultDiscriminatorConvention.cs
--- a/Engine/Source/Programs/Horde/Horde.Build/Utilities/DefaultDiscriminatorConvention.cs
+++ b/Engine/Source/Programs/Horde/Horde.Build/Utilities/DefaultDiscriminatorConvention.cs
@@ -1,12 +1,9 @@
 // Copyright Epic Games, Inc. All Rights Reserved.
 
+using System;
 using MongoDB.Bson;
 using MongoDB.Bson.IO;
 using MongoDB.Bson.Serialization.Conventions;
-using System;
-using System.Collections.Generic;
-using System.Linq;
-using System.Threading.Tasks;
 
 namespace Horde.Build.Utilities
 {
@@ -29,42 +26,42 @@ namespace Horde.Build.Utilities
 		/// <summary>
 		/// The nominal type
 		/// </summary>
-		Type BaseType;
+		readonly Type _baseType;
 
 		/// <summary>
 		/// Default type to use if the inner discriminator returns the nominal type
 		/// </summary>
-		Type DefaultType;
+		readonly Type _defaultType;
 
 		/// <summary>
 		/// Constructor
 		/// </summary>
-		/// <param name="BaseType">The base class</param>
-		/// <param name="DefaultType">Default type to use if the inner discriminator returns an interface</param>
-		public DefaultDiscriminatorConvention(Type BaseType, Type DefaultType)
+		/// <param name="baseType">The base class</param>
+		/// <param name="defaultType">Default type to use if the inner discriminator returns an interface</param>
+		public DefaultDiscriminatorConvention(Type baseType, Type defaultType)
 		{
-			this.BaseType = BaseType;
-			this.DefaultType = DefaultType;
+			_baseType = baseType;
+			_defaultType = defaultType;
 		}
 
 		/// <inheritdoc/>
 		public string ElementName => Inner.ElementName;
 
 		/// <inheritdoc/>
-		public Type GetActualType(IBsonReader BsonReader, Type NominalType)
+		public Type GetActualType(IBsonReader bsonReader, Type nominalType)
 		{
-			Type ActualType = Inner.GetActualType(BsonReader, NominalType);
-			if (ActualType == BaseType)
+			Type actualType = Inner.GetActualType(bsonReader, nominalType);
+			if (actualType == _baseType)
 			{
-				ActualType = DefaultType;
+				actualType = _defaultType;
 			}
-			return ActualType;
+			return actualType;
 		}
 
 		/// <inheritdoc/>
-		public BsonValue GetDiscriminator(Type NominalType, Type ActualType)
+		public BsonValue GetDiscriminator(Type nominalType, Type actualType)
 		{
-			return Inner.GetDiscriminator(NominalType, ActualType);
+			return Inner.GetDiscriminator(nominalType, actualType);
 		}
 	}
 
diff --git a/Engine/Source/Programs/Horde/Horde.Build/Utilities/HordeClaimTypes.cs b/Engine/Source/Programs/Horde/Horde.Build/Utilities/HordeClaimTypes.cs
--- a/Engine/Source/Programs/Horde/Horde.Build/Utilities/HordeClaimTypes.cs
+++ b/Engine/Source/Programs/Horde/Horde.Build/Utilities/HordeClaimTypes.cs
@@ -1,13 +1,9 @@
 // Copyright Epic Games, Inc. All Rights Reserved.
 
-using Horde.Build.Collections;
-using Horde.Build.Models;
-using MongoDB.Bson;
-using System;
-using System.Collections.Generic;
 using System.Linq;
 using System.Security.Claims;
-using System.Threading.Tasks;
+using Horde.Build.Collections;
+using Horde.Build.Models;
 
 namespace Horde.Build.Utilities
 {
@@ -72,54 +68,54 @@ namespace Horde.Build.Utilities
 		/// <summary>
 		/// Gets the Horde user id from a principal
 		/// </summary>
-		/// <param name="Principal"></param>
+		/// <param name="principal"></param>
 		/// <returns></returns>
-		public static UserId? GetUserId(this ClaimsPrincipal Principal)
+		public static UserId? GetUserId(this ClaimsPrincipal principal)
 		{
-			string? IdValue = Principal.FindFirstValue(HordeClaimTypes.UserId);
-			if(IdValue == null)
+			string? idValue = principal.FindFirstValue(HordeClaimTypes.UserId);
+			if(idValue == null)
 			{
 				return null;
 			}
 			else
 			{
-				return new UserId(IdValue);
+				return new UserId(idValue);
 			}
 		}
 
 		/// <summary>
 		/// Gets the Horde user name from a principal
 		/// </summary>
-		/// <param name="Principal"></param>
+		/// <param name="principal"></param>
 		/// <returns></returns>
-		public static string? GetUserName(this ClaimsPrincipal Principal)
+		public static string? GetUserName(this ClaimsPrincipal principal)
 		{
-			return Principal.FindFirstValue(HordeClaimTypes.User) ?? Principal.FindFirstValue(ClaimTypes.Name);
+			return principal.FindFirstValue(HordeClaimTypes.User) ?? principal.FindFirstValue(ClaimTypes.Name);
 		}
 
 		/// <summary>
 		/// Get the email for the given user
 		/// </summary>
-		/// <param name="User">The user to query the email for</param>
+		/// <param name="user">The user to query the email for</param>
 		/// <returns>The user's email address or null if not found</returns>
-		public static string? GetEmail(this ClaimsPrincipal User)
+		public static string? GetEmail(this ClaimsPrincipal user)
 		{
-			return User.Claims.FirstOrDefault(x => x.Type == ClaimTypes.Email)?.Value;
+			return user.Claims.FirstOrDefault(x => x.Type == ClaimTypes.Email)?.Value;
 		}
 
 		/// <summary>
 		/// Gets the perforce username for the given principal
 		/// </summary>
-		/// <param name="User">The principal to get the Perforce user for</param>
+		/// <param name="user">The principal to get the Perforce user for</param>
 		/// <returns>Perforce user name</returns>
-		public static string? GetPerforceUser(this ClaimsPrincipal User)
+		public static string? GetPerforceUser(this ClaimsPrincipal user)
 		{
-			Claim? Claim = User.FindFirst(HordeClaimTypes.PerforceUser);
-			if (Claim == null)
+			Claim? claim = user.FindFirst(HordeClaimTypes.PerforceUser);
+			if (claim == null)
 			{
 				return null;
 			}
-			return Claim.Value;
+			return claim.Value;
 		}
 	}
 }
diff --git a/Engine/Source/Programs/Horde/Horde.Build/Utilities/HordeControllerBase.cs b/Engine/Source/Programs/Horde/Horde.Build/Utilities/HordeControllerBase.cs
--- a/Engine/Source/Programs/Horde/Horde.Build/Utilities/HordeControllerBase.cs
+++ b/Engine/Source/Programs/Horde/Horde.Build/Utilities/HordeControllerBase.cs
@@ -1,4 +1,4 @@
-﻿// Copyright Epic Games, Inc. All Rights Reserved.
+// Copyright Epic Games, Inc. All Rights Reserved.
 
 using EpicGames.Core;
 using Horde.Build.Acls;
@@ -6,13 +6,6 @@ using Horde.Build.Models;
 using Microsoft.AspNetCore.Http;
 using Microsoft.AspNetCore.Mvc;
 using Microsoft.Extensions.Logging;
-using System;
-using System.Collections.Generic;
-using System.Diagnostics;
-using System.Linq;
-using System.Linq.Expressions;
-using System.Threading;
-using System.Threading.Tasks;
 
 namespace Horde.Build.Utilities
 {
@@ -30,186 +23,186 @@ namespace Horde.Build.Utilities
 		/// <summary>
 		/// Create a response indicating a problem with the request
 		/// </summary>
-		/// <param name="Message">Standard structured-logging format string</param>
-		/// <param name="Args">Arguments for the format string</param>
+		/// <param name="message">Standard structured-logging format string</param>
+		/// <param name="args">Arguments for the format string</param>
 		[NonAction]
-		protected ActionResult BadRequest(string Message, params object[] Args)
+		protected ActionResult BadRequest(string message, params object[] args)
 		{
-			return BadRequest(LogEvent.Create(LogLevel.Error, Message, Args));
+			return BadRequest(LogEvent.Create(LogLevel.Error, message, args));
 		}
 
 		/// <summary>
 		/// Create a response indicating a problem with the request
 		/// </summary>
-		/// <param name="EventId">Id identifying the error</param>
-		/// <param name="Message">Standard structured-logging format string</param>
-		/// <param name="Args">Arguments for the format string</param>
+		/// <param name="eventId">Id identifying the error</param>
+		/// <param name="message">Standard structured-logging format string</param>
+		/// <param name="args">Arguments for the format string</param>
 		[NonAction]
-		protected ActionResult BadRequest(EventId EventId, string Message, params object[] Args)
+		protected ActionResult BadRequest(EventId eventId, string message, params object[] args)
 		{
-			return BadRequest(LogEvent.Create(LogLevel.Error, EventId, Message, Args));
+			return BadRequest(LogEvent.Create(LogLevel.Error, eventId, message, args));
 		}
 
 		/// <summary>
 		/// Returns a 403 (forbidden) response with the given action
 		/// </summary>
 		[NonAction]
-		protected ActionResult Forbid(AclAction Action)
+		protected ActionResult Forbid(AclAction action)
 		{
-			return Forbid("User does not have {Action} permission", Action);
+			return Forbid("User does not have {Action} permission", action);
 		}
 
 		/// <summary>
 		/// Returns a 403 (forbidden) response with the given action and object
 		/// </summary>
 		[NonAction]
-		protected ActionResult Forbid(AclAction Action, AgentId AgentId)
+		protected ActionResult Forbid(AclAction action, AgentId agentId)
 		{
-			return Forbid(Action, "agent {AgentId}", AgentId);
+			return Forbid(action, "agent {AgentId}", agentId);
 		}
 
 		/// <summary>
 		/// Returns a 403 (forbidden) response with the given action and object
 		/// </summary>
 		[NonAction]
-		protected ActionResult Forbid(AclAction Action, JobId JobId)
+		protected ActionResult Forbid(AclAction action, JobId jobId)
 		{
-			return Forbid(Action, "job {JobId}", JobId);
+			return Forbid(action, "job {JobId}", jobId);
 		}
 
 		/// <summary>
 		/// Returns a 403 (forbidden) response with the given action and object
 		/// </summary>
 		[NonAction]
-		protected ActionResult Forbid(AclAction Action, ProjectId ProjectId)
+		protected ActionResult Forbid(AclAction action, ProjectId projectId)
 		{
-			return Forbid(Action, "project {ProjectId}", ProjectId);
+			return Forbid(action, "project {ProjectId}", projectId);
 		}
 
 		/// <summary>
 		/// Returns a 403 (forbidden) response with the given action and object
 		/// </summary>
 		[NonAction]
-		protected ActionResult Forbid(AclAction Action, StreamId StreamId)
+		protected ActionResult Forbid(AclAction action, StreamId streamId)
 		{
-			return Forbid(Action, "stream {StreamId}", StreamId);
+			return Forbid(action, "stream {StreamId}", streamId);
 		}
 
 		/// <summary>
 		/// Returns a 403 (forbidden) response with the given action and object
 		/// </summary>
 		[NonAction]
-		protected ActionResult Forbid(AclAction Action, TemplateRefId TemplateId)
+		protected ActionResult Forbid(AclAction action, TemplateRefId templateId)
 		{
-			return Forbid(Action, "template {TemplateId}", TemplateId);
+			return Forbid(action, "template {TemplateId}", templateId);
 		}
 
 		/// <summary>
 		/// Returns a 403 (forbidden) response with the given action and object
 		/// </summary>
 		[NonAction]
-		private ActionResult Forbid(AclAction Action, string ObjectMessage, object Object)
+		private ActionResult Forbid(AclAction action, string objectMessage, object obj)
 		{
-			return Forbid($"User does not have {{Action}} permission for {ObjectMessage}", Action, Object);
+			return Forbid($"User does not have {{Action}} permission for {objectMessage}", action, obj);
 		}
 
 		/// <summary>
 		/// Returns a 403 response with the given log event
 		/// </summary>
 		[NonAction]
-		protected ActionResult Forbid(string Message, params object[] Args)
+		protected ActionResult Forbid(string message, params object[] args)
 		{
-			return StatusCode(StatusCodes.Status403Forbidden, LogEvent.Create(LogLevel.Error, Message, Args));
+			return StatusCode(StatusCodes.Status403Forbidden, LogEvent.Create(LogLevel.Error, message, args));
 		}
 
 		/// <summary>
 		/// Returns a 404 response for the given object
 		/// </summary>
 		[NonAction]
-		protected ActionResult NotFound(AgentId AgentId)
+		protected ActionResult NotFound(AgentId agentId)
 		{
-			return NotFound("Agent {AgentId} not found", AgentId);
+			return NotFound("Agent {AgentId} not found", agentId);
 		}
 
 		/// <summary>
 		/// Returns a 404 response for the given object
 		/// </summary>
 		[NonAction]
-		protected ActionResult NotFound(AgentId AgentId, LeaseId LeaseId)
+		protected ActionResult NotFound(AgentId agentId, LeaseId leaseId)
 		{
-			return NotFound("Lease {LeaseId} not found for agent {AgentId}", LeaseId, AgentId);
+			return NotFound("Lease {LeaseId} not found for agent {AgentId}", leaseId, agentId);
 		}
 
 		/// <summary>
 		/// Returns a 404 response for the given object
 		/// </summary>
 		[NonAction]
-		protected ActionResult NotFound(JobId JobId)
+		protected ActionResult NotFound(JobId jobId)
 		{
-			return NotFound("Job {JobId} not found", JobId);
+			return NotFound("Job {JobId} not found", jobId);
 		}
 
 		/// <summary>
 		/// Returns a 404 response for the given object
 		/// </summary>
 		[NonAction]
-		protected ActionResult NotFound(JobId JobId, int GroupIdx)
+		protected ActionResult NotFound(JobId jobId, int groupIdx)
 		{
-			return NotFound("Group {GroupIdx} on job {JobId} not found", GroupIdx, JobId);
+			return NotFound("Group {GroupIdx} on job {JobId} not found", groupIdx, jobId);
 		}
 
 		/// <summary>
 		/// Returns a 404 response for the given object
 		/// </summary>
 		[NonAction]
-		protected ActionResult NotFound(JobId JobId, int GroupIdx, int NodeIdx)
+		protected ActionResult NotFound(JobId jobId, int groupIdx, int nodeIdx)
 		{
-			return NotFound("Node {NodeIdx} not found on job {JobId} group {GroupIdx}", NodeIdx, JobId, GroupIdx);
+			return NotFound("Node {NodeIdx} not found on job {JobId} group {GroupIdx}", nodeIdx, jobId, groupIdx);
 		}
 
 		/// <summary>
 		/// Returns a 404 response for the given object
 		/// </summary>
 		[NonAction]
-		protected ActionResult NotFound(JobId JobId, SubResourceId BatchId)
+		protected ActionResult NotFound(JobId jobId, SubResourceId batchId)
 		{
-			return NotFound("Batch {BatchId} not found on job {JobId}", BatchId, JobId);
+			return NotFound("Batch {BatchId} not found on job {JobId}", batchId, jobId);
 		}
 
 		/// <summary>
 		/// Returns a 404 response for the given object
 		/// </summary>
 		[NonAction]
-		protected ActionResult NotFound(JobId JobId, SubResourceId BatchId, SubResourceId StepId)
+		protected ActionResult NotFound(JobId jobId, SubResourceId batchId, SubResourceId stepId)
 		{
-			return NotFound("Step {StepId} not found on job {JobId} batch {BatchId}", StepId, JobId, BatchId);
+			return NotFound("Step {StepId} not found on job {JobId} batch {BatchId}", stepId, jobId, batchId);
 		}
 
 		/// <summary>
 		/// Returns a 404 response for the given object
 		/// </summary>
 		[NonAction]
-		protected ActionResult NotFound(StreamId StreamId)
+		protected ActionResult NotFound(StreamId streamId)
 		{
-			return NotFound("Stream {StreamId} not found", StreamId);
+			return NotFound("Stream {StreamId} not found", streamId);
 		}
 
 		/// <summary>
 		/// Returns a 404 response for the given object
 		/// </summary>
 		[NonAction]
-		protected ActionResult NotFound(StreamId StreamId, TemplateRefId TemplateId)
+		protected ActionResult NotFound(StreamId streamId, TemplateRefId templateId)
 		{
-			return NotFound("Template {TemplateId} not found on stream {StreamId}", TemplateId, StreamId);
+			return NotFound("Template {TemplateId} not found on stream {StreamId}", templateId, streamId);
 		}
 
 		/// <summary>
 		/// Returns a 404 response for the given object
 		/// </summary>
 		[NonAction]
-		protected ActionResult NotFound(string Message, params object[] Args)
+		protected ActionResult NotFound(string message, params object[] args)
 		{
-			return NotFound(LogEvent.Create(LogLevel.Error, Message, Args));
+			return NotFound(LogEvent.Create(LogLevel.Error, message, args));
 		}
 	}
 }
diff --git a/Engine/Source/Programs/Horde/Horde.Build/Utilities/JsonElementExtensions.cs b/Engine/Source/Programs/Horde/Horde.Build/Utilities/JsonElementExtensions.cs
--- a/Engine/Source/Programs/Horde/Horde.Build/Utilities/JsonElementExtensions.cs
+++ b/Engine/Source/Programs/Horde/Horde.Build/Utilities/JsonElementExtensions.cs
@@ -1,13 +1,9 @@
 // Copyright Epic Games, Inc. All Rights Reserved.
 
-using MongoDB.Bson.IO;
 using System;
-using System.Collections.Generic;
 using System.Diagnostics.CodeAnalysis;
-using System.Linq;
 using System.Text;
 using System.Text.Json;
-using System.Threading.Tasks;
 
 namespace Horde.Build.Utilities
 {
@@ -19,48 +15,48 @@ namespace Horde.Build.Utilities
 		/// <summary>
 		/// Checks if the element has a property with the given value
 		/// </summary>
-		/// <param name="Element">Element to search</param>
-		/// <param name="Name">Name of the property</param>
-		/// <param name="Value">Expected value of the property</param>
+		/// <param name="element">Element to search</param>
+		/// <param name="name">Name of the property</param>
+		/// <param name="value">Expected value of the property</param>
 		/// <returns>True if the property exists and matches</returns>
-		public static bool HasStringProperty(this JsonElement Element, ReadOnlySpan<byte> Name, string Value)
+		public static bool HasStringProperty(this JsonElement element, ReadOnlySpan<byte> name, string value)
 		{
-			JsonElement Property;
-			return Element.ValueKind == JsonValueKind.Object && Element.TryGetProperty(Name, out Property) && Property.ValueEquals(Value);
+			JsonElement property;
+			return element.ValueKind == JsonValueKind.Object && element.TryGetProperty(name, out property) && property.ValueEquals(value);
 		}
 
 		/// <summary>
 		/// Checks if the element has a property with the given value
 		/// </summary>
-		/// <param name="Element">Element to search</param>
-		/// <param name="Name">Name of the property</param>
-		/// <param name="Value">Expected value of the property</param>
+		/// <param name="element">Element to search</param>
+		/// <param name="name">Name of the property</param>
+		/// <param name="value">Expected value of the property</param>
 		/// <returns>True if the property exists and matches</returns>
-		public static bool HasStringProperty(this JsonElement Element, string Name, string Value)
+		public static bool HasStringProperty(this JsonElement element, string name, string value)
 		{
-			JsonElement Property;
-			return Element.ValueKind == JsonValueKind.Object && Element.TryGetProperty(Name, out Property) && Property.ValueEquals(Value);
+			JsonElement property;
+			return element.ValueKind == JsonValueKind.Object && element.TryGetProperty(name, out property) && property.ValueEquals(value);
 		}
 
 		/// <summary>
 		/// Gets a property value of a certain type
 		/// </summary>
-		/// <param name="Element">The element to get a property from</param>
-		/// <param name="Name">Name of the element</param>
-		/// <param name="ValueKind">The required type of property</param>
-		/// <param name="Value">Value of the property</param>
+		/// <param name="element">The element to get a property from</param>
+		/// <param name="name">Name of the element</param>
+		/// <param name="valueKind">The required type of property</param>
+		/// <param name="value">Value of the property</param>
 		/// <returns>True if the property exists and was a string</returns>
-		public static bool TryGetProperty(this JsonElement Element, ReadOnlySpan<byte> Name, JsonValueKind ValueKind, [NotNullWhen(true)] out JsonElement Value)
+		public static bool TryGetProperty(this JsonElement element, ReadOnlySpan<byte> name, JsonValueKind valueKind, [NotNullWhen(true)] out JsonElement value)
 		{
-			JsonElement Property;
-			if (Element.TryGetProperty(Name, out Property) && Property.ValueKind == ValueKind)
+			JsonElement property;
+			if (element.TryGetProperty(name, out property) && property.ValueKind == valueKind)
 			{
-				Value = Property;
+				value = property;
 				return true;
 			}
 			else
 			{
-				Value = new JsonElement();
+				value = new JsonElement();
 				return false;
 			}
 		}
@@ -68,22 +64,22 @@ namespace Horde.Build.Utilities
 		/// <summary>
 		/// Gets a property value of a certain type
 		/// </summary>
-		/// <param name="Element">The element to get a property from</param>
-		/// <param name="Name">Name of the element</param>
-		/// <param name="ValueKind">The required type of property</param>
-		/// <param name="Value">Value of the property</param>
+		/// <param name="element">The element to get a property from</param>
+		/// <param name="name">Name of the element</param>
+		/// <param name="valueKind">The required type of property</param>
+		/// <param name="value">Value of the property</param>
 		/// <returns>True if the property exists and was a string</returns>
-		public static bool TryGetProperty(this JsonElement Element, string Name, JsonValueKind ValueKind, [NotNullWhen(true)] out JsonElement Value)
+		public static bool TryGetProperty(this JsonElement element, string name, JsonValueKind valueKind, [NotNullWhen(true)] out JsonElement value)
 		{
-			JsonElement Property;
-			if (Element.TryGetProperty(Name, out Property) && Property.ValueKind == ValueKind)
+			JsonElement property;
+			if (element.TryGetProperty(name, out property) && property.ValueKind == valueKind)
 			{
-				Value = Property;
+				value = property;
 				return true;
 			}
 			else
 			{
-				Value = new JsonElement();
+				value = new JsonElement();
 				return false;
 			}
 		}
@@ -91,21 +87,21 @@ namespace Horde.Build.Utilities
 		/// <summary>
 		/// Gets a string property value
 		/// </summary>
-		/// <param name="Element">The element to get a property from</param>
-		/// <param name="Name">Name of the element</param>
-		/// <param name="Value">Value of the property</param>
+		/// <param name="element">The element to get a property from</param>
+		/// <param name="name">Name of the element</param>
+		/// <param name="value">Value of the property</param>
 		/// <returns>True if the property exists and was a string</returns>
-		public static bool TryGetStringProperty(this JsonElement Element, ReadOnlySpan<byte> Name, [NotNullWhen(true)] out string? Value)
+		public static bool TryGetStringProperty(this JsonElement element, ReadOnlySpan<byte> name, [NotNullWhen(true)] out string? value)
 		{
-			JsonElement Property;
-			if (Element.TryGetProperty(Name, out Property) && Property.ValueKind == JsonValueKind.String)
+			JsonElement property;
+			if (element.TryGetProperty(name, out property) && property.ValueKind == JsonValueKind.String)
 			{
-				Value = Property.GetString()!;
+				value = property.GetString()!;
 				return true;
 			}
 			else
 			{
-				Value = null;
+				value = null;
 				return false;
 			}
 		}
@@ -113,70 +109,70 @@ namespace Horde.Build.Utilities
 		/// <summary>
 		/// Gets a string property value
 		/// </summary>
-		/// <param name="Element">The element to get a property from</param>
-		/// <param name="Name">Name of the element</param>
-		/// <param name="Value">Value of the property</param>
+		/// <param name="element">The element to get a property from</param>
+		/// <param name="name">Name of the element</param>
+		/// <param name="value">Value of the property</param>
 		/// <returns>True if the property exists and was a string</returns>
-		public static bool TryGetStringProperty(this JsonElement Element, string Name, [NotNullWhen(true)] out string? Value)
+		public static bool TryGetStringProperty(this JsonElement element, string name, [NotNullWhen(true)] out string? value)
 		{
-			return TryGetStringProperty(Element, Encoding.UTF8.GetBytes(Name).AsSpan(), out Value);
+			return TryGetStringProperty(element, Encoding.UTF8.GetBytes(name).AsSpan(), out value);
 		}
 
 		/// <summary>
 		/// Gets a property value from a document or subdocument, indicated with dotted notation
 		/// </summary>
-		/// <param name="Element">Document to get a property for</param>
-		/// <param name="Name">Name of the property</param>
-		/// <param name="OutElement">Receives the nexted element</param>
+		/// <param name="element">Document to get a property for</param>
+		/// <param name="name">Name of the property</param>
+		/// <param name="outElement">Receives the nexted element</param>
 		/// <returns>True if the property exists and was of the correct type</returns>
-		public static bool TryGetNestedProperty(this JsonElement Element, ReadOnlySpan<byte> Name, [NotNullWhen(true)] out JsonElement OutElement)
+		public static bool TryGetNestedProperty(this JsonElement element, ReadOnlySpan<byte> name, [NotNullWhen(true)] out JsonElement outElement)
 		{
-			int DotIdx = Name.IndexOf((byte)'.');
-			if (DotIdx == -1)
+			int dotIdx = name.IndexOf((byte)'.');
+			if (dotIdx == -1)
 			{
-				return Element.TryGetProperty(Name, out OutElement);
+				return element.TryGetProperty(name, out outElement);
 			}
 
-			JsonElement DocValue;
-			if (Element.TryGetProperty(Name.Slice(0, DotIdx), out DocValue) && DocValue.ValueKind == JsonValueKind.Object)
+			JsonElement docValue;
+			if (element.TryGetProperty(name.Slice(0, dotIdx), out docValue) && docValue.ValueKind == JsonValueKind.Object)
 			{
-				return TryGetNestedProperty(DocValue, Name.Slice(DotIdx + 1), out OutElement);
+				return TryGetNestedProperty(docValue, name.Slice(dotIdx + 1), out outElement);
 			}
 
-			OutElement = new JsonElement();
+			outElement = new JsonElement();
 			return false;
 		}
 
 		/// <summary>
 		/// Gets a property value from a document or subdocument, indicated with dotted notation
 		/// </summary>
-		/// <param name="Element">Document to get a property for</param>
-		/// <param name="Name">Name of the property</param>
-		/// <param name="OutElement">Receives the nexted element</param>
+		/// <param name="element">Document to get a property for</param>
+		/// <param name="name">Name of the property</param>
+		/// <param name="outElement">Receives the nexted element</param>
 		/// <returns>True if the property exists and was of the correct type</returns>
-		public static bool TryGetNestedProperty(this JsonElement Element, string Name, [NotNullWhen(true)] out JsonElement OutElement)
+		public static bool TryGetNestedProperty(this JsonElement element, string name, [NotNullWhen(true)] out JsonElement outElement)
 		{
-			return TryGetNestedProperty(Element, Encoding.UTF8.GetBytes(Name).AsSpan(), out OutElement);
+			return TryGetNestedProperty(element, Encoding.UTF8.GetBytes(name).AsSpan(), out outElement);
 		}
 
 		/// <summary>
 		/// Gets an int32 value from the document
 		/// </summary>
-		/// <param name="Element">Document to get a property for</param>
-		/// <param name="Name">Name of the property</param>
-		/// <param name="OutValue">Receives the property value</param>
+		/// <param name="element">Document to get a property for</param>
+		/// <param name="name">Name of the property</param>
+		/// <param name="outValue">Receives the property value</param>
 		/// <returns>True if the property was retrieved</returns>
-		public static bool TryGetNestedProperty(this JsonElement Element, string Name, out int OutValue)
+		public static bool TryGetNestedProperty(this JsonElement element, string name, out int outValue)
 		{
-			JsonElement Value;
-			if (Element.TryGetNestedProperty(Name, out Value) && Value.ValueKind == JsonValueKind.Number)
+			JsonElement value;
+			if (element.TryGetNestedProperty(name, out value) && value.ValueKind == JsonValueKind.Number)
 			{
-				OutValue = Value.GetInt32();
+				outValue = value.GetInt32();
 				return true;
 			}
 			else
 			{
-				OutValue = 0;
+				outValue = 0;
 				return false;
 			}
 		}
@@ -184,21 +180,21 @@ namespace Horde.Build.Utilities
 		/// <summary>
 		/// Gets a string value from the document
 		/// </summary>
-		/// <param name="Element">Document to get a property for</param>
-		/// <param name="Name">Name of the property</param>
-		/// <param name="OutValue">Receives the property value</param>
+		/// <param name="element">Document to get a property for</param>
+		/// <param name="name">Name of the property</param>
+		/// <param name="outValue">Receives the property value</param>
 		/// <returns>True if the property was retrieved</returns>
-		public static bool TryGetNestedProperty(this JsonElement Element, string Name, [NotNullWhen(true)] out string? OutValue)
+		public static bool TryGetNestedProperty(this JsonElement element, string name, [NotNullWhen(true)] out string? outValue)
 		{
-			JsonElement Value;
-			if (Element.TryGetNestedProperty(Name, out Value) && Value.ValueKind == JsonValueKind.String)
+			JsonElement value;
+			if (element.TryGetNestedProperty(name, out value) && value.ValueKind == JsonValueKind.String)
 			{
-				OutValue = Value.GetString();
-				return OutValue != null;
+				outValue = value.GetString();
+				return outValue != null;
 			}
 			else
 			{
-				OutValue = null;
+				outValue = null;
 				return false;
 			}
 		}
diff --git a/Engine/Source/Programs/Horde/Horde.Build/Utilities/JwtUtils.cs b/Engine/Source/Programs/Horde/Horde.Build/Utilities/JwtUtils.cs
--- a/Engine/Source/Programs/Horde/Horde.Build/Utilities/JwtUtils.cs
+++ b/Engine/Source/Programs/Horde/Horde.Build/Utilities/JwtUtils.cs
@@ -1,13 +1,10 @@
 // Copyright Epic Games, Inc. All Rights Reserved.
 
-using Microsoft.AspNetCore.Http;
 using System;
-using System.Collections.Generic;
 using System.Diagnostics.CodeAnalysis;
 using System.IdentityModel.Tokens.Jwt;
-using System.Linq;
 using System.Text.Json;
-using System.Threading.Tasks;
+using Microsoft.AspNetCore.Http;
 
 namespace Horde.Build.Utilities
 {
@@ -19,57 +16,57 @@ namespace Horde.Build.Utilities
 		/// <summary>
 		/// Gets the bearer token from an HTTP request
 		/// </summary>
-		/// <param name="Request">The request to read from</param>
-		/// <param name="BearerPrefix">The bearer token prefix, ex. "Bearer "</param>
-		/// <param name="Token">On success, receives the bearer token</param>
+		/// <param name="request">The request to read from</param>
+		/// <param name="bearerPrefix">The bearer token prefix, ex. "Bearer "</param>
+		/// <param name="token">On success, receives the bearer token</param>
 		/// <returns>True if the bearer token was read</returns>
-		public static bool TryGetBearerToken(HttpRequest Request, string BearerPrefix, [NotNullWhen(true)] out string? Token)
+		public static bool TryGetBearerToken(HttpRequest request, string bearerPrefix, [NotNullWhen(true)] out string? token)
 		{
 			// Get the authorization header
-			string? Authorization = Request.Headers.Authorization;
-			if (String.IsNullOrEmpty(Authorization))
+			string? authorization = request.Headers.Authorization;
+			if (String.IsNullOrEmpty(authorization))
 			{
-				Token = null;
+				token = null;
 				return false;
 			}
 
 			// Check if it's a bearer token				
-			if (!Authorization.StartsWith(BearerPrefix, StringComparison.OrdinalIgnoreCase))
+			if (!authorization.StartsWith(bearerPrefix, StringComparison.OrdinalIgnoreCase))
 			{
-				Token = null;
+				token = null;
 				return false;
 			}
 
 			// Get the token
-			Token = Authorization.Substring(BearerPrefix.Length).Trim();
+			token = authorization.Substring(bearerPrefix.Length).Trim();
 			return true;
 		}
 
 		/// <summary>
 		/// Tries to parse a JWT and check the issuer matches
 		/// </summary>
-		/// <param name="Token">The token to parse</param>
-		/// <param name="JwtToken">On success, receives the parsed JWT</param>
+		/// <param name="token">The token to parse</param>
+		/// <param name="jwtToken">On success, receives the parsed JWT</param>
 		/// <returns>True if the jwt was parsed</returns>
-		public static bool TryParseJwt(string Token, [NotNullWhen(true)] out JwtSecurityToken? JwtToken)
+		public static bool TryParseJwt(string token, [NotNullWhen(true)] out JwtSecurityToken? jwtToken)
 		{
 			// Check if it's a JWT
-			JwtSecurityTokenHandler Handler = new JwtSecurityTokenHandler();
-			if (!Handler.CanReadToken(Token))
+			JwtSecurityTokenHandler handler = new JwtSecurityTokenHandler();
+			if (!handler.CanReadToken(token))
 			{
-				JwtToken = null;
+				jwtToken = null;
 				return false;
 			}
 
 			// Try to parse the JWT
 			try
 			{
-				JwtToken = Handler.ReadJwtToken(Token);
+				jwtToken = handler.ReadJwtToken(token);
 				return true;
 			}
 			catch (JsonException)
 			{
-				JwtToken = null;
+				jwtToken = null;
 				return false;
 			}
 		}
diff --git a/Engine/Source/Programs/Horde/Horde.Build/Utilities/LazyCache.cs b/Engine/Source/Programs/Horde/Horde.Build/Utilities/LazyCache.cs
--- a/Engine/Source/Programs/Horde/Horde.Build/Utilities/LazyCache.cs
+++ b/Engine/Source/Programs/Horde/Horde.Build/Utilities/LazyCache.cs
@@ -1,12 +1,8 @@
 // Copyright Epic Games, Inc. All Rights Reserved.
 
-using Microsoft.Extensions.Caching.Memory;
-using Microsoft.Extensions.Options;
 using System;
 using System.Collections.Concurrent;
-using System.Collections.Generic;
 using System.Diagnostics;
-using System.Linq;
 using System.Threading;
 using System.Threading.Tasks;
 
@@ -37,91 +33,91 @@ namespace Horde.Build.Utilities
 	{
 		class Item
 		{
-			public Task<TValue>? CurrentTask;
-			public Stopwatch Timer = Stopwatch.StartNew();
-			public Task<TValue>? UpdateTask;
+			public Task<TValue>? _currentTask;
+			public Stopwatch _timer = Stopwatch.StartNew();
+			public Task<TValue>? _updateTask;
 		}
 
-		ConcurrentDictionary<TKey, Item> Dictionary = new ConcurrentDictionary<TKey, Item>();
-		Func<TKey, Task<TValue>> GetValueAsync;
-		LazyCacheOptions Options;
+		readonly ConcurrentDictionary<TKey, Item> _dictionary = new ConcurrentDictionary<TKey, Item>();
+		readonly Func<TKey, Task<TValue>> _getValueAsync;
+		readonly LazyCacheOptions _options;
 
 		/// <summary>
 		/// Constructor
 		/// </summary>
-		/// <param name="GetValueAsync">Function used to get a value</param>
-		/// <param name="Options"></param>
-		public LazyCache(Func<TKey, Task<TValue>> GetValueAsync, LazyCacheOptions Options)
+		/// <param name="getValueAsync">Function used to get a value</param>
+		/// <param name="options"></param>
+		public LazyCache(Func<TKey, Task<TValue>> getValueAsync, LazyCacheOptions options)
 		{
-			this.GetValueAsync = GetValueAsync;
-			this.Options = Options;
+			_getValueAsync = getValueAsync;
+			_options = options;
 		}
 
 		/// <inheritdoc/>
 		public void Dispose()
 		{
-			foreach (Item Item in Dictionary.Values)
+			foreach (Item item in _dictionary.Values)
 			{
-				Item.CurrentTask?.Wait();
-				Item.UpdateTask?.Wait();
+				item._currentTask?.Wait();
+				item._updateTask?.Wait();
 			}
 		}
 
 		/// <summary>
 		/// Gets the value associated with a key
 		/// </summary>
-		/// <param name="Key">The key to query</param>
-		/// <param name="MaxAge">Maximum age for values to return</param>
+		/// <param name="key">The key to query</param>
+		/// <param name="maxAge">Maximum age for values to return</param>
 		/// <returns></returns>
-		public Task<TValue> GetAsync(TKey Key, TimeSpan? MaxAge = null)
+		public Task<TValue> GetAsync(TKey key, TimeSpan? maxAge = null)
 		{
-			Item Item = Dictionary.GetOrAdd(Key, Key => new Item());
+			Item item = _dictionary.GetOrAdd(key, key => new Item());
 
 			// Create the task to get the current value
-			Task<TValue> CurrentTask = InterlockedCreateTask(ref Item.CurrentTask, () => GetValueAsync(Key));
+			Task<TValue> currentTask = InterlockedCreateTask(ref item._currentTask, () => _getValueAsync(key));
 
 			// If an update has completed, swap it out
-			Task<TValue>? UpdateTask = Item.UpdateTask;
-			if (UpdateTask != null && UpdateTask.IsCompleted)
+			Task<TValue>? updateTask = item._updateTask;
+			if (updateTask != null && updateTask.IsCompleted)
 			{
-				Interlocked.CompareExchange(ref Item.CurrentTask, UpdateTask, CurrentTask);
-				Interlocked.CompareExchange(ref Item.UpdateTask, null, UpdateTask);
-				Item.Timer.Restart();
+				Interlocked.CompareExchange(ref item._currentTask, updateTask, currentTask);
+				Interlocked.CompareExchange(ref item._updateTask, null, updateTask);
+				item._timer.Restart();
 			}
 
 			// Check if we need to update the value
-			TimeSpan Age = Item.Timer.Elapsed;
-			if (MaxAge != null && Age > MaxAge.Value)
+			TimeSpan age = item._timer.Elapsed;
+			if (maxAge != null && age > maxAge.Value)
 			{
-				return InterlockedCreateTask(ref Item.UpdateTask, () => GetValueAsync(Key));
+				return InterlockedCreateTask(ref item._updateTask, () => _getValueAsync(key));
 			}
-			if (Age > Options.RefreshTime)
+			if (age > _options.RefreshTime)
 			{
-				InterlockedCreateTask(ref Item.UpdateTask, () => GetValueAsync(Key));
+				InterlockedCreateTask(ref item._updateTask, () => _getValueAsync(key));
 			}
 
-			return CurrentTask;
+			return currentTask;
 		}
 
 		/// <summary>
 		/// Creates a task, guaranteeing that only one task will be assigned to the given slot. Creates a cold task and only starts it once the variable is set.
 		/// </summary>
-		/// <param name="Value"></param>
-		/// <param name="CreateTask"></param>
+		/// <param name="value"></param>
+		/// <param name="createTask"></param>
 		/// <returns></returns>
-		static Task<TValue> InterlockedCreateTask(ref Task<TValue>? Value, Func<Task<TValue>> CreateTask)
+		static Task<TValue> InterlockedCreateTask(ref Task<TValue>? value, Func<Task<TValue>> createTask)
 		{
-			Task<TValue>? CurrentTask = Value;
-			while (CurrentTask == null)
+			Task<TValue>? currentTask = value;
+			while (currentTask == null)
 			{
-				Task<Task<TValue>> NewTask = new Task<Task<TValue>>(CreateTask);
-				if (Interlocked.CompareExchange(ref Value, NewTask.Unwrap(), null) == null)
+				Task<Task<TValue>> newTask = new Task<Task<TValue>>(createTask);
+				if (Interlocked.CompareExchange(ref value, newTask.Unwrap(), null) == null)
 				{
-					NewTask.Start();
+					newTask.Start();
 				}
-				CurrentTask = Value;
+				currentTask = value;
 			}
-			return CurrentTask;
+			return currentTask;
 		}
 	}
 }
diff --git a/Engine/Source/Programs/Horde/Horde.Build/Utilities/LazyCachedValue.cs b/Engine/Source/Programs/Horde/Horde.Build/Utilities/LazyCachedValue.cs
--- a/Engine/Source/Programs/Horde/Horde.Build/Utilities/LazyCachedValue.cs
+++ b/Engine/Source/Programs/Horde/Horde.Build/Utilities/LazyCachedValue.cs
@@ -1,11 +1,7 @@
 // Copyright Epic Games, Inc. All Rights Reserved.
 
 using System;
-using System.Collections.Generic;
 using System.Diagnostics;
-using System.Linq;
-using System.Reflection.Metadata.Ecma335;
-using System.Threading.Tasks;
 
 namespace Horde.Build.Utilities
 {
@@ -18,40 +14,40 @@ namespace Horde.Build.Utilities
 		/// <summary>
 		/// The current value
 		/// </summary>
-		T? Value;
+		T? _value;
 
 		/// <summary>
 		/// Generator for the new value
 		/// </summary>
-		Func<T> Generator;
+		readonly Func<T> _generator;
 
 		/// <summary>
 		/// Time since the value was updated
 		/// </summary>
-		Stopwatch Timer = Stopwatch.StartNew();
+		readonly Stopwatch _timer = Stopwatch.StartNew();
 
 		/// <summary>
 		/// Default expiry time
 		/// </summary>
-		TimeSpan DefaultMaxAge;
+		readonly TimeSpan _defaultMaxAge;
 
 		/// <summary>
 		/// Default constructor
 		/// </summary>
-		public LazyCachedValue(Func<T> Generator, TimeSpan MaxAge)
+		public LazyCachedValue(Func<T> generator, TimeSpan maxAge)
 		{
-			this.Generator = Generator;
-			this.DefaultMaxAge = MaxAge;
+			_generator = generator;
+			_defaultMaxAge = maxAge;
 		}
 
 		/// <summary>
 		/// Sets the new value
 		/// </summary>
-		/// <param name="Value">The value to store</param>
-		public void Set(T Value)
+		/// <param name="value">The value to store</param>
+		public void Set(T value)
 		{
-			this.Value = Value;
-			Timer.Restart();
+			_value = value;
+			_timer.Restart();
 		}
 
 		/// <summary>
@@ -60,22 +56,22 @@ namespace Horde.Build.Utilities
 		/// <returns>The cached value, if valid</returns>
 		public T GetCached()
 		{
-			return GetCached(DefaultMaxAge);
+			return GetCached(_defaultMaxAge);
 		}
 
 		/// <summary>
 		/// Tries to get the current value
 		/// </summary>
 		/// <returns>The cached value, if valid</returns>
-		public T GetCached(TimeSpan MaxAge)
+		public T GetCached(TimeSpan maxAge)
 		{
-			T? Current = Value;
-			if (Current == null || Timer.Elapsed > MaxAge)
+			T? current = _value;
+			if (current == null || _timer.Elapsed > maxAge)
 			{
-				Current = Generator();
-				Set(Current);
+				current = _generator();
+				Set(current);
 			}
-			return Current;
+			return current;
 		}
 
 		/// <summary>
@@ -84,9 +80,9 @@ namespace Horde.Build.Utilities
 		/// <returns>The latest value</returns>
 		public T GetLatest()
 		{
-			T NewValue = Generator();
-			Set(NewValue);
-			return NewValue;
+			T newValue = _generator();
+			Set(newValue);
+			return newValue;
 		}
 	}
 }
diff --git a/Engine/Source/Programs/Horde/Horde.Build/Utilities/MongoExtensions.cs b/Engine/Source/Programs/Horde/Horde.Build/Utilities/MongoExtensions.cs
--- a/Engine/Source/Programs/Horde/Horde.Build/Utilities/MongoExtensions.cs
+++ b/Engine/Source/Programs/Horde/Horde.Build/Utilities/MongoExtensions.cs
@@ -1,13 +1,13 @@
 // Copyright Epic Games, Inc. All Rights Reserved.
 
-using MongoDB.Bson;
-using MongoDB.Driver;
 using System;
 using System.Collections.Generic;
 using System.Linq;
 using System.Linq.Expressions;
 using System.Text.Json;
 using System.Threading.Tasks;
+using MongoDB.Bson;
+using MongoDB.Driver;
 
 namespace Horde.Build.Utilities
 {
@@ -19,81 +19,81 @@ namespace Horde.Build.Utilities
 		/// <summary>
 		/// Filters the documents returned from a search
 		/// </summary>
-		/// <param name="Query">The query to filter</param>
-		/// <param name="Index">Index of the first document to return</param>
-		/// <param name="Count">Number of documents to return</param>
+		/// <param name="query">The query to filter</param>
+		/// <param name="index">Index of the first document to return</param>
+		/// <param name="count">Number of documents to return</param>
 		/// <returns>New query</returns>
-		public static IFindFluent<TDocument, TProjection> Range<TDocument, TProjection>(this IFindFluent<TDocument, TProjection> Query, int? Index, int? Count)
+		public static IFindFluent<TDocument, TProjection> Range<TDocument, TProjection>(this IFindFluent<TDocument, TProjection> query, int? index, int? count)
 		{
-			if (Index != null && Index.Value != 0)
+			if (index != null && index.Value != 0)
 			{
-				Query = Query.Skip(Index.Value);
+				query = query.Skip(index.Value);
 			}
-			if(Count != null)
+			if(count != null)
 			{
-				Query = Query.Limit(Count.Value);
+				query = query.Limit(count.Value);
 			}
-			return Query;
+			return query;
 		}
 
 		/// <summary>
 		/// Filters the documents returned from a search
 		/// </summary>
-		/// <param name="Query">The query to filter</param>
-		/// <param name="Index">Index of the first document to return</param>
-		/// <param name="Count">Number of documents to return</param>
+		/// <param name="query">The query to filter</param>
+		/// <param name="index">Index of the first document to return</param>
+		/// <param name="count">Number of documents to return</param>
 		/// <returns>New query</returns>
-		public static IAggregateFluent<TDocument> Range<TDocument>(this IAggregateFluent<TDocument> Query, int? Index, int? Count)
+		public static IAggregateFluent<TDocument> Range<TDocument>(this IAggregateFluent<TDocument> query, int? index, int? count)
 		{
-			if (Index != null && Index.Value != 0)
+			if (index != null && index.Value != 0)
 			{
-				Query = Query.Skip(Index.Value);
+				query = query.Skip(index.Value);
 			}
-			if (Count != null)
+			if (count != null)
 			{
-				Query = Query.Limit(Count.Value);
+				query = query.Limit(count.Value);
 			}
-			return Query;
+			return query;
 		}
 
 		/// <summary>
 		/// Filters the documents returned from a search
 		/// </summary>
-		/// <param name="Query">The query to filter</param>
+		/// <param name="query">The query to filter</param>
 		/// <returns>New query</returns>
-		public static async Task<List<TResult>> ToListAsync<TDocument, TResult>(this IAsyncCursorSource<TDocument> Query) where TDocument : TResult
+		public static async Task<List<TResult>> ToListAsync<TDocument, TResult>(this IAsyncCursorSource<TDocument> query) where TDocument : TResult
 		{
-			List<TResult> Results = new List<TResult>();
-			using (IAsyncCursor<TDocument> Cursor = await Query.ToCursorAsync())
+			List<TResult> results = new List<TResult>();
+			using (IAsyncCursor<TDocument> cursor = await query.ToCursorAsync())
 			{
-				while (await Cursor.MoveNextAsync())
+				while (await cursor.MoveNextAsync())
 				{
-					foreach(TDocument Document in Cursor.Current)
+					foreach(TDocument document in cursor.Current)
 					{
-						Results.Add(Document);
+						results.Add(document);
 					}
 				}
 			}
-			return Results;
+			return results;
 		}
 
 		/// <summary>
 		/// Attempts to insert a document into a collection, handling the error case that a document with the given key already exists
 		/// </summary>
 		/// <typeparam name="TDocument"></typeparam>
-		/// <param name="Collection">Collection to insert into</param>
-		/// <param name="NewDocument">The document to insert</param>
+		/// <param name="collection">Collection to insert into</param>
+		/// <param name="newDocument">The document to insert</param>
 		/// <returns>True if the document was inserted, false if it already exists</returns>
-		public static async Task<bool> InsertOneIgnoreDuplicatesAsync<TDocument>(this IMongoCollection<TDocument> Collection, TDocument NewDocument)
+		public static async Task<bool> InsertOneIgnoreDuplicatesAsync<TDocument>(this IMongoCollection<TDocument> collection, TDocument newDocument)
 		{
 			try
 			{
-				await Collection.InsertOneAsync(NewDocument);
+				await collection.InsertOneAsync(newDocument);
 				return true;
 			}
-			catch (MongoWriteException Ex)
+			catch (MongoWriteException ex)
 			{
-				if (Ex.WriteError.Category == ServerErrorCategory.DuplicateKey)
+				if (ex.WriteError.Category == ServerErrorCategory.DuplicateKey)
 				{
 					return false;
 				}
@@ -108,27 +108,27 @@ namespace Horde.Build.Utilities
 		/// Attempts to insert a document into a collection, handling the error case that a document with the given key already exists
 		/// </summary>
 		/// <typeparam name="TDocument"></typeparam>
-		/// <param name="Collection">Collection to insert into</param>
-		/// <param name="NewDocuments">The document to insert</param>
-		public static async Task InsertManyIgnoreDuplicatesAsync<TDocument>(this IMongoCollection<TDocument> Collection, List<TDocument> NewDocuments)
+		/// <param name="collection">Collection to insert into</param>
+		/// <param name="newDocuments">The document to insert</param>
+		public static async Task InsertManyIgnoreDuplicatesAsync<TDocument>(this IMongoCollection<TDocument> collection, List<TDocument> newDocuments)
 		{
 			try
 			{
-				if (NewDocuments.Count > 0)
+				if (newDocuments.Count > 0)
 				{
-					await Collection.InsertManyAsync(NewDocuments, new InsertManyOptions { IsOrdered = false });
+					await collection.InsertManyAsync(newDocuments, new InsertManyOptions { IsOrdered = false });
 				}
 			}
-			catch (MongoWriteException Ex)
+			catch (MongoWriteException ex)
 			{
-				if (Ex.WriteError.Category != ServerErrorCategory.DuplicateKey)
+				if (ex.WriteError.Category != ServerErrorCategory.DuplicateKey)
 				{
 					throw;
 				}
 			}
-			catch (MongoBulkWriteException Ex)
+			catch (MongoBulkWriteException ex)
 			{
-				if(Ex.WriteErrors.Any(x => x.Category != ServerErrorCategory.DuplicateKey))
+				if(ex.WriteErrors.Any(x => x.Category != ServerErrorCategory.DuplicateKey))
 				{
 					throw;
 				}
@@ -140,19 +140,19 @@ namespace Horde.Build.Utilities
 		/// </summary>
 		/// <typeparam name="TDocument">The document type</typeparam>
 		/// <typeparam name="TField">Type of the field to set</typeparam>
-		/// <param name="UpdateBuilder">Update builder</param>
-		/// <param name="Field">Expression for the field to set</param>
-		/// <param name="Value">New value to set</param>
+		/// <param name="updateBuilder">Update builder</param>
+		/// <param name="field">Expression for the field to set</param>
+		/// <param name="value">New value to set</param>
 		/// <returns>Update defintiion</returns>
-		public static UpdateDefinition<TDocument> SetOrUnsetNull<TDocument, TField>(this UpdateDefinitionBuilder<TDocument> UpdateBuilder, Expression<Func<TDocument, TField?>> Field, TField? Value) where TField : struct
+		public static UpdateDefinition<TDocument> SetOrUnsetNull<TDocument, TField>(this UpdateDefinitionBuilder<TDocument> updateBuilder, Expression<Func<TDocument, TField?>> field, TField? value) where TField : struct
 		{
-			if (Value.HasValue)
+			if (value.HasValue)
 			{
-				return UpdateBuilder.Set(Field, Value.Value);
+				return updateBuilder.Set(field, value.Value);
 			}
 			else
 			{
-				return UpdateBuilder.Unset(new ExpressionFieldDefinition<TDocument, TField?>(Field));
+				return updateBuilder.Unset(new ExpressionFieldDefinition<TDocument, TField?>(field));
 			}
 		}
 
@@ -161,19 +161,19 @@ namespace Horde.Build.Utilities
 		/// </summary>
 		/// <typeparam name="TDocument">The document type</typeparam>
 		/// <typeparam name="TField">Type of the field to set</typeparam>
-		/// <param name="Update">Update builder</param>
-		/// <param name="Field">Expression for the field to set</param>
-		/// <param name="Value">New value to set</param>
+		/// <param name="update">Update builder</param>
+		/// <param name="field">Expression for the field to set</param>
+		/// <param name="value">New value to set</param>
 		/// <returns>Update defintiion</returns>
-		public static UpdateDefinition<TDocument> SetOrUnsetNull<TDocument, TField>(this UpdateDefinition<TDocument> Update, Expression<Func<TDocument, TField?>> Field, TField? Value) where TField : struct
+		public static UpdateDefinition<TDocument> SetOrUnsetNull<TDocument, TField>(this UpdateDefinition<TDocument> update, Expression<Func<TDocument, TField?>> field, TField? value) where TField : struct
 		{
-			if (Value.HasValue)
+			if (value.HasValue)
 			{
-				return Update.Set(Field, Value.Value);
+				return update.Set(field, value.Value);
 			}
 			else
 			{
-				return Update.Unset(new ExpressionFieldDefinition<TDocument, TField?>(Field));
+				return update.Unset(new ExpressionFieldDefinition<TDocument, TField?>(field));
 			}
 		}
 
@@ -181,19 +181,19 @@ namespace Horde.Build.Utilities
 		/// Sets a field to a value, or unsets it if the value is null
 		/// </summary>
 		/// <typeparam name="TDocument">The document type</typeparam>
-		/// <param name="Update">Update builder</param>
-		/// <param name="Field">Expression for the field to set</param>
-		/// <param name="Value">New value to set</param>
+		/// <param name="update">Update builder</param>
+		/// <param name="field">Expression for the field to set</param>
+		/// <param name="value">New value to set</param>
 		/// <returns>Update defintiion</returns>
-		public static UpdateDefinition<TDocument> SetOrUnsetBool<TDocument>(this UpdateDefinition<TDocument> Update, Expression<Func<TDocument, bool>> Field, bool Value)
+		public static UpdateDefinition<TDocument> SetOrUnsetBool<TDocument>(this UpdateDefinition<TDocument> update, Expression<Func<TDocument, bool>> field, bool value)
 		{
-			if (Value)
+			if (value)
 			{
-				return Update.Set(Field, Value);
+				return update.Set(field, value);
 			}
 			else
 			{
-				return Update.Unset(new ExpressionFieldDefinition<TDocument>(Field));
+				return update.Unset(new ExpressionFieldDefinition<TDocument>(field));
 			}
 		}
 
@@ -201,19 +201,19 @@ namespace Horde.Build.Utilities
 		/// Sets a field to a value, or unsets it if the value is null
 		/// </summary>
 		/// <typeparam name="TDocument">The document type</typeparam>
-		/// <param name="UpdateBuilder">Update builder</param>
-		/// <param name="Field">Expression for the field to set</param>
-		/// <param name="Value">New value to set</param>
+		/// <param name="updateBuilder">Update builder</param>
+		/// <param name="field">Expression for the field to set</param>
+		/// <param name="value">New value to set</param>
 		/// <returns>Update defintiion</returns>
-		public static UpdateDefinition<TDocument> SetOrUnsetBool<TDocument>(this UpdateDefinitionBuilder<TDocument> UpdateBuilder, Expression<Func<TDocument, bool>> Field, bool Value)
+		public static UpdateDefinition<TDocument> SetOrUnsetBool<TDocument>(this UpdateDefinitionBuilder<TDocument> updateBuilder, Expression<Func<TDocument, bool>> field, bool value)
 		{
-			if (Value)
+			if (value)
 			{
-				return UpdateBuilder.Set(Field, Value);
+				return updateBuilder.Set(field, value);
 			}
 			else
 			{
-				return UpdateBuilder.Unset(new ExpressionFieldDefinition<TDocument>(Field));
+				return updateBuilder.Unset(new ExpressionFieldDefinition<TDocument>(field));
 			}
 		}
 
@@ -222,19 +222,19 @@ namespace Horde.Build.Utilities
 		/// </summary>
 		/// <typeparam name="TDocument">The document type</typeparam>
 		/// <typeparam name="TField">Type of the field to set</typeparam>
-		/// <param name="UpdateBuilder">Update builder</param>
-		/// <param name="Field">Expression for the field to set</param>
-		/// <param name="Value">New value to set</param>
+		/// <param name="updateBuilder">Update builder</param>
+		/// <param name="field">Expression for the field to set</param>
+		/// <param name="value">New value to set</param>
 		/// <returns>Update defintiion</returns>
-		public static UpdateDefinition<TDocument> SetOrUnsetNullRef<TDocument, TField>(this UpdateDefinitionBuilder<TDocument> UpdateBuilder, Expression<Func<TDocument, TField?>> Field, TField? Value) where TField : class
+		public static UpdateDefinition<TDocument> SetOrUnsetNullRef<TDocument, TField>(this UpdateDefinitionBuilder<TDocument> updateBuilder, Expression<Func<TDocument, TField?>> field, TField? value) where TField : class
 		{
-			if (Value != null)
+			if (value != null)
 			{
-				return UpdateBuilder.Set(Field, Value);
+				return updateBuilder.Set(field, value);
 			}
 			else
 			{
-				return UpdateBuilder.Unset(new ExpressionFieldDefinition<TDocument, TField?>(Field));
+				return updateBuilder.Unset(new ExpressionFieldDefinition<TDocument, TField?>(field));
 			}
 		}
 
@@ -242,33 +242,33 @@ namespace Horde.Build.Utilities
 		/// Creates a filter definition from a linq expression. This is not generally explicitly castable, so expose it as a FilterDefinitionBuilder method.
 		/// </summary>
 		/// <typeparam name="TDocument">The document type</typeparam>
-		/// <param name="Filter">The filter builder</param>
-		/// <param name="Expression">Expression to parse</param>
+		/// <param name="filter">The filter builder</param>
+		/// <param name="expression">Expression to parse</param>
 		/// <returns>New filter definition</returns>
-		public static FilterDefinition<TDocument> Expr<TDocument>(this FilterDefinitionBuilder<TDocument> Filter, Expression<Func<TDocument, bool>> Expression)
+		public static FilterDefinition<TDocument> Expr<TDocument>(this FilterDefinitionBuilder<TDocument> filter, Expression<Func<TDocument, bool>> expression)
 		{
-			return Expression;
+			return expression;
 		}
 
 		/// <summary>
 		/// Converts a JsonElement to a BsonValue
 		/// </summary>
-		/// <param name="Element">The element to convert</param>
+		/// <param name="element">The element to convert</param>
 		/// <returns>Bson value</returns>
-		public static BsonValue ToBsonValue(this JsonElement Element)
+		public static BsonValue ToBsonValue(this JsonElement element)
 		{
-			switch(Element.ValueKind)
+			switch(element.ValueKind)
 			{
 				case JsonValueKind.True:
 					return true;
 				case JsonValueKind.False:
 					return false;
 				case JsonValueKind.String:
-					return Element.GetString();
+					return element.GetString();
 				case JsonValueKind.Array:
-					return new BsonArray(Element.EnumerateArray().Select(x => x.ToBsonValue()));
+					return new BsonArray(element.EnumerateArray().Select(x => x.ToBsonValue()));
 				case JsonValueKind.Object:
-					return new BsonDocument(Element.EnumerateObject().Select(x => new BsonElement(x.Name, x.Value.ToBsonValue())));
+					return new BsonDocument(element.EnumerateObject().Select(x => new BsonElement(x.Name, x.Value.ToBsonValue())));
 				default:
 					return BsonNull.Value;
 			}
diff --git a/Engine/Source/Programs/Horde/Horde.Build/Utilities/MongoQueryCache.cs b/Engine/Source/Programs/Horde/Horde.Build/Utilities/MongoQueryCache.cs
--- a/Engine/Source/Programs/Horde/Horde.Build/Utilities/MongoQueryCache.cs
+++ b/Engine/Source/Programs/Horde/Horde.Build/Utilities/MongoQueryCache.cs
@@ -1,14 +1,14 @@
 // Copyright Epic Games, Inc. All Rights Reserved.
 
-using Microsoft.Extensions.Caching.Memory;
-using MongoDB.Bson;
-using MongoDB.Bson.Serialization;
-using MongoDB.Driver;
 using System;
 using System.Collections.Generic;
 using System.Diagnostics;
 using System.Linq;
 using System.Threading.Tasks;
+using Microsoft.Extensions.Caching.Memory;
+using MongoDB.Bson;
+using MongoDB.Bson.Serialization;
+using MongoDB.Driver;
 
 namespace Horde.Build.Utilities
 {
@@ -20,70 +20,70 @@ namespace Horde.Build.Utilities
 	{
 		class QueryState
 		{
-			public Stopwatch Timer = Stopwatch.StartNew();
-			public List<TDocument>? Results;
-			public Task? QueryTask;
+			public Stopwatch _timer = Stopwatch.StartNew();
+			public List<TDocument>? _results;
+			public Task? _queryTask;
 		}
 
-		IMongoCollection<TDocument> Collection;
-		MemoryCache Cache;
-		TimeSpan MaxLatency;
+		readonly IMongoCollection<TDocument> _collection;
+		readonly MemoryCache _cache;
+		readonly TimeSpan _maxLatency;
 
-		public MongoQueryCache(IMongoCollection<TDocument> Collection, TimeSpan MaxLatency)
+		public MongoQueryCache(IMongoCollection<TDocument> collection, TimeSpan maxLatency)
 		{
-			this.Collection = Collection;
+			_collection = collection;
 
-			MemoryCacheOptions Options = new MemoryCacheOptions();
-			this.Cache = new MemoryCache(Options);
+			MemoryCacheOptions options = new MemoryCacheOptions();
+			_cache = new MemoryCache(options);
 
-			this.MaxLatency = MaxLatency;
+			_maxLatency = maxLatency;
 		}
 
 		public void Dispose()
 		{
-			Cache.Dispose();
+			_cache.Dispose();
 		}
 
-		async Task RefreshAsync(QueryState State, FilterDefinition<TDocument> Filter)
+		async Task RefreshAsync(QueryState state, FilterDefinition<TDocument> filter)
 		{
-			State.Results = await Collection.Find(Filter).ToListAsync();
-			State.Timer.Restart();
+			state._results = await _collection.Find(filter).ToListAsync();
+			state._timer.Restart();
 		}
 
-		public async Task<List<TDocument>> FindAsync(FilterDefinition<TDocument> Filter, int Index, int Count)
+		public async Task<List<TDocument>> FindAsync(FilterDefinition<TDocument> filter, int index, int count)
 		{
-			BsonDocument Rendered = Filter.Render(BsonSerializer.LookupSerializer<TDocument>(), BsonSerializer.SerializerRegistry);
-			BsonDocument Document = new BsonDocument { new BsonElement("filter", Rendered), new BsonElement("index", Index), new BsonElement("count", Count) };
+			BsonDocument rendered = filter.Render(BsonSerializer.LookupSerializer<TDocument>(), BsonSerializer.SerializerRegistry);
+			BsonDocument document = new BsonDocument { new BsonElement("filter", rendered), new BsonElement("index", index), new BsonElement("count", count) };
 
-			string FilterKey = Document.ToString();
+			string filterKey = document.ToString();
 
-			QueryState? State;
-			if (!Cache.TryGetValue(FilterKey, out State) || State == null)
+			QueryState? state;
+			if (!_cache.TryGetValue(filterKey, out state) || state == null)
 			{
-				State = new QueryState();
-				using (ICacheEntry CacheEntry = Cache.CreateEntry(FilterKey))
+				state = new QueryState();
+				using (ICacheEntry cacheEntry = _cache.CreateEntry(filterKey))
 				{
-					CacheEntry.SetSlidingExpiration(TimeSpan.FromMinutes(5.0));
-					CacheEntry.SetValue(State);
+					cacheEntry.SetSlidingExpiration(TimeSpan.FromMinutes(5.0));
+					cacheEntry.SetValue(state);
 				}
 			}
 
-			if(State.QueryTask != null && State.QueryTask.IsCompleted)
+			if(state._queryTask != null && state._queryTask.IsCompleted)
 			{
-				await State.QueryTask;
-				State.QueryTask = null;
+				await state._queryTask;
+				state._queryTask = null;
 			}
 
-			if (State.QueryTask == null && (State.Results == null || State.Timer.Elapsed > MaxLatency))
+			if (state._queryTask == null && (state._results == null || state._timer.Elapsed > _maxLatency))
 			{
-				State.QueryTask = Task.Run(() => RefreshAsync(State, Filter));
+				state._queryTask = Task.Run(() => RefreshAsync(state, filter));
 			}
-			if (State.QueryTask != null && (State.Results == null || State.Timer.Elapsed > MaxLatency))
+			if (state._queryTask != null && (state._results == null || state._timer.Elapsed > _maxLatency))
 			{
-				await State.QueryTask;
+				await state._queryTask;
 			}
 
-			return State.Results!;
+			return state._results!;
 		}
 	}
 }
