Horde: Fix coding conventions within Horde.Build.

diff --git a/Engine/Source/Programs/Horde/Horde.Build/Utilities/MongoTracingCollection.cs b/Engine/Source/Programs/Horde/Horde.Build/Utilities/MongoTracingCollection.cs
--- a/Engine/Source/Programs/Horde/Horde.Build/Utilities/MongoTracingCollection.cs
+++ b/Engine/Source/Programs/Horde/Horde.Build/Utilities/MongoTracingCollection.cs
@@ -20,752 +20,752 @@ namespace Horde.Build.Utiltiies
 	/// <typeparam name="T">A MongoDB document</typeparam>
 	public class MongoTracingCollection<T> : IMongoCollection<T>
 	{
-		private readonly IMongoCollection<T> Collection;
+		private readonly IMongoCollection<T> _collection;
 		private const string Prefix = "MongoCollection.";
 
 		/// <summary>
 		/// Constructor
 		/// </summary>
-		/// <param name="Collection">Collection to wrap with tracing</param>
-		public MongoTracingCollection(IMongoCollection<T> Collection)
+		/// <param name="collection">Collection to wrap with tracing</param>
+		public MongoTracingCollection(IMongoCollection<T> collection)
 		{
-			this.Collection = Collection;
+			_collection = collection;
 		}
 
 #pragma warning disable CS0618		
 		/// <inheritdoc />
-		public IAsyncCursor<TResult> Aggregate<TResult>(PipelineDefinition<T, TResult> pipeline, AggregateOptions Options = null!,
-			CancellationToken  CancellationToken = new CancellationToken())
+		public IAsyncCursor<TResult> Aggregate<TResult>(PipelineDefinition<T, TResult> pipeline, AggregateOptions options = null!,
+			CancellationToken  cancellationToken = new CancellationToken())
 		{
-			return Collection.Aggregate(pipeline, Options,  CancellationToken);
+			return _collection.Aggregate(pipeline, options,  cancellationToken);
 		}
 
 		/// <inheritdoc />
-		public IAsyncCursor<TResult> Aggregate<TResult>(IClientSessionHandle Session, PipelineDefinition<T, TResult> pipeline,
-			AggregateOptions Options = null!, CancellationToken  CancellationToken = new CancellationToken())
+		public IAsyncCursor<TResult> Aggregate<TResult>(IClientSessionHandle session, PipelineDefinition<T, TResult> pipeline,
+			AggregateOptions options = null!, CancellationToken  cancellationToken = new CancellationToken())
 		{
-			return Collection.Aggregate(Session, pipeline, Options,  CancellationToken);
+			return _collection.Aggregate(session, pipeline, options,  cancellationToken);
 		}
 
 		/// <inheritdoc />
-		public async Task<IAsyncCursor<TResult>> AggregateAsync<TResult>(PipelineDefinition<T, TResult> pipeline, AggregateOptions Options = null!,
-			CancellationToken  CancellationToken = new CancellationToken())
+		public async Task<IAsyncCursor<TResult>> AggregateAsync<TResult>(PipelineDefinition<T, TResult> pipeline, AggregateOptions options = null!,
+			CancellationToken  cancellationToken = new CancellationToken())
 		{
-			using IScope Scope = GlobalTracer.Instance.BuildSpan(Prefix + "AggregateAsync").StartActive();
-			Scope.Span.SetTag("CollectionName", Collection.CollectionNamespace.CollectionName);
-			return await Collection.AggregateAsync(pipeline, Options,  CancellationToken);
+			using IScope scope = GlobalTracer.Instance.BuildSpan(Prefix + "AggregateAsync").StartActive();
+			scope.Span.SetTag("CollectionName", _collection.CollectionNamespace.CollectionName);
+			return await _collection.AggregateAsync(pipeline, options,  cancellationToken);
 		}
 
 		/// <inheritdoc />
-		public async Task<IAsyncCursor<TResult>> AggregateAsync<TResult>(IClientSessionHandle Session, PipelineDefinition<T, TResult> Pipeline, AggregateOptions Options = null!,
-			CancellationToken  CancellationToken = new CancellationToken())
+		public async Task<IAsyncCursor<TResult>> AggregateAsync<TResult>(IClientSessionHandle session, PipelineDefinition<T, TResult> pipeline, AggregateOptions options = null!,
+			CancellationToken  cancellationToken = new CancellationToken())
 		{
-			using IScope Scope = GlobalTracer.Instance.BuildSpan(Prefix + "AggregateAsync").StartActive();
-			Scope.Span.SetTag("CollectionName", Collection.CollectionNamespace.CollectionName);
-			return await Collection.AggregateAsync(Session, Pipeline, Options, CancellationToken);
+			using IScope scope = GlobalTracer.Instance.BuildSpan(Prefix + "AggregateAsync").StartActive();
+			scope.Span.SetTag("CollectionName", _collection.CollectionNamespace.CollectionName);
+			return await _collection.AggregateAsync(session, pipeline, options, cancellationToken);
 		}
 
 		/// <inheritdoc />
-		public void AggregateToCollection<TResult>(PipelineDefinition<T, TResult> Pipeline, AggregateOptions? Options = null, CancellationToken CancellationToken = default(CancellationToken))
+		public void AggregateToCollection<TResult>(PipelineDefinition<T, TResult> pipeline, AggregateOptions? options = null, CancellationToken cancellationToken = default)
 		{
-			using IScope Scope = GlobalTracer.Instance.BuildSpan(Prefix + "AggregateToCollection").StartActive();
-			Scope.Span.SetTag("CollectionName", Collection.CollectionNamespace.CollectionName);
-			Collection.AggregateToCollection(Pipeline, Options, CancellationToken);
+			using IScope scope = GlobalTracer.Instance.BuildSpan(Prefix + "AggregateToCollection").StartActive();
+			scope.Span.SetTag("CollectionName", _collection.CollectionNamespace.CollectionName);
+			_collection.AggregateToCollection(pipeline, options, cancellationToken);
 		}
 
 		/// <inheritdoc />
-		public void AggregateToCollection<TResult>(IClientSessionHandle Session, PipelineDefinition<T, TResult> pipeline, AggregateOptions? options = null, CancellationToken cancellationToken = default(CancellationToken))
+		public void AggregateToCollection<TResult>(IClientSessionHandle session, PipelineDefinition<T, TResult> pipeline, AggregateOptions? options = null, CancellationToken cancellationToken = default)
 		{
-			using IScope Scope = GlobalTracer.Instance.BuildSpan(Prefix + "AggregateToCollection").StartActive();
-			Scope.Span.SetTag("CollectionName", Collection.CollectionNamespace.CollectionName);
-			Collection.AggregateToCollection(Session, pipeline, options, cancellationToken);
+			using IScope scope = GlobalTracer.Instance.BuildSpan(Prefix + "AggregateToCollection").StartActive();
+			scope.Span.SetTag("CollectionName", _collection.CollectionNamespace.CollectionName);
+			_collection.AggregateToCollection(session, pipeline, options, cancellationToken);
 		}
 
 		/// <inheritdoc />
-		public async Task AggregateToCollectionAsync<TResult>(PipelineDefinition<T, TResult> Pipeline, AggregateOptions? Options = null, CancellationToken CancellationToken = default(CancellationToken))
+		public async Task AggregateToCollectionAsync<TResult>(PipelineDefinition<T, TResult> pipeline, AggregateOptions? options = null, CancellationToken cancellationToken = default)
 		{
-			using IScope Scope = GlobalTracer.Instance.BuildSpan(Prefix + "AggregateToCollectionAsync").StartActive();
-			Scope.Span.SetTag("CollectionName", Collection.CollectionNamespace.CollectionName);
-			await Collection.AggregateToCollectionAsync(Pipeline, Options, CancellationToken);
+			using IScope scope = GlobalTracer.Instance.BuildSpan(Prefix + "AggregateToCollectionAsync").StartActive();
+			scope.Span.SetTag("CollectionName", _collection.CollectionNamespace.CollectionName);
+			await _collection.AggregateToCollectionAsync(pipeline, options, cancellationToken);
 		}
 
 		/// <inheritdoc />
-		public async Task AggregateToCollectionAsync<TResult>(IClientSessionHandle Session, PipelineDefinition<T, TResult> Pipeline, AggregateOptions? Options = null, CancellationToken CancellationToken = default(CancellationToken))
+		public async Task AggregateToCollectionAsync<TResult>(IClientSessionHandle session, PipelineDefinition<T, TResult> pipeline, AggregateOptions? options = null, CancellationToken cancellationToken = default)
 		{
-			using IScope Scope = GlobalTracer.Instance.BuildSpan(Prefix + "AggregateToCollectionAsync").StartActive();
-			Scope.Span.SetTag("CollectionName", Collection.CollectionNamespace.CollectionName);
-			await Collection.AggregateToCollectionAsync(Session, Pipeline, Options, CancellationToken);
+			using IScope scope = GlobalTracer.Instance.BuildSpan(Prefix + "AggregateToCollectionAsync").StartActive();
+			scope.Span.SetTag("CollectionName", _collection.CollectionNamespace.CollectionName);
+			await _collection.AggregateToCollectionAsync(session, pipeline, options, cancellationToken);
 		}
 
 		/// <inheritdoc />
-		public BulkWriteResult<T> BulkWrite(IEnumerable<WriteModel<T>> requests, BulkWriteOptions Options = null!,
-			CancellationToken  CancellationToken = new CancellationToken())
+		public BulkWriteResult<T> BulkWrite(IEnumerable<WriteModel<T>> requests, BulkWriteOptions options = null!,
+			CancellationToken  cancellationToken = new CancellationToken())
 		{
-			return Collection.BulkWrite(requests, Options,  CancellationToken);
+			return _collection.BulkWrite(requests, options,  cancellationToken);
 		}
 
 		/// <inheritdoc />
-		public BulkWriteResult<T> BulkWrite(IClientSessionHandle Session, IEnumerable<WriteModel<T>> requests, BulkWriteOptions Options = null!,
-			CancellationToken  CancellationToken = new CancellationToken())
+		public BulkWriteResult<T> BulkWrite(IClientSessionHandle session, IEnumerable<WriteModel<T>> requests, BulkWriteOptions options = null!,
+			CancellationToken  cancellationToken = new CancellationToken())
 		{
-			return Collection.BulkWrite(Session, requests, Options,  CancellationToken);
+			return _collection.BulkWrite(session, requests, options,  cancellationToken);
 		}
 
 		/// <inheritdoc />
-		public async Task<BulkWriteResult<T>> BulkWriteAsync(IEnumerable<WriteModel<T>> requests, BulkWriteOptions Options = null!,
-			CancellationToken  CancellationToken = new CancellationToken())
+		public async Task<BulkWriteResult<T>> BulkWriteAsync(IEnumerable<WriteModel<T>> requests, BulkWriteOptions options = null!,
+			CancellationToken  cancellationToken = new CancellationToken())
 		{
-			using IScope Scope = GlobalTracer.Instance.BuildSpan(Prefix + "BulkWriteAsync").StartActive();
-			Scope.Span.SetTag("CollectionName", Collection.CollectionNamespace.CollectionName);
-			return await Collection.BulkWriteAsync(requests, Options,  CancellationToken);
+			using IScope scope = GlobalTracer.Instance.BuildSpan(Prefix + "BulkWriteAsync").StartActive();
+			scope.Span.SetTag("CollectionName", _collection.CollectionNamespace.CollectionName);
+			return await _collection.BulkWriteAsync(requests, options,  cancellationToken);
 		}
 
 		/// <inheritdoc />
-		public async Task<BulkWriteResult<T>> BulkWriteAsync(IClientSessionHandle Session, IEnumerable<WriteModel<T>> requests, BulkWriteOptions Options = null!,
-			CancellationToken  CancellationToken = new CancellationToken())
+		public async Task<BulkWriteResult<T>> BulkWriteAsync(IClientSessionHandle session, IEnumerable<WriteModel<T>> requests, BulkWriteOptions options = null!,
+			CancellationToken  cancellationToken = new CancellationToken())
 		{
-			using IScope Scope = GlobalTracer.Instance.BuildSpan(Prefix + "BulkWriteAsync").StartActive();
-			Scope.Span.SetTag("CollectionName", Collection.CollectionNamespace.CollectionName);
-			return await Collection.BulkWriteAsync(Session, requests, Options,  CancellationToken);
+			using IScope scope = GlobalTracer.Instance.BuildSpan(Prefix + "BulkWriteAsync").StartActive();
+			scope.Span.SetTag("CollectionName", _collection.CollectionNamespace.CollectionName);
+			return await _collection.BulkWriteAsync(session, requests, options,  cancellationToken);
 		}
 
 		/// <inheritdoc />
-		public long Count(FilterDefinition<T> filter, CountOptions Options = null!,
-			CancellationToken  CancellationToken = new CancellationToken())
+		public long Count(FilterDefinition<T> filter, CountOptions options = null!,
+			CancellationToken  cancellationToken = new CancellationToken())
 		{
-			return Collection.Count(filter, Options,  CancellationToken);
+			return _collection.Count(filter, options,  cancellationToken);
 		}
 
 		/// <inheritdoc />
-		public long Count(IClientSessionHandle Session, FilterDefinition<T> filter, CountOptions Options = null!,
-			CancellationToken  CancellationToken = new CancellationToken())
+		public long Count(IClientSessionHandle session, FilterDefinition<T> filter, CountOptions options = null!,
+			CancellationToken  cancellationToken = new CancellationToken())
 		{
-			return Collection.Count(Session, filter, Options,  CancellationToken);
+			return _collection.Count(session, filter, options,  cancellationToken);
 		}
 
 		/// <inheritdoc />
-		public async Task<long> CountAsync(FilterDefinition<T> filter, CountOptions Options = null!,
-			CancellationToken  CancellationToken = new CancellationToken())
+		public async Task<long> CountAsync(FilterDefinition<T> filter, CountOptions options = null!,
+			CancellationToken  cancellationToken = new CancellationToken())
 		{
-			using IScope Scope = GlobalTracer.Instance.BuildSpan(Prefix + "CountAsync").StartActive();
-			Scope.Span.SetTag("CollectionName", Collection.CollectionNamespace.CollectionName);
-			return await Collection.CountAsync(filter, Options,  CancellationToken);
+			using IScope scope = GlobalTracer.Instance.BuildSpan(Prefix + "CountAsync").StartActive();
+			scope.Span.SetTag("CollectionName", _collection.CollectionNamespace.CollectionName);
+			return await _collection.CountAsync(filter, options,  cancellationToken);
 		}
 
 		/// <inheritdoc />
-		public async Task<long> CountAsync(IClientSessionHandle Session, FilterDefinition<T> filter, CountOptions Options = null!,
-			CancellationToken  CancellationToken = new CancellationToken())
+		public async Task<long> CountAsync(IClientSessionHandle session, FilterDefinition<T> filter, CountOptions options = null!,
+			CancellationToken  cancellationToken = new CancellationToken())
 		{
-			using IScope Scope = GlobalTracer.Instance.BuildSpan(Prefix + "CountAsync").StartActive();
-			Scope.Span.SetTag("CollectionName", Collection.CollectionNamespace.CollectionName);
-			return await Collection.CountAsync(Session, filter, Options,  CancellationToken);
+			using IScope scope = GlobalTracer.Instance.BuildSpan(Prefix + "CountAsync").StartActive();
+			scope.Span.SetTag("CollectionName", _collection.CollectionNamespace.CollectionName);
+			return await _collection.CountAsync(session, filter, options,  cancellationToken);
 		}
 
 		/// <inheritdoc />
-		public long CountDocuments(FilterDefinition<T> filter, CountOptions Options = null!,
-			CancellationToken  CancellationToken = new CancellationToken())
+		public long CountDocuments(FilterDefinition<T> filter, CountOptions options = null!,
+			CancellationToken  cancellationToken = new CancellationToken())
 		{
-			return Collection.CountDocuments(filter, Options,  CancellationToken);
+			return _collection.CountDocuments(filter, options,  cancellationToken);
 		}
 
 		/// <inheritdoc />
-		public long CountDocuments(IClientSessionHandle Session, FilterDefinition<T> filter, CountOptions Options = null!,
-			CancellationToken  CancellationToken = new CancellationToken())
+		public long CountDocuments(IClientSessionHandle session, FilterDefinition<T> filter, CountOptions options = null!,
+			CancellationToken  cancellationToken = new CancellationToken())
 		{
-			using IScope Scope = GlobalTracer.Instance.BuildSpan(Prefix + "CountDocumentsAsync").StartActive();
-			Scope.Span.SetTag("CollectionName", Collection.CollectionNamespace.CollectionName);
-			return Collection.CountDocuments(Session, filter, Options,  CancellationToken);
+			using IScope scope = GlobalTracer.Instance.BuildSpan(Prefix + "CountDocumentsAsync").StartActive();
+			scope.Span.SetTag("CollectionName", _collection.CollectionNamespace.CollectionName);
+			return _collection.CountDocuments(session, filter, options,  cancellationToken);
 		}
 
 		/// <inheritdoc />
-		public async Task<long> CountDocumentsAsync(FilterDefinition<T> filter, CountOptions Options = null!,
-			CancellationToken  CancellationToken = new CancellationToken())
+		public async Task<long> CountDocumentsAsync(FilterDefinition<T> filter, CountOptions options = null!,
+			CancellationToken  cancellationToken = new CancellationToken())
 		{
-			using IScope Scope = GlobalTracer.Instance.BuildSpan(Prefix + "CountDocumentsAsync").StartActive();
-			Scope.Span.SetTag("CollectionName", Collection.CollectionNamespace.CollectionName);
-			return await Collection.CountDocumentsAsync(filter, Options,  CancellationToken);
+			using IScope scope = GlobalTracer.Instance.BuildSpan(Prefix + "CountDocumentsAsync").StartActive();
+			scope.Span.SetTag("CollectionName", _collection.CollectionNamespace.CollectionName);
+			return await _collection.CountDocumentsAsync(filter, options,  cancellationToken);
 		}
 
 		/// <inheritdoc />
-		public async Task<long> CountDocumentsAsync(IClientSessionHandle Session, FilterDefinition<T> filter, CountOptions Options = null!,
-			CancellationToken  CancellationToken = new CancellationToken())
+		public async Task<long> CountDocumentsAsync(IClientSessionHandle session, FilterDefinition<T> filter, CountOptions options = null!,
+			CancellationToken  cancellationToken = new CancellationToken())
 		{
-			using IScope Scope = GlobalTracer.Instance.BuildSpan(Prefix + "CountDocumentsAsync").StartActive();
-			Scope.Span.SetTag("CollectionName", Collection.CollectionNamespace.CollectionName);
-			return await Collection.CountDocumentsAsync(Session, filter, Options,  CancellationToken);
+			using IScope scope = GlobalTracer.Instance.BuildSpan(Prefix + "CountDocumentsAsync").StartActive();
+			scope.Span.SetTag("CollectionName", _collection.CollectionNamespace.CollectionName);
+			return await _collection.CountDocumentsAsync(session, filter, options,  cancellationToken);
 		}
 
 		/// <inheritdoc />
-		public DeleteResult DeleteMany(FilterDefinition<T> filter, CancellationToken  CancellationToken = new CancellationToken())
+		public DeleteResult DeleteMany(FilterDefinition<T> filter, CancellationToken  cancellationToken = new CancellationToken())
 		{
-			return Collection.DeleteMany(filter,  CancellationToken);
+			return _collection.DeleteMany(filter,  cancellationToken);
 		}
 
 		/// <inheritdoc />
-		public DeleteResult DeleteMany(FilterDefinition<T> filter, DeleteOptions Options,
-			CancellationToken  CancellationToken = new CancellationToken())
+		public DeleteResult DeleteMany(FilterDefinition<T> filter, DeleteOptions options,
+			CancellationToken  cancellationToken = new CancellationToken())
 		{
-			return Collection.DeleteMany(filter, Options,  CancellationToken);
+			return _collection.DeleteMany(filter, options,  cancellationToken);
 		}
 
 		/// <inheritdoc />
-		public DeleteResult DeleteMany(IClientSessionHandle Session, FilterDefinition<T> filter, DeleteOptions Options = null!,
-			CancellationToken  CancellationToken = new CancellationToken())
+		public DeleteResult DeleteMany(IClientSessionHandle session, FilterDefinition<T> filter, DeleteOptions options = null!,
+			CancellationToken  cancellationToken = new CancellationToken())
 		{
-			return Collection.DeleteMany(Session, filter, Options,  CancellationToken);
+			return _collection.DeleteMany(session, filter, options,  cancellationToken);
 		}
 
 		/// <inheritdoc />
-		public async Task<DeleteResult> DeleteManyAsync(FilterDefinition<T> filter, CancellationToken  CancellationToken = new CancellationToken())
+		public async Task<DeleteResult> DeleteManyAsync(FilterDefinition<T> filter, CancellationToken  cancellationToken = new CancellationToken())
 		{
-			using IScope Scope = GlobalTracer.Instance.BuildSpan(Prefix + "DeleteManyAsync").StartActive();
-			Scope.Span.SetTag("CollectionName", Collection.CollectionNamespace.CollectionName);
-			return await Collection.DeleteManyAsync(filter,  CancellationToken);
+			using IScope scope = GlobalTracer.Instance.BuildSpan(Prefix + "DeleteManyAsync").StartActive();
+			scope.Span.SetTag("CollectionName", _collection.CollectionNamespace.CollectionName);
+			return await _collection.DeleteManyAsync(filter,  cancellationToken);
 		}
 
 		/// <inheritdoc />
-		public async Task<DeleteResult> DeleteManyAsync(FilterDefinition<T> filter, DeleteOptions Options,
-			CancellationToken  CancellationToken = new CancellationToken())
+		public async Task<DeleteResult> DeleteManyAsync(FilterDefinition<T> filter, DeleteOptions options,
+			CancellationToken  cancellationToken = new CancellationToken())
 		{
-			using IScope Scope = GlobalTracer.Instance.BuildSpan(Prefix + "DeleteManyAsync").StartActive();
-			Scope.Span.SetTag("CollectionName", Collection.CollectionNamespace.CollectionName);
-			return await Collection.DeleteManyAsync(filter, Options,  CancellationToken);
+			using IScope scope = GlobalTracer.Instance.BuildSpan(Prefix + "DeleteManyAsync").StartActive();
+			scope.Span.SetTag("CollectionName", _collection.CollectionNamespace.CollectionName);
+			return await _collection.DeleteManyAsync(filter, options,  cancellationToken);
 		}
 
 		/// <inheritdoc />
-		public async Task<DeleteResult> DeleteManyAsync(IClientSessionHandle Session, FilterDefinition<T> filter, DeleteOptions Options = null!,
-			CancellationToken  CancellationToken = new CancellationToken())
+		public async Task<DeleteResult> DeleteManyAsync(IClientSessionHandle session, FilterDefinition<T> filter, DeleteOptions options = null!,
+			CancellationToken  cancellationToken = new CancellationToken())
 		{
-			using IScope Scope = GlobalTracer.Instance.BuildSpan(Prefix + "DeleteManyAsync").StartActive();
-			Scope.Span.SetTag("CollectionName", Collection.CollectionNamespace.CollectionName);
-			return await Collection.DeleteManyAsync(Session, filter, Options,  CancellationToken);
+			using IScope scope = GlobalTracer.Instance.BuildSpan(Prefix + "DeleteManyAsync").StartActive();
+			scope.Span.SetTag("CollectionName", _collection.CollectionNamespace.CollectionName);
+			return await _collection.DeleteManyAsync(session, filter, options,  cancellationToken);
 		}
 
 		/// <inheritdoc />
-		public DeleteResult DeleteOne(FilterDefinition<T> filter, CancellationToken  CancellationToken = new CancellationToken())
+		public DeleteResult DeleteOne(FilterDefinition<T> filter, CancellationToken  cancellationToken = new CancellationToken())
 		{
-			return Collection.DeleteOne(filter,  CancellationToken);
+			return _collection.DeleteOne(filter,  cancellationToken);
 		}
 
 		/// <inheritdoc />
-		public DeleteResult DeleteOne(FilterDefinition<T> filter, DeleteOptions Options,
-			CancellationToken  CancellationToken = new CancellationToken())
+		public DeleteResult DeleteOne(FilterDefinition<T> filter, DeleteOptions options,
+			CancellationToken  cancellationToken = new CancellationToken())
 		{
-			return Collection.DeleteOne(filter, Options,  CancellationToken);
+			return _collection.DeleteOne(filter, options,  cancellationToken);
 		}
 
 		/// <inheritdoc />
-		public DeleteResult DeleteOne(IClientSessionHandle Session, FilterDefinition<T> filter, DeleteOptions Options = null!,
-			CancellationToken  CancellationToken = new CancellationToken())
+		public DeleteResult DeleteOne(IClientSessionHandle session, FilterDefinition<T> filter, DeleteOptions options = null!,
+			CancellationToken  cancellationToken = new CancellationToken())
 		{
-			return Collection.DeleteOne(Session, filter, Options,  CancellationToken);
+			return _collection.DeleteOne(session, filter, options,  cancellationToken);
 		}
 
 		/// <inheritdoc />
-		public async Task<DeleteResult> DeleteOneAsync(FilterDefinition<T> filter, CancellationToken  CancellationToken = new CancellationToken())
+		public async Task<DeleteResult> DeleteOneAsync(FilterDefinition<T> filter, CancellationToken  cancellationToken = new CancellationToken())
 		{
-			using IScope Scope = GlobalTracer.Instance.BuildSpan(Prefix + "DeleteOneAsync").StartActive();
-			Scope.Span.SetTag("CollectionName", Collection.CollectionNamespace.CollectionName);
-			return await Collection.DeleteOneAsync(filter,  CancellationToken);
+			using IScope scope = GlobalTracer.Instance.BuildSpan(Prefix + "DeleteOneAsync").StartActive();
+			scope.Span.SetTag("CollectionName", _collection.CollectionNamespace.CollectionName);
+			return await _collection.DeleteOneAsync(filter,  cancellationToken);
 		}
 
 		/// <inheritdoc />
-		public async Task<DeleteResult> DeleteOneAsync(FilterDefinition<T> filter, DeleteOptions Options,
-			CancellationToken  CancellationToken = new CancellationToken())
+		public async Task<DeleteResult> DeleteOneAsync(FilterDefinition<T> filter, DeleteOptions options,
+			CancellationToken  cancellationToken = new CancellationToken())
 		{
-			using IScope Scope = GlobalTracer.Instance.BuildSpan(Prefix + "DeleteOneAsync").StartActive();
-			Scope.Span.SetTag("CollectionName", Collection.CollectionNamespace.CollectionName);
-			return await Collection.DeleteOneAsync(filter, Options,  CancellationToken);
+			using IScope scope = GlobalTracer.Instance.BuildSpan(Prefix + "DeleteOneAsync").StartActive();
+			scope.Span.SetTag("CollectionName", _collection.CollectionNamespace.CollectionName);
+			return await _collection.DeleteOneAsync(filter, options,  cancellationToken);
 		}
 
 		/// <inheritdoc />
-		public Task<DeleteResult> DeleteOneAsync(IClientSessionHandle Session, FilterDefinition<T> filter, DeleteOptions Options = null!,
-			CancellationToken  CancellationToken = new CancellationToken())
+		public Task<DeleteResult> DeleteOneAsync(IClientSessionHandle session, FilterDefinition<T> filter, DeleteOptions options = null!,
+			CancellationToken  cancellationToken = new CancellationToken())
 		{
-			using IScope Scope = GlobalTracer.Instance.BuildSpan(Prefix + "DeleteOneAsync").StartActive();
-			Scope.Span.SetTag("CollectionName", Collection.CollectionNamespace.CollectionName);
-			return Collection.DeleteOneAsync(Session, filter, Options,  CancellationToken);
+			using IScope scope = GlobalTracer.Instance.BuildSpan(Prefix + "DeleteOneAsync").StartActive();
+			scope.Span.SetTag("CollectionName", _collection.CollectionNamespace.CollectionName);
+			return _collection.DeleteOneAsync(session, filter, options,  cancellationToken);
 		}
 
 		/// <inheritdoc />
-		public IAsyncCursor<TField> Distinct<TField>(FieldDefinition<T, TField> field, FilterDefinition<T> filter, DistinctOptions Options = null!,
-			CancellationToken  CancellationToken = new CancellationToken())
+		public IAsyncCursor<TField> Distinct<TField>(FieldDefinition<T, TField> field, FilterDefinition<T> filter, DistinctOptions options = null!,
+			CancellationToken  cancellationToken = new CancellationToken())
 		{
-			return Collection.Distinct(field, filter, Options,  CancellationToken);
+			return _collection.Distinct(field, filter, options,  cancellationToken);
 		}
 
 		/// <inheritdoc />
-		public IAsyncCursor<TField> Distinct<TField>(IClientSessionHandle Session, FieldDefinition<T, TField> field, FilterDefinition<T> filter,
-			DistinctOptions Options = null!, CancellationToken  CancellationToken = new CancellationToken())
+		public IAsyncCursor<TField> Distinct<TField>(IClientSessionHandle session, FieldDefinition<T, TField> field, FilterDefinition<T> filter,
+			DistinctOptions options = null!, CancellationToken  cancellationToken = new CancellationToken())
 		{
-			return Collection.Distinct(Session, field, filter, Options,  CancellationToken);
+			return _collection.Distinct(session, field, filter, options,  cancellationToken);
 		}
 
 		/// <inheritdoc />
-		public async Task<IAsyncCursor<TField>> DistinctAsync<TField>(FieldDefinition<T, TField> field, FilterDefinition<T> filter, DistinctOptions Options = null!,
-			CancellationToken  CancellationToken = new CancellationToken())
+		public async Task<IAsyncCursor<TField>> DistinctAsync<TField>(FieldDefinition<T, TField> field, FilterDefinition<T> filter, DistinctOptions options = null!,
+			CancellationToken  cancellationToken = new CancellationToken())
 		{
-			using IScope Scope = GlobalTracer.Instance.BuildSpan(Prefix + "DistinctAsync").StartActive();
-			Scope.Span.SetTag("CollectionName", Collection.CollectionNamespace.CollectionName);
-			return await Collection.DistinctAsync(field, filter, Options,  CancellationToken);
+			using IScope scope = GlobalTracer.Instance.BuildSpan(Prefix + "DistinctAsync").StartActive();
+			scope.Span.SetTag("CollectionName", _collection.CollectionNamespace.CollectionName);
+			return await _collection.DistinctAsync(field, filter, options,  cancellationToken);
 		}
 
 		/// <inheritdoc />
-		public async Task<IAsyncCursor<TField>> DistinctAsync<TField>(IClientSessionHandle Session, FieldDefinition<T, TField> field, FilterDefinition<T> filter,
-			DistinctOptions Options = null!, CancellationToken  CancellationToken = new CancellationToken())
+		public async Task<IAsyncCursor<TField>> DistinctAsync<TField>(IClientSessionHandle session, FieldDefinition<T, TField> field, FilterDefinition<T> filter,
+			DistinctOptions options = null!, CancellationToken  cancellationToken = new CancellationToken())
 		{
-			using IScope Scope = GlobalTracer.Instance.BuildSpan(Prefix + "DistinctAsync").StartActive();
-			Scope.Span.SetTag("CollectionName", Collection.CollectionNamespace.CollectionName);
-			return await Collection.DistinctAsync(Session, field, filter, Options,  CancellationToken);
+			using IScope scope = GlobalTracer.Instance.BuildSpan(Prefix + "DistinctAsync").StartActive();
+			scope.Span.SetTag("CollectionName", _collection.CollectionNamespace.CollectionName);
+			return await _collection.DistinctAsync(session, field, filter, options,  cancellationToken);
 		}
 
 		/// <inheritdoc />
-		public long EstimatedDocumentCount(EstimatedDocumentCountOptions Options = null!,
-			CancellationToken  CancellationToken = new CancellationToken())
+		public long EstimatedDocumentCount(EstimatedDocumentCountOptions options = null!,
+			CancellationToken  cancellationToken = new CancellationToken())
 		{
-			return Collection.EstimatedDocumentCount(Options,  CancellationToken);
+			return _collection.EstimatedDocumentCount(options,  cancellationToken);
 		}
 
 		/// <inheritdoc />
-		public async Task<long> EstimatedDocumentCountAsync(EstimatedDocumentCountOptions Options = null!,
-			CancellationToken  CancellationToken = new CancellationToken())
+		public async Task<long> EstimatedDocumentCountAsync(EstimatedDocumentCountOptions options = null!,
+			CancellationToken  cancellationToken = new CancellationToken())
 		{
-			using IScope Scope = GlobalTracer.Instance.BuildSpan(Prefix + "EstimatedDocumentCountAsync").StartActive();
-			Scope.Span.SetTag("CollectionName", Collection.CollectionNamespace.CollectionName);
-			return await Collection.EstimatedDocumentCountAsync(Options,  CancellationToken);
+			using IScope scope = GlobalTracer.Instance.BuildSpan(Prefix + "EstimatedDocumentCountAsync").StartActive();
+			scope.Span.SetTag("CollectionName", _collection.CollectionNamespace.CollectionName);
+			return await _collection.EstimatedDocumentCountAsync(options,  cancellationToken);
 		}
 
 		/// <inheritdoc />
-		public IAsyncCursor<TProjection> FindSync<TProjection>(FilterDefinition<T> filter, FindOptions<T, TProjection> Options = null!,
-			CancellationToken  CancellationToken = new CancellationToken())
+		public IAsyncCursor<TProjection> FindSync<TProjection>(FilterDefinition<T> filter, FindOptions<T, TProjection> options = null!,
+			CancellationToken  cancellationToken = new CancellationToken())
 		{
-			return Collection.FindSync(filter, Options,  CancellationToken);
+			return _collection.FindSync(filter, options,  cancellationToken);
 		}
 
 		/// <inheritdoc />
-		public IAsyncCursor<TProjection> FindSync<TProjection>(IClientSessionHandle Session, FilterDefinition<T> filter, FindOptions<T, TProjection> Options = null!,
-			CancellationToken  CancellationToken = new CancellationToken())
+		public IAsyncCursor<TProjection> FindSync<TProjection>(IClientSessionHandle session, FilterDefinition<T> filter, FindOptions<T, TProjection> options = null!,
+			CancellationToken  cancellationToken = new CancellationToken())
 		{
-			return Collection.FindSync(Session, filter, Options,  CancellationToken);
+			return _collection.FindSync(session, filter, options,  cancellationToken);
 		}
 
 		/// <inheritdoc />
-		public async Task<IAsyncCursor<TProjection>> FindAsync<TProjection>(FilterDefinition<T> filter, FindOptions<T, TProjection> Options = null!,
-			CancellationToken  CancellationToken = new CancellationToken())
+		public async Task<IAsyncCursor<TProjection>> FindAsync<TProjection>(FilterDefinition<T> filter, FindOptions<T, TProjection> options = null!,
+			CancellationToken  cancellationToken = new CancellationToken())
 		{
-			using IScope Scope = GlobalTracer.Instance.BuildSpan(Prefix + "FindAsync").StartActive();
-			Scope.Span.SetTag("CollectionName", Collection.CollectionNamespace.CollectionName);
-			return await Collection.FindAsync(filter, Options,  CancellationToken);
+			using IScope scope = GlobalTracer.Instance.BuildSpan(Prefix + "FindAsync").StartActive();
+			scope.Span.SetTag("CollectionName", _collection.CollectionNamespace.CollectionName);
+			return await _collection.FindAsync(filter, options,  cancellationToken);
 		}
 
 		/// <inheritdoc />
-		public async Task<IAsyncCursor<TProjection>> FindAsync<TProjection>(IClientSessionHandle Session, FilterDefinition<T> filter, FindOptions<T, TProjection> Options = null!,
-			CancellationToken  CancellationToken = new CancellationToken())
+		public async Task<IAsyncCursor<TProjection>> FindAsync<TProjection>(IClientSessionHandle session, FilterDefinition<T> filter, FindOptions<T, TProjection> options = null!,
+			CancellationToken  cancellationToken = new CancellationToken())
 		{
-			using IScope Scope = GlobalTracer.Instance.BuildSpan(Prefix + "FindAsync").StartActive();
-			Scope.Span.SetTag("CollectionName", Collection.CollectionNamespace.CollectionName);
-			return await Collection.FindAsync(Session, filter, Options,  CancellationToken);
+			using IScope scope = GlobalTracer.Instance.BuildSpan(Prefix + "FindAsync").StartActive();
+			scope.Span.SetTag("CollectionName", _collection.CollectionNamespace.CollectionName);
+			return await _collection.FindAsync(session, filter, options,  cancellationToken);
 		}
 
 		/// <inheritdoc />
-		public TProjection FindOneAndDelete<TProjection>(FilterDefinition<T> filter, FindOneAndDeleteOptions<T, TProjection> Options = null!,
-			CancellationToken  CancellationToken = new CancellationToken())
+		public TProjection FindOneAndDelete<TProjection>(FilterDefinition<T> filter, FindOneAndDeleteOptions<T, TProjection> options = null!,
+			CancellationToken  cancellationToken = new CancellationToken())
 		{
-			return Collection.FindOneAndDelete(filter, Options,  CancellationToken);
+			return _collection.FindOneAndDelete(filter, options,  cancellationToken);
 		}
 
 		/// <inheritdoc />
-		public TProjection FindOneAndDelete<TProjection>(IClientSessionHandle Session, FilterDefinition<T> filter,
-			FindOneAndDeleteOptions<T, TProjection> Options = null!, CancellationToken  CancellationToken = new CancellationToken())
+		public TProjection FindOneAndDelete<TProjection>(IClientSessionHandle session, FilterDefinition<T> filter,
+			FindOneAndDeleteOptions<T, TProjection> options = null!, CancellationToken  cancellationToken = new CancellationToken())
 		{
-			return Collection.FindOneAndDelete(Session, filter, Options,  CancellationToken);
+			return _collection.FindOneAndDelete(session, filter, options,  cancellationToken);
 		}
 
 		/// <inheritdoc />
-		public async Task<TProjection> FindOneAndDeleteAsync<TProjection>(FilterDefinition<T> filter, FindOneAndDeleteOptions<T, TProjection> Options = null!,
-			CancellationToken  CancellationToken = new CancellationToken())
+		public async Task<TProjection> FindOneAndDeleteAsync<TProjection>(FilterDefinition<T> filter, FindOneAndDeleteOptions<T, TProjection> options = null!,
+			CancellationToken  cancellationToken = new CancellationToken())
 		{
-			using IScope Scope = GlobalTracer.Instance.BuildSpan(Prefix + "FindOneAndDeleteAsync").StartActive();
-			Scope.Span.SetTag("CollectionName", Collection.CollectionNamespace.CollectionName);
-			return await Collection.FindOneAndDeleteAsync(filter, Options,  CancellationToken);
+			using IScope scope = GlobalTracer.Instance.BuildSpan(Prefix + "FindOneAndDeleteAsync").StartActive();
+			scope.Span.SetTag("CollectionName", _collection.CollectionNamespace.CollectionName);
+			return await _collection.FindOneAndDeleteAsync(filter, options,  cancellationToken);
 		}
 
 		/// <inheritdoc />
-		public async Task<TProjection> FindOneAndDeleteAsync<TProjection>(IClientSessionHandle Session, FilterDefinition<T> filter,
-			FindOneAndDeleteOptions<T, TProjection> Options = null!, CancellationToken  CancellationToken = new CancellationToken())
+		public async Task<TProjection> FindOneAndDeleteAsync<TProjection>(IClientSessionHandle session, FilterDefinition<T> filter,
+			FindOneAndDeleteOptions<T, TProjection> options = null!, CancellationToken  cancellationToken = new CancellationToken())
 		{
-			using IScope Scope = GlobalTracer.Instance.BuildSpan(Prefix + "FindOneAndDeleteAsync").StartActive();
-			Scope.Span.SetTag("CollectionName", Collection.CollectionNamespace.CollectionName);
-			return await Collection.FindOneAndDeleteAsync(Session, filter, Options,  CancellationToken);
+			using IScope scope = GlobalTracer.Instance.BuildSpan(Prefix + "FindOneAndDeleteAsync").StartActive();
+			scope.Span.SetTag("CollectionName", _collection.CollectionNamespace.CollectionName);
+			return await _collection.FindOneAndDeleteAsync(session, filter, options,  cancellationToken);
 		}
 
 		/// <inheritdoc />
 		public TProjection FindOneAndReplace<TProjection>(FilterDefinition<T> filter, T replacement,
-			FindOneAndReplaceOptions<T, TProjection> Options = null!, CancellationToken  CancellationToken = new CancellationToken())
+			FindOneAndReplaceOptions<T, TProjection> options = null!, CancellationToken  cancellationToken = new CancellationToken())
 		{
-			return Collection.FindOneAndReplace(filter, replacement, Options,  CancellationToken);
+			return _collection.FindOneAndReplace(filter, replacement, options,  cancellationToken);
 		}
 
 		/// <inheritdoc />
-		public TProjection FindOneAndReplace<TProjection>(IClientSessionHandle Session, FilterDefinition<T> filter, T replacement,
-			FindOneAndReplaceOptions<T, TProjection> Options = null!, CancellationToken  CancellationToken = new CancellationToken())
+		public TProjection FindOneAndReplace<TProjection>(IClientSessionHandle session, FilterDefinition<T> filter, T replacement,
+			FindOneAndReplaceOptions<T, TProjection> options = null!, CancellationToken  cancellationToken = new CancellationToken())
 		{
-			return Collection.FindOneAndReplace(Session, filter, replacement, Options,  CancellationToken);
+			return _collection.FindOneAndReplace(session, filter, replacement, options,  cancellationToken);
 		}
 
 		/// <inheritdoc />
 		public async Task<TProjection> FindOneAndReplaceAsync<TProjection>(FilterDefinition<T> filter, T replacement,
-			FindOneAndReplaceOptions<T, TProjection> Options = null!, CancellationToken  CancellationToken = new CancellationToken())
+			FindOneAndReplaceOptions<T, TProjection> options = null!, CancellationToken  cancellationToken = new CancellationToken())
 		{
-			using IScope Scope = GlobalTracer.Instance.BuildSpan(Prefix + "FindOneAndReplaceAsync").StartActive();
-			Scope.Span.SetTag("CollectionName", Collection.CollectionNamespace.CollectionName);
-			return await Collection.FindOneAndReplaceAsync(filter, replacement, Options,  CancellationToken);
+			using IScope scope = GlobalTracer.Instance.BuildSpan(Prefix + "FindOneAndReplaceAsync").StartActive();
+			scope.Span.SetTag("CollectionName", _collection.CollectionNamespace.CollectionName);
+			return await _collection.FindOneAndReplaceAsync(filter, replacement, options,  cancellationToken);
 		}
 
 		/// <inheritdoc />
-		public async Task<TProjection> FindOneAndReplaceAsync<TProjection>(IClientSessionHandle Session, FilterDefinition<T> filter, T replacement,
-			FindOneAndReplaceOptions<T, TProjection> Options = null!, CancellationToken  CancellationToken = new CancellationToken())
+		public async Task<TProjection> FindOneAndReplaceAsync<TProjection>(IClientSessionHandle session, FilterDefinition<T> filter, T replacement,
+			FindOneAndReplaceOptions<T, TProjection> options = null!, CancellationToken  cancellationToken = new CancellationToken())
 		{
-			using IScope Scope = GlobalTracer.Instance.BuildSpan(Prefix + "FindOneAndReplaceAsync").StartActive();
-			Scope.Span.SetTag("CollectionName", Collection.CollectionNamespace.CollectionName);
-			return await Collection.FindOneAndReplaceAsync(Session, filter, replacement, Options,  CancellationToken);
+			using IScope scope = GlobalTracer.Instance.BuildSpan(Prefix + "FindOneAndReplaceAsync").StartActive();
+			scope.Span.SetTag("CollectionName", _collection.CollectionNamespace.CollectionName);
+			return await _collection.FindOneAndReplaceAsync(session, filter, replacement, options,  cancellationToken);
 		}
 
 		/// <inheritdoc />
 		public TProjection FindOneAndUpdate<TProjection>(FilterDefinition<T> filter, UpdateDefinition<T> update,
-			FindOneAndUpdateOptions<T, TProjection> Options = null!, CancellationToken  CancellationToken = new CancellationToken())
+			FindOneAndUpdateOptions<T, TProjection> options = null!, CancellationToken  cancellationToken = new CancellationToken())
 		{
-			return Collection.FindOneAndUpdate(filter, update, Options,  CancellationToken);
+			return _collection.FindOneAndUpdate(filter, update, options,  cancellationToken);
 		}
 
 		/// <inheritdoc />
-		public TProjection FindOneAndUpdate<TProjection>(IClientSessionHandle Session, FilterDefinition<T> filter,
-			UpdateDefinition<T> update, FindOneAndUpdateOptions<T, TProjection> Options = null!,
-			CancellationToken  CancellationToken = new CancellationToken())
+		public TProjection FindOneAndUpdate<TProjection>(IClientSessionHandle session, FilterDefinition<T> filter,
+			UpdateDefinition<T> update, FindOneAndUpdateOptions<T, TProjection> options = null!,
+			CancellationToken  cancellationToken = new CancellationToken())
 		{
-			return Collection.FindOneAndUpdate(Session, filter, update, Options,  CancellationToken);
+			return _collection.FindOneAndUpdate(session, filter, update, options,  cancellationToken);
 		}
 
 		/// <inheritdoc />
 		public async Task<TProjection> FindOneAndUpdateAsync<TProjection>(FilterDefinition<T> filter, UpdateDefinition<T> update,
-			FindOneAndUpdateOptions<T, TProjection> Options = null!, CancellationToken  CancellationToken = new CancellationToken())
+			FindOneAndUpdateOptions<T, TProjection> options = null!, CancellationToken  cancellationToken = new CancellationToken())
 		{
-			using IScope Scope = GlobalTracer.Instance.BuildSpan(Prefix + "FindOneAndUpdateAsync").StartActive();
-			Scope.Span.SetTag("CollectionName", Collection.CollectionNamespace.CollectionName);
-			return await Collection.FindOneAndUpdateAsync(filter, update, Options,  CancellationToken);
+			using IScope scope = GlobalTracer.Instance.BuildSpan(Prefix + "FindOneAndUpdateAsync").StartActive();
+			scope.Span.SetTag("CollectionName", _collection.CollectionNamespace.CollectionName);
+			return await _collection.FindOneAndUpdateAsync(filter, update, options,  cancellationToken);
 		}
 
 		/// <inheritdoc />
-		public async Task<TProjection> FindOneAndUpdateAsync<TProjection>(IClientSessionHandle Session, FilterDefinition<T> filter, UpdateDefinition<T> update,
-			FindOneAndUpdateOptions<T, TProjection> Options = null!, CancellationToken  CancellationToken = new CancellationToken())
+		public async Task<TProjection> FindOneAndUpdateAsync<TProjection>(IClientSessionHandle session, FilterDefinition<T> filter, UpdateDefinition<T> update,
+			FindOneAndUpdateOptions<T, TProjection> options = null!, CancellationToken  cancellationToken = new CancellationToken())
 		{
-			using IScope Scope = GlobalTracer.Instance.BuildSpan(Prefix + "FindOneAndUpdateAsync").StartActive();
-			Scope.Span.SetTag("CollectionName", Collection.CollectionNamespace.CollectionName);
-			return await Collection.FindOneAndUpdateAsync(Session, filter, update, Options,  CancellationToken);
+			using IScope scope = GlobalTracer.Instance.BuildSpan(Prefix + "FindOneAndUpdateAsync").StartActive();
+			scope.Span.SetTag("CollectionName", _collection.CollectionNamespace.CollectionName);
+			return await _collection.FindOneAndUpdateAsync(session, filter, update, options,  cancellationToken);
 		}
 
 		/// <inheritdoc />
-		public void InsertOne( T Document, InsertOneOptions Options = null!,
-			CancellationToken  CancellationToken = new CancellationToken())
+		public void InsertOne( T document, InsertOneOptions options = null!,
+			CancellationToken  cancellationToken = new CancellationToken())
 		{
-			Collection.InsertOne(Document, Options,  CancellationToken);
+			_collection.InsertOne(document, options,  cancellationToken);
 		}
 
 		/// <inheritdoc />
-		public void InsertOne(IClientSessionHandle Session, T Document, InsertOneOptions Options = null!,
-			CancellationToken  CancellationToken = new CancellationToken())
+		public void InsertOne(IClientSessionHandle session, T document, InsertOneOptions options = null!,
+			CancellationToken  cancellationToken = new CancellationToken())
 		{
-			Collection.InsertOne(Session, Document, Options,  CancellationToken);
+			_collection.InsertOne(session, document, options,  cancellationToken);
 		}
 
 		/// <inheritdoc />
-		public async Task InsertOneAsync( T Document, CancellationToken CancellationToken)
+		public async Task InsertOneAsync( T document, CancellationToken cancellationToken)
 		{
-			using IScope Scope = GlobalTracer.Instance.BuildSpan(Prefix + "InsertOneAsync").StartActive();
-			Scope.Span.SetTag("CollectionName", Collection.CollectionNamespace.CollectionName);
-			await Collection.InsertOneAsync(Document, CancellationToken);
+			using IScope scope = GlobalTracer.Instance.BuildSpan(Prefix + "InsertOneAsync").StartActive();
+			scope.Span.SetTag("CollectionName", _collection.CollectionNamespace.CollectionName);
+			await _collection.InsertOneAsync(document, cancellationToken);
 		}
 
 		/// <inheritdoc />
-		public async Task InsertOneAsync( T Document, InsertOneOptions Options = null!,
-			CancellationToken  CancellationToken = new CancellationToken())
+		public async Task InsertOneAsync( T document, InsertOneOptions options = null!,
+			CancellationToken  cancellationToken = new CancellationToken())
 		{
-			using IScope Scope = GlobalTracer.Instance.BuildSpan(Prefix + "InsertOneAsync").StartActive();
-			Scope.Span.SetTag("CollectionName", Collection.CollectionNamespace.CollectionName);
-			await Collection.InsertOneAsync(Document, Options,  CancellationToken);
+			using IScope scope = GlobalTracer.Instance.BuildSpan(Prefix + "InsertOneAsync").StartActive();
+			scope.Span.SetTag("CollectionName", _collection.CollectionNamespace.CollectionName);
+			await _collection.InsertOneAsync(document, options,  cancellationToken);
 		}
 
 		/// <inheritdoc />
-		public async Task InsertOneAsync(IClientSessionHandle Session, T Document, InsertOneOptions Options = null!,
-			CancellationToken CancellationToken = new CancellationToken())
+		public async Task InsertOneAsync(IClientSessionHandle session, T document, InsertOneOptions options = null!,
+			CancellationToken cancellationToken = new CancellationToken())
 		{
-			using IScope Scope = GlobalTracer.Instance.BuildSpan(Prefix + "InsertOneAsync").StartActive();
-			Scope.Span.SetTag("CollectionName", Collection.CollectionNamespace.CollectionName);
-			await Collection.InsertOneAsync(Session, Document, Options,  CancellationToken);
+			using IScope scope = GlobalTracer.Instance.BuildSpan(Prefix + "InsertOneAsync").StartActive();
+			scope.Span.SetTag("CollectionName", _collection.CollectionNamespace.CollectionName);
+			await _collection.InsertOneAsync(session, document, options,  cancellationToken);
 		}
 
 		/// <inheritdoc />
-		public void InsertMany(IEnumerable<T> Documents, InsertManyOptions Options = null!,
-			CancellationToken  CancellationToken = new CancellationToken())
+		public void InsertMany(IEnumerable<T> documents, InsertManyOptions options = null!,
+			CancellationToken  cancellationToken = new CancellationToken())
 		{
-			Collection.InsertMany(Documents, Options,  CancellationToken);
+			_collection.InsertMany(documents, options,  cancellationToken);
 		}
 
 		/// <inheritdoc />
-		public void InsertMany(IClientSessionHandle Session, IEnumerable<T> Documents, InsertManyOptions Options = null!,
-			CancellationToken  CancellationToken = new CancellationToken())
+		public void InsertMany(IClientSessionHandle session, IEnumerable<T> documents, InsertManyOptions options = null!,
+			CancellationToken  cancellationToken = new CancellationToken())
 		{
-			Collection.InsertMany(Session, Documents, Options,  CancellationToken);
+			_collection.InsertMany(session, documents, options,  cancellationToken);
 		}
 
 		/// <inheritdoc />
-		public async Task InsertManyAsync(IEnumerable<T> Documents, InsertManyOptions Options = null!,
-			CancellationToken  CancellationToken = new CancellationToken())
+		public async Task InsertManyAsync(IEnumerable<T> documents, InsertManyOptions options = null!,
+			CancellationToken  cancellationToken = new CancellationToken())
 		{
-			using IScope Scope = GlobalTracer.Instance.BuildSpan(Prefix + "InsertManyAsync").StartActive();
-			Scope.Span.SetTag("CollectionName", Collection.CollectionNamespace.CollectionName);
-			await Collection.InsertManyAsync(Documents, Options,  CancellationToken);
+			using IScope scope = GlobalTracer.Instance.BuildSpan(Prefix + "InsertManyAsync").StartActive();
+			scope.Span.SetTag("CollectionName", _collection.CollectionNamespace.CollectionName);
+			await _collection.InsertManyAsync(documents, options,  cancellationToken);
 		}
 
 		/// <inheritdoc />
-		public async Task InsertManyAsync(IClientSessionHandle Session, IEnumerable<T> documents, InsertManyOptions Options = null!,
-			CancellationToken  CancellationToken = new CancellationToken())
+		public async Task InsertManyAsync(IClientSessionHandle session, IEnumerable<T> documents, InsertManyOptions options = null!,
+			CancellationToken  cancellationToken = new CancellationToken())
 		{
-			using IScope Scope = GlobalTracer.Instance.BuildSpan(Prefix + "InsertManyAsync").StartActive();
-			Scope.Span.SetTag("CollectionName", Collection.CollectionNamespace.CollectionName);
-			await Collection.InsertManyAsync(Session, documents, Options,  CancellationToken);
+			using IScope scope = GlobalTracer.Instance.BuildSpan(Prefix + "InsertManyAsync").StartActive();
+			scope.Span.SetTag("CollectionName", _collection.CollectionNamespace.CollectionName);
+			await _collection.InsertManyAsync(session, documents, options,  cancellationToken);
 		}
 
 		/// <inheritdoc />
-		public IAsyncCursor<TResult> MapReduce<TResult>(BsonJavaScript map, BsonJavaScript reduce, MapReduceOptions<T, TResult> Options = null!,
-			CancellationToken  CancellationToken = new CancellationToken())
+		public IAsyncCursor<TResult> MapReduce<TResult>(BsonJavaScript map, BsonJavaScript reduce, MapReduceOptions<T, TResult> options = null!,
+			CancellationToken  cancellationToken = new CancellationToken())
 		{
-			return Collection.MapReduce(map, reduce, Options,  CancellationToken);
+			return _collection.MapReduce(map, reduce, options,  cancellationToken);
 		}
 
 		/// <inheritdoc />
-		public IAsyncCursor<TResult> MapReduce<TResult>(IClientSessionHandle Session, BsonJavaScript map, BsonJavaScript reduce,
-			MapReduceOptions<T, TResult> Options = null!, CancellationToken  CancellationToken = new CancellationToken())
+		public IAsyncCursor<TResult> MapReduce<TResult>(IClientSessionHandle session, BsonJavaScript map, BsonJavaScript reduce,
+			MapReduceOptions<T, TResult> options = null!, CancellationToken  cancellationToken = new CancellationToken())
 		{
-			return Collection.MapReduce(Session, map, reduce, Options,  CancellationToken);
+			return _collection.MapReduce(session, map, reduce, options,  cancellationToken);
 		}
 
 		/// <inheritdoc />
-		public async Task<IAsyncCursor<TResult>> MapReduceAsync<TResult>(BsonJavaScript map, BsonJavaScript reduce, MapReduceOptions<T, TResult> Options = null!,
-			CancellationToken  CancellationToken = new CancellationToken())
+		public async Task<IAsyncCursor<TResult>> MapReduceAsync<TResult>(BsonJavaScript map, BsonJavaScript reduce, MapReduceOptions<T, TResult> options = null!,
+			CancellationToken  cancellationToken = new CancellationToken())
 		{
-			using IScope Scope = GlobalTracer.Instance.BuildSpan(Prefix + "MapReduceAsync").StartActive();
-			Scope.Span.SetTag("CollectionName", Collection.CollectionNamespace.CollectionName);
-			return await Collection.MapReduceAsync(map, reduce, Options,  CancellationToken);
+			using IScope scope = GlobalTracer.Instance.BuildSpan(Prefix + "MapReduceAsync").StartActive();
+			scope.Span.SetTag("CollectionName", _collection.CollectionNamespace.CollectionName);
+			return await _collection.MapReduceAsync(map, reduce, options,  cancellationToken);
 		}
 
 		/// <inheritdoc />
-		public async Task<IAsyncCursor<TResult>> MapReduceAsync<TResult>(IClientSessionHandle Session, BsonJavaScript map, BsonJavaScript reduce,
-			MapReduceOptions<T, TResult> Options = null!, CancellationToken  CancellationToken = new CancellationToken())
+		public async Task<IAsyncCursor<TResult>> MapReduceAsync<TResult>(IClientSessionHandle session, BsonJavaScript map, BsonJavaScript reduce,
+			MapReduceOptions<T, TResult> options = null!, CancellationToken  cancellationToken = new CancellationToken())
 		{
-			using IScope Scope = GlobalTracer.Instance.BuildSpan(Prefix + "MapReduceAsync").StartActive();
-			Scope.Span.SetTag("CollectionName", Collection.CollectionNamespace.CollectionName);
-			return await Collection.MapReduceAsync(Session, map, reduce, Options,  CancellationToken);
+			using IScope scope = GlobalTracer.Instance.BuildSpan(Prefix + "MapReduceAsync").StartActive();
+			scope.Span.SetTag("CollectionName", _collection.CollectionNamespace.CollectionName);
+			return await _collection.MapReduceAsync(session, map, reduce, options,  cancellationToken);
 		}
 
 		/// <inheritdoc />
 		public IFilteredMongoCollection<TDerivedDocument> OfType<TDerivedDocument>() where TDerivedDocument : T
 		{
-			return Collection.OfType<TDerivedDocument>();
+			return _collection.OfType<TDerivedDocument>();
 		}
 
 		/// <inheritdoc />
-		public ReplaceOneResult ReplaceOne(FilterDefinition<T> filter, T replacement, ReplaceOptions Options = null!,
-			CancellationToken  CancellationToken = new CancellationToken())
+		public ReplaceOneResult ReplaceOne(FilterDefinition<T> filter, T replacement, ReplaceOptions options = null!,
+			CancellationToken  cancellationToken = new CancellationToken())
 		{
-			return Collection.ReplaceOne(filter, replacement, Options,  CancellationToken);
+			return _collection.ReplaceOne(filter, replacement, options,  cancellationToken);
 		}
 
 		/// <inheritdoc />
-		public ReplaceOneResult ReplaceOne(FilterDefinition<T> filter, T replacement, UpdateOptions Options,
-			CancellationToken  CancellationToken = new CancellationToken())
+		public ReplaceOneResult ReplaceOne(FilterDefinition<T> filter, T replacement, UpdateOptions options,
+			CancellationToken  cancellationToken = new CancellationToken())
 		{
-			return Collection.ReplaceOne(filter, replacement, Options,  CancellationToken);
+			return _collection.ReplaceOne(filter, replacement, options,  cancellationToken);
 		}
 
 		/// <inheritdoc />
-		public ReplaceOneResult ReplaceOne(IClientSessionHandle Session, FilterDefinition<T> filter, T replacement,
-			ReplaceOptions Options = null!, CancellationToken  CancellationToken = new CancellationToken())
+		public ReplaceOneResult ReplaceOne(IClientSessionHandle session, FilterDefinition<T> filter, T replacement,
+			ReplaceOptions options = null!, CancellationToken  cancellationToken = new CancellationToken())
 		{
-			return Collection.ReplaceOne(Session, filter, replacement, Options,  CancellationToken);
+			return _collection.ReplaceOne(session, filter, replacement, options,  cancellationToken);
 		}
 
 		/// <inheritdoc />
-		public ReplaceOneResult ReplaceOne(IClientSessionHandle Session, FilterDefinition<T> filter, T replacement, UpdateOptions Options,
-			CancellationToken  CancellationToken = new CancellationToken())
+		public ReplaceOneResult ReplaceOne(IClientSessionHandle session, FilterDefinition<T> filter, T replacement, UpdateOptions options,
+			CancellationToken  cancellationToken = new CancellationToken())
 		{
-			return Collection.ReplaceOne(Session, filter, replacement, Options,  CancellationToken);
+			return _collection.ReplaceOne(session, filter, replacement, options,  cancellationToken);
 		}
 
 		/// <inheritdoc />
-		public async Task<ReplaceOneResult> ReplaceOneAsync(FilterDefinition<T> filter, T replacement, ReplaceOptions Options = null!,
-			CancellationToken  CancellationToken = new CancellationToken())
+		public async Task<ReplaceOneResult> ReplaceOneAsync(FilterDefinition<T> filter, T replacement, ReplaceOptions options = null!,
+			CancellationToken  cancellationToken = new CancellationToken())
 		{
-			using IScope Scope = GlobalTracer.Instance.BuildSpan(Prefix + "ReplaceOneAsync").StartActive();
-			Scope.Span.SetTag("CollectionName", Collection.CollectionNamespace.CollectionName);
-			return await Collection.ReplaceOneAsync(filter, replacement, Options,  CancellationToken);
+			using IScope scope = GlobalTracer.Instance.BuildSpan(Prefix + "ReplaceOneAsync").StartActive();
+			scope.Span.SetTag("CollectionName", _collection.CollectionNamespace.CollectionName);
+			return await _collection.ReplaceOneAsync(filter, replacement, options,  cancellationToken);
 		}
 
 		/// <inheritdoc />
-		public async Task<ReplaceOneResult> ReplaceOneAsync(FilterDefinition<T> filter, T replacement, UpdateOptions Options,
-			CancellationToken  CancellationToken = new CancellationToken())
+		public async Task<ReplaceOneResult> ReplaceOneAsync(FilterDefinition<T> filter, T replacement, UpdateOptions options,
+			CancellationToken  cancellationToken = new CancellationToken())
 		{
-			using IScope Scope = GlobalTracer.Instance.BuildSpan(Prefix + "ReplaceOneAsync").StartActive();
-			Scope.Span.SetTag("CollectionName", Collection.CollectionNamespace.CollectionName);
-			return await Collection.ReplaceOneAsync(filter, replacement, Options,  CancellationToken);
+			using IScope scope = GlobalTracer.Instance.BuildSpan(Prefix + "ReplaceOneAsync").StartActive();
+			scope.Span.SetTag("CollectionName", _collection.CollectionNamespace.CollectionName);
+			return await _collection.ReplaceOneAsync(filter, replacement, options,  cancellationToken);
 		}
 
 		/// <inheritdoc />
-		public async Task<ReplaceOneResult> ReplaceOneAsync(IClientSessionHandle Session, FilterDefinition<T> filter, T replacement,
-			ReplaceOptions Options = null!, CancellationToken  CancellationToken = new CancellationToken())
+		public async Task<ReplaceOneResult> ReplaceOneAsync(IClientSessionHandle session, FilterDefinition<T> filter, T replacement,
+			ReplaceOptions options = null!, CancellationToken  cancellationToken = new CancellationToken())
 		{
-			using IScope Scope = GlobalTracer.Instance.BuildSpan(Prefix + "ReplaceOneAsync").StartActive();
-			Scope.Span.SetTag("CollectionName", Collection.CollectionNamespace.CollectionName);
-			return await Collection.ReplaceOneAsync(Session, filter, replacement, Options,  CancellationToken);
+			using IScope scope = GlobalTracer.Instance.BuildSpan(Prefix + "ReplaceOneAsync").StartActive();
+			scope.Span.SetTag("CollectionName", _collection.CollectionNamespace.CollectionName);
+			return await _collection.ReplaceOneAsync(session, filter, replacement, options,  cancellationToken);
 		}
 
 		/// <inheritdoc />
-		public async Task<ReplaceOneResult> ReplaceOneAsync(IClientSessionHandle Session, FilterDefinition<T> filter, T replacement, UpdateOptions Options,
-			CancellationToken  CancellationToken = new CancellationToken())
+		public async Task<ReplaceOneResult> ReplaceOneAsync(IClientSessionHandle session, FilterDefinition<T> filter, T replacement, UpdateOptions options,
+			CancellationToken  cancellationToken = new CancellationToken())
 		{
-			using IScope Scope = GlobalTracer.Instance.BuildSpan(Prefix + "ReplaceOneAsync").StartActive();
-			Scope.Span.SetTag("CollectionName", Collection.CollectionNamespace.CollectionName);
-			return await Collection.ReplaceOneAsync(Session, filter, replacement, Options,  CancellationToken);
+			using IScope scope = GlobalTracer.Instance.BuildSpan(Prefix + "ReplaceOneAsync").StartActive();
+			scope.Span.SetTag("CollectionName", _collection.CollectionNamespace.CollectionName);
+			return await _collection.ReplaceOneAsync(session, filter, replacement, options,  cancellationToken);
 		}
 
 		/// <inheritdoc />
-		public UpdateResult UpdateMany(FilterDefinition<T> filter, UpdateDefinition<T> update, UpdateOptions Options = null!,
-			CancellationToken  CancellationToken = new CancellationToken())
+		public UpdateResult UpdateMany(FilterDefinition<T> filter, UpdateDefinition<T> update, UpdateOptions options = null!,
+			CancellationToken  cancellationToken = new CancellationToken())
 		{
-			return Collection.UpdateMany(filter, update, Options,  CancellationToken);
+			return _collection.UpdateMany(filter, update, options,  cancellationToken);
 		}
 
 		/// <inheritdoc />
-		public UpdateResult UpdateMany(IClientSessionHandle Session, FilterDefinition<T> filter, UpdateDefinition<T> update,
-			UpdateOptions Options = null!, CancellationToken  CancellationToken = new CancellationToken())
+		public UpdateResult UpdateMany(IClientSessionHandle session, FilterDefinition<T> filter, UpdateDefinition<T> update,
+			UpdateOptions options = null!, CancellationToken  cancellationToken = new CancellationToken())
 		{
-			return Collection.UpdateMany(Session, filter, update, Options,  CancellationToken);
+			return _collection.UpdateMany(session, filter, update, options,  cancellationToken);
 		}
 
 		/// <inheritdoc />
-		public async Task<UpdateResult> UpdateManyAsync(FilterDefinition<T> filter, UpdateDefinition<T> update, UpdateOptions Options = null!,
-			CancellationToken  CancellationToken = new CancellationToken())
+		public async Task<UpdateResult> UpdateManyAsync(FilterDefinition<T> filter, UpdateDefinition<T> update, UpdateOptions options = null!,
+			CancellationToken  cancellationToken = new CancellationToken())
 		{
-			using IScope Scope = GlobalTracer.Instance.BuildSpan(Prefix + "UpdateManyAsync").StartActive();
-			Scope.Span.SetTag("CollectionName", Collection.CollectionNamespace.CollectionName);
-			return await Collection.UpdateManyAsync(filter, update, Options,  CancellationToken);
+			using IScope scope = GlobalTracer.Instance.BuildSpan(Prefix + "UpdateManyAsync").StartActive();
+			scope.Span.SetTag("CollectionName", _collection.CollectionNamespace.CollectionName);
+			return await _collection.UpdateManyAsync(filter, update, options,  cancellationToken);
 		}
 
 		/// <inheritdoc />
-		public async Task<UpdateResult> UpdateManyAsync(IClientSessionHandle Session, FilterDefinition<T> filter, UpdateDefinition<T> update,
-			UpdateOptions Options = null!, CancellationToken  CancellationToken = new CancellationToken())
+		public async Task<UpdateResult> UpdateManyAsync(IClientSessionHandle session, FilterDefinition<T> filter, UpdateDefinition<T> update,
+			UpdateOptions options = null!, CancellationToken  cancellationToken = new CancellationToken())
 		{
-			using IScope Scope = GlobalTracer.Instance.BuildSpan(Prefix + "UpdateManyAsync").StartActive();
-			Scope.Span.SetTag("CollectionName", Collection.CollectionNamespace.CollectionName);
-			return await Collection.UpdateManyAsync(Session, filter, update, Options,  CancellationToken);
+			using IScope scope = GlobalTracer.Instance.BuildSpan(Prefix + "UpdateManyAsync").StartActive();
+			scope.Span.SetTag("CollectionName", _collection.CollectionNamespace.CollectionName);
+			return await _collection.UpdateManyAsync(session, filter, update, options,  cancellationToken);
 		}
 
 		/// <inheritdoc />
-		public UpdateResult UpdateOne(FilterDefinition<T> filter, UpdateDefinition<T> update, UpdateOptions Options = null!,
-			CancellationToken  CancellationToken = new CancellationToken())
+		public UpdateResult UpdateOne(FilterDefinition<T> filter, UpdateDefinition<T> update, UpdateOptions options = null!,
+			CancellationToken  cancellationToken = new CancellationToken())
 		{
-			return Collection.UpdateOne(filter, update, Options,  CancellationToken);
+			return _collection.UpdateOne(filter, update, options,  cancellationToken);
 		}
 
 		/// <inheritdoc />
-		public UpdateResult UpdateOne(IClientSessionHandle Session, FilterDefinition<T> filter, UpdateDefinition<T> update,
-			UpdateOptions Options = null!, CancellationToken  CancellationToken = new CancellationToken())
+		public UpdateResult UpdateOne(IClientSessionHandle session, FilterDefinition<T> filter, UpdateDefinition<T> update,
+			UpdateOptions options = null!, CancellationToken  cancellationToken = new CancellationToken())
 		{
-			return Collection.UpdateOne(Session, filter, update, Options,  CancellationToken);
+			return _collection.UpdateOne(session, filter, update, options,  cancellationToken);
 		}
 
 		/// <inheritdoc />
-		public async Task<UpdateResult> UpdateOneAsync(FilterDefinition<T> filter, UpdateDefinition<T> update, UpdateOptions Options = null!,
-			CancellationToken  CancellationToken = new CancellationToken())
+		public async Task<UpdateResult> UpdateOneAsync(FilterDefinition<T> filter, UpdateDefinition<T> update, UpdateOptions options = null!,
+			CancellationToken  cancellationToken = new CancellationToken())
 		{
-			using IScope Scope = GlobalTracer.Instance.BuildSpan(Prefix + "UpdateOneAsync").StartActive();
-			Scope.Span.SetTag("CollectionName", Collection.CollectionNamespace.CollectionName);
-			return await Collection.UpdateOneAsync(filter, update, Options,  CancellationToken);
+			using IScope scope = GlobalTracer.Instance.BuildSpan(Prefix + "UpdateOneAsync").StartActive();
+			scope.Span.SetTag("CollectionName", _collection.CollectionNamespace.CollectionName);
+			return await _collection.UpdateOneAsync(filter, update, options,  cancellationToken);
 		}
 
 		/// <inheritdoc />
-		public async Task<UpdateResult> UpdateOneAsync(IClientSessionHandle Session, FilterDefinition<T> filter, UpdateDefinition<T> update,
-			UpdateOptions Options = null!, CancellationToken  CancellationToken = new CancellationToken())
+		public async Task<UpdateResult> UpdateOneAsync(IClientSessionHandle session, FilterDefinition<T> filter, UpdateDefinition<T> update,
+			UpdateOptions options = null!, CancellationToken  cancellationToken = new CancellationToken())
 		{
-			using IScope Scope = GlobalTracer.Instance.BuildSpan(Prefix + "UpdateOneAsync").StartActive();
-			Scope.Span.SetTag("CollectionName", Collection.CollectionNamespace.CollectionName);
-			return await Collection.UpdateOneAsync(Session, filter, update, Options,  CancellationToken);
+			using IScope scope = GlobalTracer.Instance.BuildSpan(Prefix + "UpdateOneAsync").StartActive();
+			scope.Span.SetTag("CollectionName", _collection.CollectionNamespace.CollectionName);
+			return await _collection.UpdateOneAsync(session, filter, update, options,  cancellationToken);
 		}
 
 		/// <inheritdoc />
-		public IChangeStreamCursor<TResult> Watch<TResult>(PipelineDefinition<ChangeStreamDocument<T>, TResult> pipeline, ChangeStreamOptions Options = null!,
-			CancellationToken  CancellationToken = new CancellationToken())
+		public IChangeStreamCursor<TResult> Watch<TResult>(PipelineDefinition<ChangeStreamDocument<T>, TResult> pipeline, ChangeStreamOptions options = null!,
+			CancellationToken  cancellationToken = new CancellationToken())
 		{
-			return Collection.Watch(pipeline, Options,  CancellationToken);
+			return _collection.Watch(pipeline, options,  cancellationToken);
 		}
 
 		/// <inheritdoc />
-		public IChangeStreamCursor<TResult> Watch<TResult>(IClientSessionHandle Session, PipelineDefinition<ChangeStreamDocument<T>, TResult> pipeline,
-			ChangeStreamOptions Options = null!, CancellationToken  CancellationToken = new CancellationToken())
+		public IChangeStreamCursor<TResult> Watch<TResult>(IClientSessionHandle session, PipelineDefinition<ChangeStreamDocument<T>, TResult> pipeline,
+			ChangeStreamOptions options = null!, CancellationToken  cancellationToken = new CancellationToken())
 		{
-			return Collection.Watch(Session, pipeline, Options,  CancellationToken);
+			return _collection.Watch(session, pipeline, options,  cancellationToken);
 		}
 
 		/// <inheritdoc />
-		public async Task<IChangeStreamCursor<TResult>> WatchAsync<TResult>(PipelineDefinition<ChangeStreamDocument<T>, TResult> pipeline, ChangeStreamOptions Options = null!,
-			CancellationToken  CancellationToken = new CancellationToken())
+		public async Task<IChangeStreamCursor<TResult>> WatchAsync<TResult>(PipelineDefinition<ChangeStreamDocument<T>, TResult> pipeline, ChangeStreamOptions options = null!,
+			CancellationToken  cancellationToken = new CancellationToken())
 		{
-			using IScope Scope = GlobalTracer.Instance.BuildSpan(Prefix + "WatchAsync").StartActive();
-			Scope.Span.SetTag("CollectionName", Collection.CollectionNamespace.CollectionName);
-			return await Collection.WatchAsync(pipeline, Options,  CancellationToken);
+			using IScope scope = GlobalTracer.Instance.BuildSpan(Prefix + "WatchAsync").StartActive();
+			scope.Span.SetTag("CollectionName", _collection.CollectionNamespace.CollectionName);
+			return await _collection.WatchAsync(pipeline, options,  cancellationToken);
 		}
 
 		/// <inheritdoc />
-		public async Task<IChangeStreamCursor<TResult>> WatchAsync<TResult>(IClientSessionHandle Session, PipelineDefinition<ChangeStreamDocument<T>, TResult> pipeline, ChangeStreamOptions Options = null!,
-			CancellationToken  CancellationToken = new CancellationToken())
+		public async Task<IChangeStreamCursor<TResult>> WatchAsync<TResult>(IClientSessionHandle session, PipelineDefinition<ChangeStreamDocument<T>, TResult> pipeline, ChangeStreamOptions options = null!,
+			CancellationToken  cancellationToken = new CancellationToken())
 		{
-			using IScope Scope = GlobalTracer.Instance.BuildSpan(Prefix + "WatchAsync").StartActive();
-			Scope.Span.SetTag("CollectionName", Collection.CollectionNamespace.CollectionName);
-			return await Collection.WatchAsync(Session, pipeline, Options,  CancellationToken);
+			using IScope scope = GlobalTracer.Instance.BuildSpan(Prefix + "WatchAsync").StartActive();
+			scope.Span.SetTag("CollectionName", _collection.CollectionNamespace.CollectionName);
+			return await _collection.WatchAsync(session, pipeline, options,  cancellationToken);
 		}
 
 		/// <inheritdoc />
-		public IMongoCollection<T> WithReadConcern(ReadConcern ReadConcern)
+		public IMongoCollection<T> WithReadConcern(ReadConcern readConcern)
 		{
-			return Collection.WithReadConcern(ReadConcern);
+			return _collection.WithReadConcern(readConcern);
 		}
 
 		/// <inheritdoc />
-		public IMongoCollection<T> WithReadPreference(ReadPreference ReadPreference)
+		public IMongoCollection<T> WithReadPreference(ReadPreference readPreference)
 		{
-			return Collection.WithReadPreference(ReadPreference);
+			return _collection.WithReadPreference(readPreference);
 		}
 
 		/// <inheritdoc />
-		public IMongoCollection<T> WithWriteConcern(WriteConcern WriteConcern)
+		public IMongoCollection<T> WithWriteConcern(WriteConcern writeConcern)
 		{
-			return Collection.WithWriteConcern(WriteConcern);
+			return _collection.WithWriteConcern(writeConcern);
 		}
 #pragma warning restore CS0618
 		
 		/// <inheritdoc />
-		public CollectionNamespace CollectionNamespace => Collection.CollectionNamespace;
+		public CollectionNamespace CollectionNamespace => _collection.CollectionNamespace;
 
 		/// <inheritdoc />
-		public IMongoDatabase Database => Collection.Database;
+		public IMongoDatabase Database => _collection.Database;
 
 		/// <inheritdoc />
-		public IBsonSerializer<T> DocumentSerializer => Collection.DocumentSerializer;
+		public IBsonSerializer<T> DocumentSerializer => _collection.DocumentSerializer;
 
 		/// <inheritdoc />
-		public IMongoIndexManager<T> Indexes => Collection.Indexes;
+		public IMongoIndexManager<T> Indexes => _collection.Indexes;
 
 		/// <inheritdoc />
-		public MongoCollectionSettings Settings => Collection.Settings;
+		public MongoCollectionSettings Settings => _collection.Settings;
 	}
 }
\ No newline at end of file
diff --git a/Engine/Source/Programs/Horde/Horde.Build/Utilities/NoOpDogStatsd.cs b/Engine/Source/Programs/Horde/Horde.Build/Utilities/NoOpDogStatsd.cs
--- a/Engine/Source/Programs/Horde/Horde.Build/Utilities/NoOpDogStatsd.cs
+++ b/Engine/Source/Programs/Horde/Horde.Build/Utilities/NoOpDogStatsd.cs
@@ -1,7 +1,6 @@
 ﻿// Copyright Epic Games, Inc. All Rights Reserved.
 
 using System;
-using System.Diagnostics.CodeAnalysis;
 using System.IO;
 using StatsdClient;
 
@@ -17,87 +16,85 @@ namespace Horde.Build.Utilities
 		public ITelemetryCounters TelemetryCounters { get; } = null!;
 		
 		/// <inheritdoc />
-		public void Configure(StatsdConfig Config)
+		public void Configure(StatsdConfig config)
 		{
 		}
 
 		/// <inheritdoc />
-		public void Counter(string StatName, double Value, double SampleRate = 1, string[] Tags = null!)
+		public void Counter(string statName, double value, double sampleRate = 1, string[] tags = null!)
 		{
 		}
 
 		/// <inheritdoc />
-		public void Decrement(string StatName, int Value = 1, double SampleRate = 1, params string[] Tags)
+		public void Decrement(string statName, int value = 1, double sampleRate = 1, params string[] tags)
 		{
 		}
 
 		/// <inheritdoc />
-		public void Event(string Title, string Text, string AlertType = null!, string AggregationKey = null!, string SourceType = null!,
-			int? DateHappened = null, string Priority = null!, string Hostname = null!, string[] Tags = null!)
+		public void Event(string title, string text, string alertType = null!, string aggregationKey = null!, string sourceType = null!,
+			int? dateHappened = null, string priority = null!, string hostname = null!, string[] tags = null!)
 		{
 		}
 
 		/// <inheritdoc />
-		public void Gauge(string StatName, double Value, double SampleRate = 1, string[] Tags = null!)
+		public void Gauge(string statName, double value, double sampleRate = 1, string[] tags = null!)
 		{
 		}
 
 		/// <inheritdoc />
-		public void Histogram(string StatName, double Value, double SampleRate = 1, string[] Tags = null!)
+		public void Histogram(string statName, double value, double sampleRate = 1, string[] tags = null!)
 		{
 		}
 
 		/// <inheritdoc />
-		public void Distribution(string StatName, double Value, double SampleRate = 1, string[] Tags = null!)
+		public void Distribution(string statName, double value, double sampleRate = 1, string[] tags = null!)
 		{
 		}
 
 		/// <inheritdoc />
-		public void Increment(string StatName, int Value = 1, double SampleRate = 1, string[] Tags = null!)
+		public void Increment(string statName, int value = 1, double sampleRate = 1, string[] tags = null!)
 		{
 		}
 
 		/// <inheritdoc />
-		public void Set<T>(string StatName, T Value, double SampleRate = 1, string[] Tags = null!)
+		public void Set<T>(string statName, T value, double sampleRate = 1, string[] tags = null!)
 		{
 		}
 
 		/// <inheritdoc />
-		public void Set(string StatName, string Value, double SampleRate = 1, string[] Tags = null!)
+		public void Set(string statName, string value, double sampleRate = 1, string[] tags = null!)
 		{
 		}
 
 		/// <inheritdoc />
-		public IDisposable StartTimer(string Name, double SampleRate = 1, string[] Tags = null!)
+		public IDisposable StartTimer(string name, double sampleRate = 1, string[] tags = null!)
 		{
 			// Some random object to satisfy IDisposable
 			return new MemoryStream();
 		}
 
 		/// <inheritdoc />
-		public void Time(Action Action, string StatName, double SampleRate = 1, string[] Tags = null!)
+		public void Time(Action action, string statName, double sampleRate = 1, string[] tags = null!)
 		{
 		}
 
 		/// <inheritdoc />
-		public T Time<T>(Func<T> Func, string StatName, double SampleRate = 1, string[] Tags = null!)
+		public T Time<T>(Func<T> func, string statName, double sampleRate = 1, string[] tags = null!)
 		{
-			return Func();
+			return func();
 		}
 
 		/// <inheritdoc />
-		public void Timer(string StatName, double Value, double SampleRate = 1, string[] Tags = null!)
+		public void Timer(string statName, double value, double sampleRate = 1, string[] tags = null!)
 		{
 		}
 
 		/// <inheritdoc />
-		public void ServiceCheck(string Name, Status Status, int? Timestamp = null, string Hostname = null!, string[] Tags = null!, string Message = null!)
+		public void ServiceCheck(string name, Status status, int? timestamp = null, string hostname = null!, string[] tags = null!, string message = null!)
 		{
 		}
 		
 		/// <inheritdoc />
-		[SuppressMessage("Design", "CA1063:Implement IDisposable correctly")]
-		[SuppressMessage("Usage", "CA1816:Call GC.SuppressFinalize correctly")]
 		public void Dispose()
 		{
 		}
diff --git a/Engine/Source/Programs/Horde/Horde.Build/Utilities/NullDiscriminatorConvention.cs b/Engine/Source/Programs/Horde/Horde.Build/Utilities/NullDiscriminatorConvention.cs
--- a/Engine/Source/Programs/Horde/Horde.Build/Utilities/NullDiscriminatorConvention.cs
+++ b/Engine/Source/Programs/Horde/Horde.Build/Utilities/NullDiscriminatorConvention.cs
@@ -1,12 +1,9 @@
 // Copyright Epic Games, Inc. All Rights Reserved.
 
+using System;
 using MongoDB.Bson;
 using MongoDB.Bson.IO;
 using MongoDB.Bson.Serialization.Conventions;
-using System;
-using System.Collections.Generic;
-using System.Linq;
-using System.Threading.Tasks;
 
 namespace Horde.Build.Utilities
 {
@@ -21,9 +18,9 @@ namespace Horde.Build.Utilities
 		public static NullDiscriminatorConvention Instance { get; } = new NullDiscriminatorConvention();
 
 		/// <inheritdoc/>
-		public Type GetActualType(IBsonReader BsonReader, Type NominalType)
+		public Type GetActualType(IBsonReader bsonReader, Type nominalType)
 		{
-			return NominalType;
+			return nominalType;
 		}
 
 		/// <inheritdoc/>
diff --git a/Engine/Source/Programs/Horde/Horde.Build/Utilities/ObjectId.cs b/Engine/Source/Programs/Horde/Horde.Build/Utilities/ObjectId.cs
--- a/Engine/Source/Programs/Horde/Horde.Build/Utilities/ObjectId.cs
+++ b/Engine/Source/Programs/Horde/Horde.Build/Utilities/ObjectId.cs
@@ -1,16 +1,15 @@
 // Copyright Epic Games, Inc. All Rights Reserved.
 
-using EpicGames.Core;
-using EpicGames.Serialization.Converters;
-using MongoDB.Bson;
-using MongoDB.Bson.Serialization;
-using MongoDB.Bson.Serialization.Serializers;
 using System;
 using System.ComponentModel;
 using System.Diagnostics.CodeAnalysis;
 using System.Globalization;
 using System.Text.Json;
 using System.Text.Json.Serialization;
+using EpicGames.Core;
+using MongoDB.Bson;
+using MongoDB.Bson.Serialization;
+using MongoDB.Bson.Serialization.Serializers;
 
 namespace Horde.Build.Utilities
 {
@@ -35,74 +34,74 @@ namespace Horde.Build.Utilities
 		/// <summary>
 		/// Constructor
 		/// </summary>
-		public ObjectId(ObjectId Value) => this.Value = Value;
+		public ObjectId(ObjectId value) => Value = value;
 
 		/// <summary>
 		/// Constructor
 		/// </summary>
-		public ObjectId(byte[] Bytes) => this.Value = new ObjectId(Bytes);
+		public ObjectId(byte[] bytes) => Value = new ObjectId(bytes);
 
 		/// <summary>
 		/// Constructor
 		/// </summary>
-		public ObjectId(string Value) => this.Value = new ObjectId(Value);
+		public ObjectId(string value) => Value = new ObjectId(value);
 
 		/// <inheritdoc cref="ObjectId.GenerateNewId()"/>
 		[SuppressMessage("Design", "CA1000:Do not declare static members on generic types", Justification = "<Pending>")]
 		public static ObjectId<T> GenerateNewId() => new ObjectId<T>(ObjectId.GenerateNewId());
 
-		/// <inheritdoc cref="ObjectId.Parse(string)"/>
+		/// <inheritdoc cref="ObjectId.Parse(String)"/>
 		[SuppressMessage("Design", "CA1000:Do not declare static members on generic types", Justification = "<Pending>")]
-		public static ObjectId<T> Parse(string Str) => new ObjectId<T>(ObjectId.Parse(Str));
+		public static ObjectId<T> Parse(string str) => new ObjectId<T>(ObjectId.Parse(str));
 
-		/// <inheritdoc cref="ObjectId.TryParse(string, out ObjectId)"/>
+		/// <inheritdoc cref="ObjectId.TryParse(String, out ObjectId)"/>
 		[SuppressMessage("Design", "CA1000:Do not declare static members on generic types", Justification = "<Pending>")]
-		public static bool TryParse(string Str, out ObjectId<T> Value)
+		public static bool TryParse(string str, out ObjectId<T> value)
 		{
-			if (ObjectId.TryParse(Str, out ObjectId Id))
+			if (ObjectId.TryParse(str, out ObjectId id))
 			{
-				Value = new ObjectId<T>(Id);
+				value = new ObjectId<T>(id);
 				return true;
 			}
 			else
 			{
-				Value = Empty;
+				value = Empty;
 				return false;
 			}
 		}
 
 		/// <inheritdoc/>
-		public override bool Equals(object? Obj) => Obj is ObjectId<T> Other && Equals(Other);
+		public override bool Equals(object? obj) => obj is ObjectId<T> other && Equals(other);
 
 		/// <inheritdoc/>
 		public override int GetHashCode() => Value.GetHashCode();
 
 		/// <inheritdoc/>
-		public bool Equals(ObjectId<T> Other) => Value.Equals(Other.Value);
+		public bool Equals(ObjectId<T> other) => Value.Equals(other.Value);
 
 		/// <inheritdoc/>
 		public override string ToString() => Value.ToString();
 
 		/// <inheritdoc/>
-		public int CompareTo(ObjectId<T> Other) => Value.CompareTo(Other.Value);
+		public int CompareTo(ObjectId<T> other) => Value.CompareTo(other.Value);
 
 		/// <inheritdoc cref="ObjectId.op_LessThan"/>
-		public static bool operator <(ObjectId<T> Left, ObjectId<T> Right) => Left.Value < Right.Value;
+		public static bool operator <(ObjectId<T> left, ObjectId<T> right) => left.Value < right.Value;
 
 		/// <inheritdoc cref="ObjectId.op_GreaterThan"/>
-		public static bool operator >(ObjectId<T> Left, ObjectId<T> Right) => Left.Value > Right.Value;
+		public static bool operator >(ObjectId<T> left, ObjectId<T> right) => left.Value > right.Value;
 
 		/// <inheritdoc cref="ObjectId.op_LessThanOrEqual"/>
-		public static bool operator <=(ObjectId<T> Left, ObjectId<T> Right) => Left.Value <= Right.Value;
+		public static bool operator <=(ObjectId<T> left, ObjectId<T> right) => left.Value <= right.Value;
 
 		/// <inheritdoc cref="ObjectId.op_GreaterThanOrEqual"/>
-		public static bool operator >=(ObjectId<T> Left, ObjectId<T> Right) => Left.Value >= Right.Value;
+		public static bool operator >=(ObjectId<T> left, ObjectId<T> right) => left.Value >= right.Value;
 
 		/// <inheritdoc cref="ObjectId.op_Equality"/>
-		public static bool operator ==(ObjectId<T> Left, ObjectId<T> Right) => Left.Value == Right.Value;
+		public static bool operator ==(ObjectId<T> left, ObjectId<T> right) => left.Value == right.Value;
 
 		/// <inheritdoc cref="ObjectId.op_Inequality"/>
-		public static bool operator !=(ObjectId<T> Left, ObjectId<T> Right) => Left.Value != Right.Value;
+		public static bool operator !=(ObjectId<T> left, ObjectId<T> right) => left.Value != right.Value;
 	}
 
 	/// <summary>
@@ -111,15 +110,15 @@ namespace Horde.Build.Utilities
 	public class ObjectIdJsonConverter<T> : JsonConverter<ObjectId<T>>
 	{
 		/// <inheritdoc/>
-		public override ObjectId<T> Read(ref Utf8JsonReader Reader, Type TypeToConvert, JsonSerializerOptions Options)
+		public override ObjectId<T> Read(ref Utf8JsonReader reader, Type typeToConvert, JsonSerializerOptions options)
 		{
-			return new ObjectId<T>(ObjectId.Parse(Reader.GetString()!));
+			return new ObjectId<T>(ObjectId.Parse(reader.GetString()!));
 		}
 
 		/// <inheritdoc/>
-		public override void Write(Utf8JsonWriter Writer, ObjectId<T> Value, JsonSerializerOptions Options)
+		public override void Write(Utf8JsonWriter writer, ObjectId<T> value, JsonSerializerOptions options)
 		{
-			Writer.WriteStringValue(Value.ToString());
+			writer.WriteStringValue(value.ToString());
 		}
 	}
 
@@ -129,15 +128,15 @@ namespace Horde.Build.Utilities
 	public class JsonObjectIdConverterFactory : JsonConverterFactory
 	{
 		/// <inheritdoc/>
-		public override bool CanConvert(Type TypeToConvert)
+		public override bool CanConvert(Type typeToConvert)
 		{
-			return TypeToConvert.IsGenericType && TypeToConvert.GetGenericTypeDefinition() == typeof(ObjectId<>);
+			return typeToConvert.IsGenericType && typeToConvert.GetGenericTypeDefinition() == typeof(ObjectId<>);
 		}
 
 		/// <inheritdoc/>
-		public override JsonConverter? CreateConverter(Type Type, JsonSerializerOptions Options)
+		public override JsonConverter? CreateConverter(Type type, JsonSerializerOptions options)
 		{
-			return (JsonConverter?)Activator.CreateInstance(typeof(ObjectIdJsonConverter<>).MakeGenericType(Type.GetGenericArguments()));
+			return (JsonConverter?)Activator.CreateInstance(typeof(ObjectIdJsonConverter<>).MakeGenericType(type.GetGenericArguments()));
 		}
 	}
 
@@ -147,15 +146,15 @@ namespace Horde.Build.Utilities
 	public sealed class ObjectIdBsonSerializer<T> : SerializerBase<ObjectId<T>>
 	{
 		/// <inheritdoc/>
-		public override ObjectId<T> Deserialize(BsonDeserializationContext Context, BsonDeserializationArgs Args)
+		public override ObjectId<T> Deserialize(BsonDeserializationContext context, BsonDeserializationArgs args)
 		{
-			return new ObjectId<T>(Context.Reader.ReadObjectId());
+			return new ObjectId<T>(context.Reader.ReadObjectId());
 		}
 
 		/// <inheritdoc/>
-		public override void Serialize(BsonSerializationContext Context, BsonSerializationArgs Args, ObjectId<T> Value)
+		public override void Serialize(BsonSerializationContext context, BsonSerializationArgs args, ObjectId<T> value)
 		{
-			Context.Writer.WriteObjectId(Value.Value);
+			context.Writer.WriteObjectId(value.Value);
 		}
 	}
 
@@ -165,11 +164,11 @@ namespace Horde.Build.Utilities
 	public sealed class ObjectIdSerializationProvider : BsonSerializationProviderBase
 	{
 		/// <inheritdoc/>
-		public override IBsonSerializer? GetSerializer(Type Type, IBsonSerializerRegistry SerializerRegistry)
+		public override IBsonSerializer? GetSerializer(Type type, IBsonSerializerRegistry serializerRegistry)
 		{
-			if (Type.IsGenericType && Type.GetGenericTypeDefinition() == typeof(ObjectId<>))
+			if (type.IsGenericType && type.GetGenericTypeDefinition() == typeof(ObjectId<>))
 			{
-				return (IBsonSerializer?)Activator.CreateInstance(typeof(ObjectIdBsonSerializer<>).MakeGenericType(Type.GetGenericArguments()));
+				return (IBsonSerializer?)Activator.CreateInstance(typeof(ObjectIdBsonSerializer<>).MakeGenericType(type.GetGenericArguments()));
 			}
 			else
 			{
@@ -183,65 +182,65 @@ namespace Horde.Build.Utilities
 	/// </summary>
 	sealed class ObjectIdTypeConverter : TypeConverter
 	{
-		Type Type;
+		readonly Type _type;
 
 		/// <summary>
 		/// Constructor
 		/// </summary>
-		/// <param name="Type"></param>
-		public ObjectIdTypeConverter(Type Type)
+		/// <param name="type"></param>
+		public ObjectIdTypeConverter(Type type)
 		{
-			this.Type = Type;
+			_type = type;
 		}
 
 		/// <inheritdoc/>
-		public override bool CanConvertFrom(ITypeDescriptorContext? Context, Type SourceType)
+		public override bool CanConvertFrom(ITypeDescriptorContext? context, Type sourceType)
 		{
-			return SourceType == typeof(string) || base.CanConvertFrom(Context, SourceType);
+			return sourceType == typeof(string) || base.CanConvertFrom(context, sourceType);
 		}
 
 		/// <inheritdoc/>
-		public override object ConvertFrom(ITypeDescriptorContext? Context, CultureInfo? Culture, object Value)
+		public override object ConvertFrom(ITypeDescriptorContext? context, CultureInfo? culture, object value)
 		{
-			return Activator.CreateInstance(Type, Value)!;
+			return Activator.CreateInstance(_type, value)!;
 		}
 
 		/// <inheritdoc/>
-		public override bool CanConvertTo(ITypeDescriptorContext? Context, Type? DestinationType)
+		public override bool CanConvertTo(ITypeDescriptorContext? context, Type? destinationType)
 		{
-			if (DestinationType == null)
+			if (destinationType == null)
 			{
 				return false;
 			}
-			if (DestinationType == typeof(string))
+			if (destinationType == typeof(string))
 			{
 				return true;
 			}
-			if (DestinationType.IsGenericType)
+			if (destinationType.IsGenericType)
 			{
-				Type GenericTypeDefinition = DestinationType.GetGenericTypeDefinition();
-				if (GenericTypeDefinition == typeof(ObjectId<>))
+				Type genericTypeDefinition = destinationType.GetGenericTypeDefinition();
+				if (genericTypeDefinition == typeof(ObjectId<>))
 				{
 					return true;
 				}
-				if (GenericTypeDefinition == typeof(Nullable<>))
+				if (genericTypeDefinition == typeof(Nullable<>))
 				{
-					return CanConvertTo(Context, GenericTypeDefinition.GetGenericArguments()[0]);
+					return CanConvertTo(context, genericTypeDefinition.GetGenericArguments()[0]);
 				}
 			}
 			return false;
 		}
 
 		/// <inheritdoc/>
-		public override object? ConvertTo(ITypeDescriptorContext? Context, CultureInfo? Culture, object? Value, Type DestinationType)
+		public override object? ConvertTo(ITypeDescriptorContext? context, CultureInfo? culture, object? value, Type destinationType)
 		{
-			if (DestinationType == typeof(string))
+			if (destinationType == typeof(string))
 			{
-				return Value?.ToString();
+				return value?.ToString();
 			}
 			else
 			{
-				return Activator.CreateInstance(DestinationType, Value);
+				return Activator.CreateInstance(destinationType, value);
 			}
 		}
 	}
diff --git a/Engine/Source/Programs/Horde/Horde.Build/Utilities/ObjectIdHelpers.cs b/Engine/Source/Programs/Horde/Horde.Build/Utilities/ObjectIdHelpers.cs
--- a/Engine/Source/Programs/Horde/Horde.Build/Utilities/ObjectIdHelpers.cs
+++ b/Engine/Source/Programs/Horde/Horde.Build/Utilities/ObjectIdHelpers.cs
@@ -1,30 +1,26 @@
 // Copyright Epic Games, Inc. All Rights Reserved.
 
 using MongoDB.Bson;
-using System;
-using System.Collections.Generic;
-using System.Linq;
-using System.Threading.Tasks;
 
 namespace Horde.Build.Utilities
 {
 	static class ObjectIdHelpers
 	{
-		public static ObjectId ToObjectId(this string Text)
+		public static ObjectId ToObjectId(this string text)
 		{
-			if (Text.Length == 0)
+			if (text.Length == 0)
 			{
 				return ObjectId.Empty;
 			}
 			else
 			{
-				return ObjectId.Parse(Text);
+				return ObjectId.Parse(text);
 			}
 		}
 
-		public static ObjectId<T> ToObjectId<T>(this string Text)
+		public static ObjectId<T> ToObjectId<T>(this string text)
 		{
-			return new ObjectId<T>(Text.ToObjectId());
+			return new ObjectId<T>(text.ToObjectId());
 		}
 	}
 }
