Horde: Fix coding conventions within Horde.Build.

diff --git a/Engine/Source/Programs/Horde/Horde.Build/Utilities/ObjectSet.cs b/Engine/Source/Programs/Horde/Horde.Build/Utilities/ObjectSet.cs
--- a/Engine/Source/Programs/Horde/Horde.Build/Utilities/ObjectSet.cs
+++ b/Engine/Source/Programs/Horde/Horde.Build/Utilities/ObjectSet.cs
@@ -1,15 +1,13 @@
 // Copyright Epic Games, Inc. All Rights Reserved.
 
-using EpicGames.Core;
-using EpicGames.Horde.Storage;
-using EpicGames.Serialization;
-using EpicGames.Serialization.Converters;
-using Horde.Build.Storage;
 using System;
 using System.Collections.Generic;
 using System.Diagnostics.CodeAnalysis;
 using System.Linq;
 using System.Threading.Tasks;
+using EpicGames.Core;
+using EpicGames.Horde.Storage;
+using EpicGames.Serialization;
 
 namespace Horde.Build.Utilities
 {
@@ -23,12 +21,12 @@ namespace Horde.Build.Utilities
 		public int Length { get; }
 		public IoHash[] Refs { get; }
 
-		public ObjectPackEntry(IoHash Hash, int Offset, int Length, IoHash[] Refs)
+		public ObjectPackEntry(IoHash hash, int offset, int length, IoHash[] refs)
 		{
-			this.Hash = Hash;
-			this.Offset = Offset;
-			this.Length = Length;
-			this.Refs = Refs;
+			Hash = hash;
+			Offset = offset;
+			Length = length;
+			Refs = refs;
 		}
 	}
 
@@ -41,23 +39,23 @@ namespace Horde.Build.Utilities
 		public DateTime Time { get; }
 		public ObjectPackEntry[] Blobs { get; }
 		public IoHash DataHash { get; }
-		public int DataSize;
+		public int _dataSize;
 
-		Dictionary<IoHash, ObjectPackEntry> HashToInfo;
+		readonly Dictionary<IoHash, ObjectPackEntry> _hashToInfo;
 
-		public ObjectPackIndex(DateTime Time, ObjectPackEntry[] Blobs, IoHash DataHash, int DataSize)
+		public ObjectPackIndex(DateTime time, ObjectPackEntry[] blobs, IoHash dataHash, int dataSize)
 		{
-			this.Time = Time;
-			this.Blobs = Blobs;
-			this.DataHash = DataHash;
-			this.DataSize = DataSize;
+			Time = time;
+			Blobs = blobs;
+			DataHash = dataHash;
+			_dataSize = dataSize;
 
-			HashToInfo = Blobs.ToDictionary(x => x.Hash, x => x);
+			_hashToInfo = blobs.ToDictionary(x => x.Hash, x => x);
 		}
 
-		public bool Contains(IoHash Hash) => HashToInfo.ContainsKey(Hash);
+		public bool Contains(IoHash hash) => _hashToInfo.ContainsKey(hash);
 
-		public bool TryGetEntry(IoHash Hash, [NotNullWhen(true)] out ObjectPackEntry? BlobInfo) => HashToInfo.TryGetValue(Hash, out BlobInfo);
+		public bool TryGetEntry(IoHash hash, [NotNullWhen(true)] out ObjectPackEntry? blobInfo) => _hashToInfo.TryGetValue(hash, out blobInfo);
 	}
 
 	/// <summary>
@@ -87,48 +85,48 @@ namespace Horde.Build.Utilities
 		}
 
 		/// <inheritdoc/>
-		public override ObjectPackIndex Read(CbField Field)
+		public override ObjectPackIndex Read(CbField field)
 		{
-			EncodeFormat Format = CbSerializer.Deserialize<EncodeFormat>(Field);
+			EncodeFormat format = CbSerializer.Deserialize<EncodeFormat>(field);
 
-			ObjectPackEntry[] Objects = new ObjectPackEntry[Format.Exports!.Length];
+			ObjectPackEntry[] objects = new ObjectPackEntry[format.Exports!.Length];
 
-			int Offset = 0;
-			for (int Idx = 0; Idx < Format.Exports.Length; Idx++)
+			int offset = 0;
+			for (int idx = 0; idx < format.Exports.Length; idx++)
 			{
-				Objects[Idx] = new ObjectPackEntry(Format.Exports[Idx], Offset, Format.Lengths![Idx], Format.Refs![Idx]);
-				Offset += Format.Lengths[Idx];
+				objects[idx] = new ObjectPackEntry(format.Exports[idx], offset, format.Lengths![idx], format.Refs![idx]);
+				offset += format.Lengths[idx];
 			}
 
-			return new ObjectPackIndex(Format.Time, Objects, Format.DataHash, Format.DataSize);
+			return new ObjectPackIndex(format.Time, objects, format.DataHash, format.DataSize);
 		}
 
 		/// <inheritdoc/>
-		public override void Write(CbWriter Writer, ObjectPackIndex Index)
+		public override void Write(CbWriter writer, ObjectPackIndex index)
 		{
-			Writer.BeginObject();
-			WriteInternal(Writer, Index);
-			Writer.EndObject();
+			writer.BeginObject();
+			WriteInternal(writer, index);
+			writer.EndObject();
 		}
 
 		/// <inheritdoc/>
-		public override void WriteNamed(CbWriter Writer, Utf8String Name, ObjectPackIndex Index)
+		public override void WriteNamed(CbWriter writer, Utf8String name, ObjectPackIndex index)
 		{
-			Writer.BeginObject(Name);
-			WriteInternal(Writer, Index);
-			Writer.EndObject();
+			writer.BeginObject(name);
+			WriteInternal(writer, index);
+			writer.EndObject();
 		}
 
-		static void WriteInternal(CbWriter Writer, ObjectPackIndex Index)
+		static void WriteInternal(CbWriter writer, ObjectPackIndex index)
 		{
-			EncodeFormat Format = new EncodeFormat();
-			Format.Time = Index.Time;
-			Format.Exports = Index.Blobs.ConvertAll(x => x.Hash).ToArray();
-			Format.Lengths = Index.Blobs.ConvertAll(x => x.Length).ToArray();
-			Format.Refs = Index.Blobs.ConvertAll(x => x.Refs).ToArray();
-			Format.DataHash = Index.DataHash;
-			Format.DataSize = Index.DataSize;
-			CbSerializer.Serialize(Writer, Format);
+			EncodeFormat format = new EncodeFormat();
+			format.Time = index.Time;
+			format.Exports = index.Blobs.ConvertAll(x => x.Hash).ToArray();
+			format.Lengths = index.Blobs.ConvertAll(x => x.Length).ToArray();
+			format.Refs = index.Blobs.ConvertAll(x => x.Refs).ToArray();
+			format.DataHash = index.DataHash;
+			format.DataSize = index._dataSize;
+			CbSerializer.Serialize(writer, format);
 		}
 	}
 
@@ -137,38 +135,38 @@ namespace Horde.Build.Utilities
 	/// </summary>
 	class ObjectSet
 	{
-		readonly IStorageClient StorageClient;
-		readonly NamespaceId NamespaceId;
+		readonly IStorageClient _storageClient;
+		readonly NamespaceId _namespaceId;
 		public int MaxPackSize { get; }
 
 		public HashSet<IoHash> RootSet { get; set; } = new HashSet<IoHash>();
 
-		DateTime Time;
+		DateTime _time;
 
-		int NextPackSize;
-		byte[] NextPackData;
-		List<ObjectPackEntry> NextPackEntries = new List<ObjectPackEntry>();
-		Dictionary<IoHash, ObjectPackEntry> NextPackHashToEntry = new Dictionary<IoHash, ObjectPackEntry>();
+		int _nextPackSize;
+		byte[] _nextPackData;
+		readonly List<ObjectPackEntry> _nextPackEntries = new List<ObjectPackEntry>();
+		readonly Dictionary<IoHash, ObjectPackEntry> _nextPackHashToEntry = new Dictionary<IoHash, ObjectPackEntry>();
 
 		public List<ObjectPackIndex> PackIndexes { get; } = new List<ObjectPackIndex>();
-		List<Task> WriteTasks = new List<Task>();
+		readonly List<Task> _writeTasks = new List<Task>();
 
 		/// <summary>
 		/// Constructor
 		/// </summary>
-		/// <param name="StorageClient"></param>
-		/// <param name="NamespaceId"></param>
-		/// <param name="MaxPackSize"></param>
-		/// <param name="Time">The initial update time; used to determine the age of blobs</param>
-		public ObjectSet(IStorageClient StorageClient, NamespaceId NamespaceId, int MaxPackSize, DateTime Time)
+		/// <param name="storageClient"></param>
+		/// <param name="namespaceId"></param>
+		/// <param name="maxPackSize"></param>
+		/// <param name="time">The initial update time; used to determine the age of blobs</param>
+		public ObjectSet(IStorageClient storageClient, NamespaceId namespaceId, int maxPackSize, DateTime time)
 		{
-			this.StorageClient = StorageClient;
-			this.NamespaceId = NamespaceId;
-			this.MaxPackSize = MaxPackSize;
+			_storageClient = storageClient;
+			_namespaceId = namespaceId;
+			MaxPackSize = maxPackSize;
 
-			NextPackData = null!;
+			_nextPackData = null!;
 
-			SetTime(Time);
+			SetTime(time);
 			Reset();
 		}
 
@@ -177,93 +175,93 @@ namespace Horde.Build.Utilities
 		/// </summary>
 		void Reset()
 		{
-			NextPackSize = 0;
-			NextPackData = new byte[MaxPackSize];
-			NextPackEntries.Clear();
-			NextPackHashToEntry.Clear();
+			_nextPackSize = 0;
+			_nextPackData = new byte[MaxPackSize];
+			_nextPackEntries.Clear();
+			_nextPackHashToEntry.Clear();
 		}
 
 		/// <summary>
 		/// Reset the current timestamp
 		/// </summary>
-		/// <param name="Time">The new timestamp</param>
-		public void SetTime(DateTime Time)
+		/// <param name="time">The new timestamp</param>
+		public void SetTime(DateTime time)
 		{
-			this.Time = Time;
+			_time = time;
 		}
 
 		/// <summary>
 		/// Copies an existing entry into storage
 		/// </summary>
-		/// <param name="Hash">Hash of the data</param>
-		/// <param name="Data">The data buffer</param>
-		/// <param name="Refs">References to other objects</param>
-		public void Add(IoHash Hash, ReadOnlySpan<byte> Data, ReadOnlySpan<IoHash> Refs)
+		/// <param name="hash">Hash of the data</param>
+		/// <param name="data">The data buffer</param>
+		/// <param name="refs">References to other objects</param>
+		public void Add(IoHash hash, ReadOnlySpan<byte> data, ReadOnlySpan<IoHash> refs)
 		{
-			if (!NextPackHashToEntry.ContainsKey(Hash))
+			if (!_nextPackHashToEntry.ContainsKey(hash))
 			{
 				// Create enough space for the new data
-				CreateSpace(Data.Length);
+				CreateSpace(data.Length);
 
 				// Copy the data into the buffer
-				Data.CopyTo(NextPackData.AsSpan(NextPackSize));
+				data.CopyTo(_nextPackData.AsSpan(_nextPackSize));
 
 				// Add the blob
-				ObjectPackEntry Entry = new ObjectPackEntry(Hash, NextPackSize, Data.Length, Refs.ToArray());
-				NextPackHashToEntry.Add(Hash, Entry);
-				NextPackEntries.Add(Entry);
-				NextPackSize += Data.Length;
+				ObjectPackEntry entry = new ObjectPackEntry(hash, _nextPackSize, data.Length, refs.ToArray());
+				_nextPackHashToEntry.Add(hash, entry);
+				_nextPackEntries.Add(entry);
+				_nextPackSize += data.Length;
 			}
 		}
 
 		/// <summary>
 		/// Adds an item to the packer
 		/// </summary>
-		/// <param name="Size">Size of the data</param>
-		/// <param name="ReadData">Delegate to copy the data into a span</param>
-		/// <param name="Refs">References to other objects</param>
-		public IoHash Add(int Size, Action<Memory<byte>> ReadData, IoHash[] Refs)
+		/// <param name="size">Size of the data</param>
+		/// <param name="readData">Delegate to copy the data into a span</param>
+		/// <param name="refs">References to other objects</param>
+		public IoHash Add(int size, Action<Memory<byte>> readData, IoHash[] refs)
 		{
 			// Get the last blob and make sure there's enough space in it
-			CreateSpace(Size);
+			CreateSpace(size);
 
 			// Copy the data into the new blob
-			Memory<byte> Output = NextPackData.AsMemory(NextPackSize, Size);
-			ReadData(Output);
+			Memory<byte> output = _nextPackData.AsMemory(_nextPackSize, size);
+			readData(output);
 
 			// Update the metadata for it
-			IoHash Hash = IoHash.Compute(Output.Span);
-			if (!NextPackHashToEntry.ContainsKey(Hash))
+			IoHash hash = IoHash.Compute(output.Span);
+			if (!_nextPackHashToEntry.ContainsKey(hash))
 			{
-				ObjectPackEntry Entry = new ObjectPackEntry(Hash, NextPackSize, Size, Refs);
-				NextPackHashToEntry.Add(Hash, Entry);
-				NextPackEntries.Add(Entry);
-				NextPackSize += Size;
+				ObjectPackEntry entry = new ObjectPackEntry(hash, _nextPackSize, size, refs);
+				_nextPackHashToEntry.Add(hash, entry);
+				_nextPackEntries.Add(entry);
+				_nextPackSize += size;
 			}
-			return Hash;
+			return hash;
 		}
 
 		/// <summary>
 		/// Finds data for an object with the given hash, from the current pack files
 		/// </summary>
-		/// <param name="Hash"></param>
+		/// <param name="hash"></param>
 		/// <returns></returns>
-		public async Task<ReadOnlyMemory<byte>> GetObjectDataAsync(IoHash Hash)
+		public async Task<ReadOnlyMemory<byte>> GetObjectDataAsync(IoHash hash)
 		{
-			await Task.WhenAll(WriteTasks);
+			await Task.WhenAll(_writeTasks);
 
-			ObjectPackEntry? Entry;
-			if (NextPackHashToEntry.TryGetValue(Hash, out Entry))
+			ObjectPackEntry? entry;
+			if (_nextPackHashToEntry.TryGetValue(hash, out entry))
 			{
-				return NextPackData.AsMemory(Entry.Offset, Entry.Length);
+				return _nextPackData.AsMemory(entry.Offset, entry.Length);
 			}
 
-			foreach (ObjectPackIndex Pack in PackIndexes)
+			foreach (ObjectPackIndex pack in PackIndexes)
 			{
-				if (Pack.TryGetEntry(Hash, out Entry))
+				if (pack.TryGetEntry(hash, out entry))
 				{
-					ReadOnlyMemory<byte> PackData = await StorageClient.ReadBlobToMemoryAsync(NamespaceId, Pack.DataHash);
-					return PackData.Slice(Entry.Offset, Entry.Length);
+					ReadOnlyMemory<byte> packData = await _storageClient.ReadBlobToMemoryAsync(_namespaceId, pack.DataHash);
+					return packData.Slice(entry.Offset, entry.Length);
 				}
 			}
 
@@ -273,28 +271,28 @@ namespace Horde.Build.Utilities
 		/// <summary>
 		/// Tries to find an entry for the given hash from the current set of pack files
 		/// </summary>
-		/// <param name="Hash">Hash of the </param>
-		/// <param name="Entry"></param>
+		/// <param name="hash">Hash of the </param>
+		/// <param name="entry"></param>
 		/// <returns></returns>
-		public bool TryGetEntry(IoHash Hash, [NotNullWhen(true)] out ObjectPackEntry? Entry)
+		public bool TryGetEntry(IoHash hash, [NotNullWhen(true)] out ObjectPackEntry? entry)
 		{
-			ObjectPackEntry? LocalEntry;
-			if (NextPackHashToEntry.TryGetValue(Hash, out LocalEntry))
+			ObjectPackEntry? localEntry;
+			if (_nextPackHashToEntry.TryGetValue(hash, out localEntry))
 			{
-				Entry = LocalEntry;
+				entry = localEntry;
 				return true;
 			}
 
-			foreach (ObjectPackIndex Pack in PackIndexes)
+			foreach (ObjectPackIndex pack in PackIndexes)
 			{
-				if (Pack.TryGetEntry(Hash, out LocalEntry))
+				if (pack.TryGetEntry(hash, out localEntry))
 				{
-					Entry = LocalEntry;
+					entry = localEntry;
 					return true;
 				}
 			}
 
-			Entry = null;
+			entry = null;
 			return false;
 		}
 
@@ -304,116 +302,116 @@ namespace Horde.Build.Utilities
 		public async Task FlushAsync()
 		{
 			// Find the live set of objects
-			HashSet<IoHash> LiveSet = new HashSet<IoHash>();
-			foreach(IoHash RootHash in RootSet)
+			HashSet<IoHash> liveSet = new HashSet<IoHash>();
+			foreach(IoHash rootHash in RootSet)
 			{
-				FindLiveSet(RootHash, LiveSet);
+				FindLiveSet(rootHash, liveSet);
 			}
 
 			// Find the total cost of all the current blobs, then loop through the blobs trying to find a more optimal arrangement
-			double TotalCost = PackIndexes.Sum(x => GetCostHeuristic(x)) + GetCostHeuristic(NextPackSize, TimeSpan.Zero);
+			double totalCost = PackIndexes.Sum(x => GetCostHeuristic(x)) + GetCostHeuristic(_nextPackSize, TimeSpan.Zero);
 			for (; ; )
 			{
 				// Exclude any objects that are in the pending blobs, since we will always upload these
-				HashSet<IoHash> NewLiveSet = new HashSet<IoHash>(LiveSet);
-				NewLiveSet.ExceptWith(NextPackHashToEntry.Values.Select(x => x.Hash));
+				HashSet<IoHash> newLiveSet = new HashSet<IoHash>(liveSet);
+				newLiveSet.ExceptWith(_nextPackHashToEntry.Values.Select(x => x.Hash));
 
 				// Get the size and cost of the next blob
-				double NextBlobCost = GetCostHeuristic(NextPackSize, TimeSpan.Zero);
+				double nextBlobCost = GetCostHeuristic(_nextPackSize, TimeSpan.Zero);
 
 				// Pass through all the blobs to find the best one to merge in
-				double MergeCost = TotalCost;
-				ObjectPackIndex? MergePack = null;
-				for (int Idx = PackIndexes.Count - 1; Idx >= 0; Idx--)
+				double mergeCost = totalCost;
+				ObjectPackIndex? mergePack = null;
+				for (int idx = PackIndexes.Count - 1; idx >= 0; idx--)
 				{
-					ObjectPackIndex PackIndex = PackIndexes[Idx];
+					ObjectPackIndex packIndex = PackIndexes[idx];
 
 					// Try to merge any old blobs with the next blob
-					if (PackIndex.Time < Time)
+					if (packIndex.Time < _time)
 					{
 						// Calculate the cost of the last blob if we merge this one with it. We remove blobs as we iterate
 						// through the list, since subsequent blobs will not usefully contribute the same items.
-						double NewTotalCost = TotalCost - GetCostHeuristic(PackIndex) - NextBlobCost;
+						double newTotalCost = totalCost - GetCostHeuristic(packIndex) - nextBlobCost;
 
-						int NewNextPackSize = NextPackSize;
-						foreach (ObjectPackEntry Entry in PackIndex.Blobs)
+						int newNextPackSize = _nextPackSize;
+						foreach (ObjectPackEntry entry in packIndex.Blobs)
 						{
-							if (LiveSet.Contains(Entry.Hash))
+							if (liveSet.Contains(entry.Hash))
 							{
-								if (NewNextPackSize + Entry.Length > MaxPackSize)
+								if (newNextPackSize + entry.Length > MaxPackSize)
 								{
-									NewTotalCost += GetCostHeuristic(NewNextPackSize, TimeSpan.Zero);
-									NewNextPackSize = 0;
+									newTotalCost += GetCostHeuristic(newNextPackSize, TimeSpan.Zero);
+									newNextPackSize = 0;
 								}
-								NewNextPackSize += Entry.Length;
+								newNextPackSize += entry.Length;
 							}
 						}
 
-						NewTotalCost += GetCostHeuristic(NewNextPackSize, TimeSpan.Zero);
+						newTotalCost += GetCostHeuristic(newNextPackSize, TimeSpan.Zero);
 
 						// Compute the potential cost if we replace the partial blob with the useful parts of this blob
-						if (NewTotalCost < MergeCost)
+						if (newTotalCost < mergeCost)
 						{
-							MergePack = PackIndex;
-							MergeCost = NewTotalCost;
+							mergePack = packIndex;
+							mergeCost = newTotalCost;
 						}
 					}
 
 					// Remove any items in this blob from the remaining live set. No other blobs need to include them.
-					NewLiveSet.ExceptWith(PackIndex.Blobs.Select(x => x.Hash));
+					newLiveSet.ExceptWith(packIndex.Blobs.Select(x => x.Hash));
 				}
 
 				// Bail out if we didn't find anything to merge
-				if (MergePack == null)
+				if (mergePack == null)
 				{
 					break;
 				}
 
 				// Get the data for this blob
-				ReadOnlyMemory<byte> MergeData = await StorageClient.ReadBlobToMemoryAsync(NamespaceId, MergePack.DataHash);
+				ReadOnlyMemory<byte> mergeData = await _storageClient.ReadBlobToMemoryAsync(_namespaceId, mergePack.DataHash);
 
 				// Add anything that's still part of the live set into the new blobs
-				int Offset = 0;
-				foreach (ObjectPackEntry Blob in MergePack.Blobs)
+				int offset = 0;
+				foreach (ObjectPackEntry blob in mergePack.Blobs)
 				{
-					if (LiveSet.Contains(Blob.Hash))
+					if (liveSet.Contains(blob.Hash))
 					{
-						ReadOnlyMemory<byte> Data = MergeData.Slice(Offset, Blob.Length);
-						Add(Blob.Hash, Data.Span, Blob.Refs);
+						ReadOnlyMemory<byte> data = mergeData.Slice(offset, blob.Length);
+						Add(blob.Hash, data.Span, blob.Refs);
 					}
-					Offset += Blob.Length;
+					offset += blob.Length;
 				}
 
 				// Discard the old blob
-				PackIndexes.Remove(MergePack);
-				TotalCost = MergeCost;
+				PackIndexes.Remove(mergePack);
+				totalCost = mergeCost;
 			}
 
 			// Write the current blob
 			FlushCurrentPack();
 
 			// Wait for all the writes to finish
-			await Task.WhenAll(WriteTasks);
-			WriteTasks.Clear();
+			await Task.WhenAll(_writeTasks);
+			_writeTasks.Clear();
 		}
 
 		/// <summary>
 		/// Finds the live set for a particular tree, and updates tree entries with the size of used items within them
 		/// </summary>
-		/// <param name="Hash"></param>
-		/// <param name="LiveSet"></param>
-		void FindLiveSet(IoHash Hash, HashSet<IoHash> LiveSet)
+		/// <param name="hash"></param>
+		/// <param name="liveSet"></param>
+		void FindLiveSet(IoHash hash, HashSet<IoHash> liveSet)
 		{
-			if (LiveSet.Add(Hash))
+			if (liveSet.Add(hash))
 			{
-				ObjectPackEntry? Entry;
-				if (!TryGetEntry(Hash, out Entry))
+				ObjectPackEntry? entry;
+				if (!TryGetEntry(hash, out entry))
 				{
-					throw new Exception($"Missing blob {Hash} from working set");
+					throw new Exception($"Missing blob {hash} from working set");
 				}
-				foreach (IoHash Ref in Entry.Refs)
+				foreach (IoHash refHash in entry.Refs)
 				{
-					FindLiveSet(Ref, LiveSet);
+					FindLiveSet(refHash, liveSet);
 				}
 			}
 		}
@@ -421,19 +419,19 @@ namespace Horde.Build.Utilities
 		/// <summary>
 		/// Creates enough space to store the given block of data
 		/// </summary>
-		/// <param name="Size"></param>
-		void CreateSpace(int Size)
+		/// <param name="size"></param>
+		void CreateSpace(int size)
 		{
 			// Get the last blob and make sure there's enough space in it
-			if (NextPackSize + Size > MaxPackSize)
+			if (_nextPackSize + size > MaxPackSize)
 			{
 				FlushCurrentPack();
 			}
 
 			// Resize the next blob buffer if necessary
-			if (Size > NextPackData.Length)
+			if (size > _nextPackData.Length)
 			{
-				Array.Resize(ref NextPackData, Size);
+				Array.Resize(ref _nextPackData, size);
 			}
 		}
 
@@ -442,35 +440,35 @@ namespace Horde.Build.Utilities
 		/// </summary>
 		void FlushCurrentPack()
 		{
-			if (NextPackSize > 0)
+			if (_nextPackSize > 0)
 			{
 				// Write the buffer to storage
-				Array.Resize(ref NextPackData, NextPackSize);
-				ReadOnlyMemory<byte> Data = NextPackData;
-				IoHash DataHash = IoHash.Compute(Data.Span);
-				WriteTasks.Add(Task.Run(() => StorageClient.WriteBlobFromMemoryAsync(NamespaceId, DataHash, Data)));
+				Array.Resize(ref _nextPackData, _nextPackSize);
+				ReadOnlyMemory<byte> data = _nextPackData;
+				IoHash dataHash = IoHash.Compute(data.Span);
+				_writeTasks.Add(Task.Run(() => _storageClient.WriteBlobFromMemoryAsync(_namespaceId, dataHash, data)));
 
 				// Create the new index
-				ObjectPackIndex Index = new ObjectPackIndex(Time, NextPackEntries.ToArray(), DataHash, NextPackSize);
-				PackIndexes.Add(Index);
+				ObjectPackIndex index = new ObjectPackIndex(_time, _nextPackEntries.ToArray(), dataHash, _nextPackSize);
+				PackIndexes.Add(index);
 
 				// Clear the next pack buffer
 				Reset();
 			}
 		}
 
-		/// <inheritdoc cref="GetCostHeuristic(int, TimeSpan)"/>
-		/// <param name="Index">Index to calculate the heuristic for</param>
-		public double GetCostHeuristic(ObjectPackIndex Index) => GetCostHeuristic(Index.Blobs.Length, Time - Index.Time);
+		/// <inheritdoc cref="GetCostHeuristic(Int32, TimeSpan)"/>
+		/// <param name="index">Index to calculate the heuristic for</param>
+		public double GetCostHeuristic(ObjectPackIndex index) => GetCostHeuristic(index.Blobs.Length, _time - index.Time);
 
 		/// <summary>
 		/// Heuristic which estimates the cost of a particular blob. This is used to compare scenarios of merging blobs to reduce download
 		/// size against keeping older blobs which a lot of agents already have.
 		/// </summary>
-		/// <param name="Size">Size of the blob</param>
-		/// <param name="Age">Age of the blob</param>
+		/// <param name="size">Size of the blob</param>
+		/// <param name="age">Age of the blob</param>
 		/// <returns>Heuristic for the cost of a blob</returns>
-		public static double GetCostHeuristic(int Size, TimeSpan Age)
+		public static double GetCostHeuristic(int size, TimeSpan age)
 		{
 			// Time overhead to starting a download
 			const double DownloadInit = 0.1;
@@ -482,10 +480,10 @@ namespace Horde.Build.Utilities
 			const double CleanSyncProbability = 0.2;
 
 			// Average length of time between agents having to update
-			TimeSpan AverageCoherence = TimeSpan.FromHours(4.0);
+			TimeSpan averageCoherence = TimeSpan.FromHours(4.0);
 
 			// Scale the age into a -1.0 -> 1.0 range around AverageCoherence
-			double ScaledAge = (AverageCoherence - Age).TotalSeconds / AverageCoherence.TotalSeconds;
+			double scaledAge = (averageCoherence - age).TotalSeconds / averageCoherence.TotalSeconds;
 
 			// Get the probability of agents having to sync this blob based on its age. This is modeled as a logistic function (1 / (1 + e^-x))
 			// with value of 0.5 at AverageCoherence, and MaxInterval at zero.
@@ -495,17 +493,17 @@ namespace Horde.Build.Utilities
 			//    e^-x = (1 / MaxInterval) - 1
 			//    x = -ln((1 / MaxInterval) - 1)
 			const double MaxInterval = 0.95;
-			double SigmoidScale = -Math.Log((1.0 / MaxInterval) - 1.0);
+			double sigmoidScale = -Math.Log((1.0 / MaxInterval) - 1.0);
 
 			// Find the probability of having to sync this 
-			double Param = ScaledAge * SigmoidScale;
-			double Probability = 1.0 / (1.0 + Math.Exp(-Param));
+			double param = scaledAge * sigmoidScale;
+			double probability = 1.0 / (1.0 + Math.Exp(-param));
 
 			// Scale the probability against having to do a full sync
-			Probability = CleanSyncProbability + (Probability * (1.0 - CleanSyncProbability));
+			probability = CleanSyncProbability + (probability * (1.0 - CleanSyncProbability));
 
 			// Compute the final cost estimate; the amount of time we expect agents to spend downloading the file
-			return Probability * (DownloadInit + (Size / DownloadRate));
+			return probability * (DownloadInit + (size / DownloadRate));
 		}
 	}
 }
diff --git a/Engine/Source/Programs/Horde/Horde.Build/Utilities/PropertyFilter.cs b/Engine/Source/Programs/Horde/Horde.Build/Utilities/PropertyFilter.cs
--- a/Engine/Source/Programs/Horde/Horde.Build/Utilities/PropertyFilter.cs
+++ b/Engine/Source/Programs/Horde/Horde.Build/Utilities/PropertyFilter.cs
@@ -1,16 +1,11 @@
 // Copyright Epic Games, Inc. All Rights Reserved.
 
-using Microsoft.AspNetCore.Mvc;
-using Microsoft.AspNetCore.Mvc.ModelBinding;
 using System;
 using System.Collections;
 using System.Collections.Generic;
 using System.ComponentModel;
-using System.Diagnostics.CodeAnalysis;
 using System.Globalization;
 using System.Reflection;
-using System.Text;
-using System.Threading.Tasks;
 
 namespace Horde.Build.Utilities
 {
@@ -23,137 +18,137 @@ namespace Horde.Build.Utilities
 		/// <summary>
 		/// Specific filters for properties of this object. A key of '*' will match any property not matched by anything else, and a value of null will include the whole object.
 		/// </summary>
-		Dictionary<string, PropertyFilter?> NameToFilter = new Dictionary<string, PropertyFilter?>(StringComparer.OrdinalIgnoreCase);
+		readonly Dictionary<string, PropertyFilter?> _nameToFilter = new Dictionary<string, PropertyFilter?>(StringComparer.OrdinalIgnoreCase);
 
 		/// <summary>
 		/// Parse a list of fields as a property filter
 		/// </summary>
-		/// <param name="Fields">List of fields to parse, separated by commas</param>
+		/// <param name="fields">List of fields to parse, separated by commas</param>
 		/// <returns>New property filter</returns>
-		public static PropertyFilter Parse(string? Fields)
+		public static PropertyFilter Parse(string? fields)
 		{
-			PropertyFilter RootFilter = new PropertyFilter();
-			if(Fields != null)
+			PropertyFilter rootFilter = new PropertyFilter();
+			if(fields != null)
 			{
-				foreach (string Field in Fields.Split(','))
+				foreach (string field in fields.Split(','))
 				{
 					// Split the field name into segments
-					string[] Segments = Field.Split('.');
+					string[] segments = field.Split('.');
 
 					// Create the tree of segments
-					Dictionary<string, PropertyFilter?> NameToFilter = RootFilter.NameToFilter;
-					for (int Idx = 0; Idx < Segments.Length - 1; Idx++)
+					Dictionary<string, PropertyFilter?> nameToFilter = rootFilter._nameToFilter;
+					for (int idx = 0; idx < segments.Length - 1; idx++)
 					{
-						PropertyFilter? NextFilter;
-						if (!NameToFilter.TryGetValue(Segments[Idx], out NextFilter))
+						PropertyFilter? nextFilter;
+						if (!nameToFilter.TryGetValue(segments[idx], out nextFilter))
 						{
-							NextFilter = new PropertyFilter();
-							NameToFilter[Segments[Idx]] = NextFilter;
+							nextFilter = new PropertyFilter();
+							nameToFilter[segments[idx]] = nextFilter;
 						}
-						else if (NextFilter == null)
+						else if (nextFilter == null)
 						{
 							break;
 						}
-						NameToFilter = NextFilter.NameToFilter;
+						nameToFilter = nextFilter._nameToFilter;
 					}
-					NameToFilter[Segments[Segments.Length - 1]] = null;
+					nameToFilter[segments[^1]] = null;
 				}
 			}
-			return RootFilter;
+			return rootFilter;
 		}
 
 		/// <summary>
 		/// Checks whether the filter contains the given field
 		/// </summary>
-		/// <param name="Filter">The filter to check</param>
-		/// <param name="Field">Name of the field to check for</param>
+		/// <param name="filter">The filter to check</param>
+		/// <param name="field">Name of the field to check for</param>
 		/// <returns>True if the filter includes the given field</returns>
-		public static bool Includes(PropertyFilter? Filter, string Field)
+		public static bool Includes(PropertyFilter? filter, string field)
 		{
-			return Filter == null || Filter.Includes(Field);
+			return filter == null || filter.Includes(field);
 		}
 
 		/// <summary>
 		/// Attempts to apply a property filter to a response object, where the property filter may be null.
 		/// </summary>
-		/// <param name="Response">Response object</param>
-		/// <param name="Filter">The filter to apply</param>
+		/// <param name="response">Response object</param>
+		/// <param name="filter">The filter to apply</param>
 		/// <returns>Filtered object</returns>
-		public static object Apply(object Response, PropertyFilter? Filter)
+		public static object Apply(object response, PropertyFilter? filter)
 		{
-			return (Filter == null) ? Response : Filter.ApplyTo(Response);
+			return (filter == null) ? response : filter.ApplyTo(response);
 		}
 
 		/// <summary>
 		/// Checks whether the filter contains the given field
 		/// </summary>
-		/// <param name="Field">Name of the field to check for</param>
+		/// <param name="field">Name of the field to check for</param>
 		/// <returns>True if the filter includes the given field</returns>
-		public bool Includes(string Field)
+		public bool Includes(string field)
 		{
-			return NameToFilter.Count == 0 || NameToFilter.ContainsKey(Field) || NameToFilter.ContainsKey("*");
+			return _nameToFilter.Count == 0 || _nameToFilter.ContainsKey(field) || _nameToFilter.ContainsKey("*");
 		}
 
 		/// <summary>
 		/// Extract the given fields from the response
 		/// </summary>
-		/// <param name="Response">The response to filter</param>
+		/// <param name="response">The response to filter</param>
 		/// <returns>Filtered response</returns>
-		public object ApplyTo(object Response)
+		public object ApplyTo(object response)
 		{
 			// If there are no filters at this depth, return the whole object
-			if (NameToFilter.Count == 0)
+			if (_nameToFilter.Count == 0)
 			{
-				return Response;
+				return response;
 			}
 
 			// Check if this object is a list. If it is, we'll apply the filter to each element in the list.
-			IList? List = Response as IList;
-			if (List != null)
+			IList? list = response as IList;
+			if (list != null)
 			{
-				List<object?> NewResponse = new List<object?>();
-				for (int Index = 0; Index < List.Count; Index++)
+				List<object?> newResponse = new List<object?>();
+				for (int index = 0; index < list.Count; index++)
 				{
-					object? Value = List[Index];
-					if (Value == null)
+					object? value = list[index];
+					if (value == null)
 					{
-						NewResponse.Add(Value);
+						newResponse.Add(value);
 					}
 					else
 					{
-						NewResponse.Add(ApplyTo(Value));
+						newResponse.Add(ApplyTo(value));
 					}
 				}
-				return NewResponse;
+				return newResponse;
 			}
 			else
 			{
-				Dictionary<string, object?> NewResponse = new Dictionary<string, object?>(StringComparer.Ordinal);
+				Dictionary<string, object?> newResponse = new Dictionary<string, object?>(StringComparer.Ordinal);
 
 				// Check if this object is a dictionary.
-				IDictionary? Dictionary = Response as IDictionary;
-				if (Dictionary != null)
+				IDictionary? dictionary = response as IDictionary;
+				if (dictionary != null)
 				{
-					foreach (DictionaryEntry? Entry in Dictionary)
+					foreach (DictionaryEntry? entry in dictionary)
 					{
-						if (Entry.HasValue)
+						if (entry.HasValue)
 						{
-							string? Key = Entry.Value.Key.ToString();
-							if (Key != null)
+							string? key = entry.Value.Key.ToString();
+							if (key != null)
 							{
-								Key = ConvertPascalToCamelCase(Key);
+								key = ConvertPascalToCamelCase(key);
 
-								PropertyFilter? Filter;
-								if (NameToFilter.TryGetValue(Key, out Filter))
+								PropertyFilter? filter;
+								if (_nameToFilter.TryGetValue(key, out filter))
 								{
-									object? Value = Entry.Value.Value;
-									if (Filter == null)
+									object? value = entry.Value.Value;
+									if (filter == null)
 									{
-										NewResponse[Key] = Value;
+										newResponse[key] = value;
 									}
-									else if (Value != null)
+									else if (value != null)
 									{
-										NewResponse[Key] = Filter.ApplyTo(Value);
+										newResponse[key] = filter.ApplyTo(value);
 									}
 								}
 							}
@@ -163,42 +158,42 @@ namespace Horde.Build.Utilities
 				else
 				{
 					// Otherwise try to get the properties of this object
-					foreach (PropertyInfo Property in Response.GetType().GetProperties(BindingFlags.Public | BindingFlags.GetProperty | BindingFlags.Instance))
+					foreach (PropertyInfo property in response.GetType().GetProperties(BindingFlags.Public | BindingFlags.GetProperty | BindingFlags.Instance))
 					{
-						PropertyFilter? Filter;
-						if (NameToFilter.TryGetValue(Property.Name, out Filter))
+						PropertyFilter? filter;
+						if (_nameToFilter.TryGetValue(property.Name, out filter))
 						{
-							string Key = ConvertPascalToCamelCase(Property.Name);
-							object? Value = Property.GetValue(Response);
-							if (Filter == null)
+							string key = ConvertPascalToCamelCase(property.Name);
+							object? value = property.GetValue(response);
+							if (filter == null)
 							{
-								NewResponse[Key] = Value;
+								newResponse[key] = value;
 							}
-							else if (Value != null)
+							else if (value != null)
 							{
-								NewResponse[Key] = Filter.ApplyTo(Value);
+								newResponse[key] = filter.ApplyTo(value);
 							}
 						}
 					}
 				}
-				return NewResponse;
+				return newResponse;
 			}
 		}
 
 		/// <summary>
 		/// Convert a property name to camel case
 		/// </summary>
-		/// <param name="Name">The name to convert to camelCase</param>
+		/// <param name="name">The name to convert to camelCase</param>
 		/// <returns>Camelcase version of the name</returns>
-		static string ConvertPascalToCamelCase(string Name)
+		static string ConvertPascalToCamelCase(string name)
 		{
-			if (Name[0] >= 'A' && Name[0] <= 'Z')
+			if (name[0] >= 'A' && name[0] <= 'Z')
 			{
-				return (char)(Name[0] - 'A' + 'a') + Name.Substring(1);
+				return (char)(name[0] - 'A' + 'a') + name.Substring(1);
 			}
 			else
 			{
-				return Name;
+				return name;
 			}
 		}
 	}
@@ -209,15 +204,15 @@ namespace Horde.Build.Utilities
 	public class PropertyFilterTypeConverter : TypeConverter
 	{
 		/// <inheritdoc/>
-		public override bool CanConvertFrom(ITypeDescriptorContext? Context, Type SourceType)
+		public override bool CanConvertFrom(ITypeDescriptorContext? context, Type sourceType)
 		{
-			return SourceType == typeof(string);
+			return sourceType == typeof(string);
 		}
 
 		/// <inheritdoc/>
-		public override object ConvertFrom(ITypeDescriptorContext? Context, CultureInfo? Culture, object Value)
+		public override object ConvertFrom(ITypeDescriptorContext? context, CultureInfo? culture, object value)
 		{
-			return PropertyFilter.Parse((string)Value);
+			return PropertyFilter.Parse((string)value);
 		}
 	}
 
@@ -229,18 +224,18 @@ namespace Horde.Build.Utilities
 		/// <summary>
 		/// Apply a filter to a response object
 		/// </summary>
-		/// <param name="Obj">The object to filter</param>
-		/// <param name="Filter">The filter to apply</param>
+		/// <param name="obj">The object to filter</param>
+		/// <param name="filter">The filter to apply</param>
 		/// <returns>The filtered object</returns>
-		public static object ApplyFilter(this object Obj, PropertyFilter? Filter)
+		public static object ApplyFilter(this object obj, PropertyFilter? filter)
 		{
-			if (Filter == null)
+			if (filter == null)
 			{
-				return Obj;
+				return obj;
 			}
 			else
 			{
-				return Filter.ApplyTo(Obj);
+				return filter.ApplyTo(obj);
 			}
 		}
 	}
diff --git a/Engine/Source/Programs/Horde/Horde.Build/Utilities/ReadOnlyTrie.cs b/Engine/Source/Programs/Horde/Horde.Build/Utilities/ReadOnlyTrie.cs
--- a/Engine/Source/Programs/Horde/Horde.Build/Utilities/ReadOnlyTrie.cs
+++ b/Engine/Source/Programs/Horde/Horde.Build/Utilities/ReadOnlyTrie.cs
@@ -1,13 +1,10 @@
 // Copyright Epic Games, Inc. All Rights Reserved.
 
-using EpicGames.Core;
-using Horde.Build.Controllers;
 using System;
 using System.Collections;
 using System.Collections.Generic;
-using System.Linq;
 using System.Runtime.InteropServices;
-using System.Threading.Tasks;
+using EpicGames.Core;
 
 namespace Horde.Build.Logs
 {
@@ -24,21 +21,21 @@ namespace Horde.Build.Logs
 			/// <summary>
 			/// The current node index
 			/// </summary>
-			public int Index;
+			public int _index;
 
 			/// <summary>
 			/// Value in the current node (0-15)
 			/// </summary>
-			public int Value;
+			public int _value;
 		}
 
 		/// <summary>
 		/// Delegate for filtering values during a tree traversal
 		/// </summary>
-		/// <param name="Value">The current value</param>
-		/// <param name="Mask">Mask for which bits in the value are valid</param>
+		/// <param name="value">The current value</param>
+		/// <param name="mask">Mask for which bits in the value are valid</param>
 		/// <returns>True if values matching the given mask should be enumerated</returns>
-		public delegate bool VisitorDelegate(ulong Value, ulong Mask);
+		public delegate bool VisitorDelegate(ulong value, ulong mask);
 
 		/// <summary>
 		/// Height of the tree
@@ -53,7 +50,7 @@ namespace Horde.Build.Logs
 		/// <summary>
 		/// Array of child offsets for each node. Excludes the last layer of the tree.
 		/// </summary>
-		int[] FirstChildIndex;
+		readonly int[] _firstChildIndex;
 
 		/// <summary>
 		/// Empty index definition
@@ -63,39 +60,39 @@ namespace Horde.Build.Logs
 		/// <summary>
 		/// Constructor
 		/// </summary>
-		/// <param name="NodeData">Node data</param>
-		public ReadOnlyTrie(ushort[] NodeData)
+		/// <param name="nodeData">Node data</param>
+		public ReadOnlyTrie(ushort[] nodeData)
 		{
-			this.NodeData = NodeData;
-			this.FirstChildIndex = CreateChildLookup(NodeData);
+			NodeData = nodeData;
+			_firstChildIndex = CreateChildLookup(nodeData);
 		}
 
 		/// <summary>
 		/// Tests whether the given value is in the trie
 		/// </summary>
-		/// <param name="Value">The value to check for</param>
+		/// <param name="value">The value to check for</param>
 		/// <returns>True if the value is in the trie</returns>
-		public bool Contains(ulong Value)
+		public bool Contains(ulong value)
 		{
-			int Index = 0;
-			for (int Shift = (sizeof(ulong) * 8) - 4; Shift >= 0; Shift -= 4)
+			int index = 0;
+			for (int shift = (sizeof(ulong) * 8) - 4; shift >= 0; shift -= 4)
 			{
-				int Mask = NodeData[Index];
-				int Flag = 1 << (int)((Value >> Shift) & 15);
-				if ((Mask & Flag) == 0)
+				int mask = NodeData[index];
+				int flag = 1 << (int)((value >> shift) & 15);
+				if ((mask & flag) == 0)
 				{
 					return false;
 				}
 
-				Index = FirstChildIndex[Index];
+				index = _firstChildIndex[index];
 				for (; ; )
 				{
-					Mask &= (Mask - 1);
-					if ((Mask & Flag) == 0)
+					mask &= (mask - 1);
+					if ((mask & flag) == 0)
 					{
 						break;
 					}
-					Index++;
+					index++;
 				}
 			}
 			return true;
@@ -104,7 +101,7 @@ namespace Horde.Build.Logs
 		/// <inheritdoc/>
 		public IEnumerator<ulong> GetEnumerator()
 		{
-			return EnumerateRange(ulong.MinValue, ulong.MaxValue).GetEnumerator();
+			return EnumerateRange(UInt64.MinValue, UInt64.MaxValue).GetEnumerator();
 		}
 
 		/// <inheritdoc/>
@@ -113,68 +110,68 @@ namespace Horde.Build.Logs
 		/// <summary>
 		/// Enumerate all values matching a given filter
 		/// </summary>
-		/// <param name="Predicate">Predicate for which values to include</param>
+		/// <param name="predicate">Predicate for which values to include</param>
 		/// <returns>Values satisfying the given predicate</returns>
-		public IEnumerable<ulong> EnumerateValues(VisitorDelegate Predicate)
+		public IEnumerable<ulong> EnumerateValues(VisitorDelegate predicate)
 		{
-			int Depth = 0;
-			ulong Value = 0;
+			int depth = 0;
+			ulong value = 0;
 
-			StackItem[] Stack = new StackItem[Height];
-			Stack[1].Index = FirstChildIndex[0];
+			StackItem[] stack = new StackItem[Height];
+			stack[1]._index = _firstChildIndex[0];
 
 			for (; ; )
 			{
-				StackItem Current = Stack[Depth];
-				if (Current.Value >= 16)
+				StackItem current = stack[depth];
+				if (current._value >= 16)
 				{
 					// Move up the tree if we've enumerated all the branches at the current level
-					Depth--;
-					if (Depth < 0)
+					depth--;
+					if (depth < 0)
 					{
 						yield break;
 					}
-					Stack[Depth].Value++;
+					stack[depth]._value++;
 
 					// Increment the child index too. These are stored sequentially for a given parent. The value will be cleared when we recurse into it.
-					Stack[Depth + 1].Index++;
+					stack[depth + 1]._index++;
 				}
-				else if ((NodeData[Current.Index] & (1 << Current.Value)) == 0)
+				else if ((NodeData[current._index] & (1 << current._value)) == 0)
 				{
 					// This branch does not exist. Skip it.
-					Stack[Depth].Value++;
+					stack[depth]._value++;
 				}
 				else
 				{
 					// Get the value and mask for the current node
-					int Shift = (Stack.Length - Depth - 1) * 4;
-					ulong Mask = ~((1UL << Shift) - 1);
-					Value = (Value & (Mask << 4)) | ((ulong)(uint)Current.Value << Shift);
+					int shift = (stack.Length - depth - 1) * 4;
+					ulong mask = ~((1UL << shift) - 1);
+					value = (value & (mask << 4)) | ((ulong)(uint)current._value << shift);
 
-					if (!Predicate(Value, Mask))
+					if (!predicate(value, mask))
 					{
 						// This node is excluded, skip it
-						Stack[Depth].Value++;
-						if (Depth + 1 < Stack.Length)
+						stack[depth]._value++;
+						if (depth + 1 < stack.Length)
 						{
-							Stack[Depth + 1].Index++;
+							stack[depth + 1]._index++;
 						}
 					}
-					else if(Depth + 1 < Stack.Length)
+					else if(depth + 1 < stack.Length)
 					{
 						// Move down the tree
-						Depth++;
-						Stack[Depth].Value = 0;
-						if (Depth + 1 < Stack.Length)
+						depth++;
+						stack[depth]._value = 0;
+						if (depth + 1 < stack.Length)
 						{
-							Stack[Depth + 1].Index = FirstChildIndex[Stack[Depth].Index];
+							stack[depth + 1]._index = _firstChildIndex[stack[depth]._index];
 						}
 					}
 					else
 					{
 						// Yield the current value
-						yield return Value;
-						Stack[Depth].Value++;
+						yield return value;
+						stack[depth]._value++;
 					}
 				}
 			}
@@ -183,82 +180,82 @@ namespace Horde.Build.Logs
 		/// <summary>
 		/// Enumerates all values in the trie between the given ranges
 		/// </summary>
-		/// <param name="MinValue">Minimum value to enumerate</param>
-		/// <param name="MaxValue">Maximum value to enumerate</param>
+		/// <param name="minValue">Minimum value to enumerate</param>
+		/// <param name="maxValue">Maximum value to enumerate</param>
 		/// <returns>Sequence of values</returns>
-		public IEnumerable<ulong> EnumerateRange(ulong MinValue, ulong MaxValue)
+		public IEnumerable<ulong> EnumerateRange(ulong minValue, ulong maxValue)
 		{
-			return EnumerateValues((Value, Mask) => (Value >= (MinValue & Mask) && Value <= (MaxValue & Mask)));
+			return EnumerateValues((value, mask) => (value >= (minValue & mask) && value <= (maxValue & mask)));
 		}
 
 		/// <summary>
 		/// Creates a lookup for child node offsets from raw node data
 		/// </summary>
-		/// <param name="NodeData">Array of masks for each node</param>
+		/// <param name="nodeData">Array of masks for each node</param>
 		/// <returns>Array of offsets</returns>
-		static int[] CreateChildLookup(ushort[] NodeData)
+		static int[] CreateChildLookup(ushort[] nodeData)
 		{
-			List<int> ChildOffsets = new List<int>();
-			if (NodeData.Length > 0)
+			List<int> childOffsets = new List<int>();
+			if (nodeData.Length > 0)
 			{
-				int NodeCount = 1;
+				int nodeCount = 1;
 
-				int Index = 0;
-				int ChildIndex = NodeCount;
+				int index = 0;
+				int childIndex = nodeCount;
 
-				for (int Level = 0; Level < Height; Level++)
+				for (int level = 0; level < Height; level++)
 				{
-					int NextNodeCount = 0;
-					for (int Idx = 0; Idx < NodeCount; Idx++)
+					int nextNodeCount = 0;
+					for (int idx = 0; idx < nodeCount; idx++)
 					{
-						ushort Node = NodeData[Index++];
+						ushort node = nodeData[index++];
 
-						int NumChildren = CountBits(Node);
-						ChildOffsets.Add(ChildIndex);
-						ChildIndex += NumChildren;
+						int numChildren = CountBits(node);
+						childOffsets.Add(childIndex);
+						childIndex += numChildren;
 
-						NextNodeCount += NumChildren;
+						nextNodeCount += numChildren;
 					}
-					NodeCount = NextNodeCount;
+					nodeCount = nextNodeCount;
 				}
 			}
-			return ChildOffsets.ToArray();
+			return childOffsets.ToArray();
 		}
 
 		/// <summary>
 		/// Count the number of set bits in the given value
 		/// </summary>
-		/// <param name="Value">Value to test</param>
+		/// <param name="value">Value to test</param>
 		/// <returns>Number of set bits</returns>
-		static int CountBits(ushort Value)
+		static int CountBits(ushort value)
 		{
-			int Count = Value;
-			Count = (Count & 0b0101010101010101) + ((Count >> 1) & 0b0101010101010101);
-			Count = (Count & 0b0011001100110011) + ((Count >> 2) & 0b0011001100110011);
-			Count = (Count & 0b0000111100001111) + ((Count >> 4) & 0b0000111100001111);
-			Count = (Count & 0b0000000011111111) + ((Count >> 8) & 0b0000000011111111);
-			return Count;
+			int count = value;
+			count = (count & 0b0101010101010101) + ((count >> 1) & 0b0101010101010101);
+			count = (count & 0b0011001100110011) + ((count >> 2) & 0b0011001100110011);
+			count = (count & 0b0000111100001111) + ((count >> 4) & 0b0000111100001111);
+			count = (count & 0b0000000011111111) + ((count >> 8) & 0b0000000011111111);
+			return count;
 		}
 
 		/// <summary>
 		/// Read a trie from the given buffer
 		/// </summary>
-		/// <param name="Reader">Reader to read from</param>
+		/// <param name="reader">Reader to read from</param>
 		/// <returns>New trie</returns>
-		public static ReadOnlyTrie Read(MemoryReader Reader)
+		public static ReadOnlyTrie Read(MemoryReader reader)
 		{
-			ReadOnlyMemory<byte> Nodes = Reader.ReadVariableLengthBytes();
-			ushort[] NodeData = MemoryMarshal.Cast<byte, ushort>(Nodes.Span).ToArray();
-			return new ReadOnlyTrie(NodeData);
+			ReadOnlyMemory<byte> nodes = reader.ReadVariableLengthBytes();
+			ushort[] nodeData = MemoryMarshal.Cast<byte, ushort>(nodes.Span).ToArray();
+			return new ReadOnlyTrie(nodeData);
 		}
 
 		/// <summary>
 		/// Write this trie to the given buffer
 		/// </summary>
-		/// <param name="Writer">Writer to output to</param>
-		public void Write(MemoryWriter Writer)
+		/// <param name="writer">Writer to output to</param>
+		public void Write(MemoryWriter writer)
 		{
-			Writer.WriteVariableLengthBytes(MemoryMarshal.AsBytes<ushort>(NodeData));
+			writer.WriteVariableLengthBytes(MemoryMarshal.AsBytes<ushort>(NodeData));
 		}
 
 		/// <summary>
@@ -279,21 +276,21 @@ namespace Horde.Build.Logs
 		/// <summary>
 		/// Read a trie from the given buffer
 		/// </summary>
-		/// <param name="Reader">Reader to read from</param>
+		/// <param name="reader">Reader to read from</param>
 		/// <returns>New trie</returns>
-		public static ReadOnlyTrie ReadTrie(this MemoryReader Reader)
+		public static ReadOnlyTrie ReadTrie(this MemoryReader reader)
 		{
-			return ReadOnlyTrie.Read(Reader);
+			return ReadOnlyTrie.Read(reader);
 		}
 
 		/// <summary>
 		/// Write this trie to the given buffer
 		/// </summary>
-		/// <param name="Writer">Writer to output to</param>
-		/// <param name="Trie">Trie to write</param>
-		public static void WriteTrie(this MemoryWriter Writer, ReadOnlyTrie Trie)
+		/// <param name="writer">Writer to output to</param>
+		/// <param name="trie">Trie to write</param>
+		public static void WriteTrie(this MemoryWriter writer, ReadOnlyTrie trie)
 		{
-			Trie.Write(Writer);
+			trie.Write(writer);
 		}
 	}
 }
diff --git a/Engine/Source/Programs/Horde/Horde.Build/Utilities/ReadOnlyTrieBuilder.cs b/Engine/Source/Programs/Horde/Horde.Build/Utilities/ReadOnlyTrieBuilder.cs
--- a/Engine/Source/Programs/Horde/Horde.Build/Utilities/ReadOnlyTrieBuilder.cs
+++ b/Engine/Source/Programs/Horde/Horde.Build/Utilities/ReadOnlyTrieBuilder.cs
@@ -1,16 +1,6 @@
 // Copyright Epic Games, Inc. All Rights Reserved.
 
-using Amazon.Runtime;
-using EpicGames.Core;
-using Horde.Build.Logs;
-using Microsoft.AspNetCore.Mvc;
-using Microsoft.AspNetCore.Routing.Constraints;
-using System;
 using System.Collections.Generic;
-using System.Linq;
-using System.Runtime.CompilerServices;
-using System.Runtime.InteropServices;
-using System.Threading.Tasks;
 
 namespace Horde.Build.Logs
 {
@@ -24,65 +14,65 @@ namespace Horde.Build.Logs
 		/// </summary>
 		class Node
 		{
-			public Node?[]? Children;
+			public Node?[]? _children;
 		}
 
 		/// <summary>
 		/// The root node
 		/// </summary>
-		Node Root;
+		readonly Node _root;
 
 		/// <summary>
 		/// Default constructor
 		/// </summary>
 		public ReadOnlyTrieBuilder()
 		{
-			Root = new Node();
+			_root = new Node();
 		}
 
 		/// <summary>
 		/// Adds a value to the trie
 		/// </summary>
-		/// <param name="Value">Value to add</param>
-		public void Add(ulong Value)
+		/// <param name="value">Value to add</param>
+		public void Add(ulong value)
 		{
 			// Loop through the tree until we've added the item
-			Node Leaf = Root;
-			for (int Shift = (sizeof(ulong) * 8) - 4; Shift >= 0; Shift -= 4)
+			Node leaf = _root;
+			for (int shift = (sizeof(ulong) * 8) - 4; shift >= 0; shift -= 4)
 			{
-				int Index = (int)(Value >> Shift) & 15;
-				if (Leaf.Children == null)
+				int index = (int)(value >> shift) & 15;
+				if (leaf._children == null)
 				{
-					Leaf.Children = new Node[16];
+					leaf._children = new Node[16];
 				}
-				if (Leaf.Children[Index] == null)
+				if (leaf._children[index] == null)
 				{
-					Leaf.Children[Index] = new Node();
+					leaf._children[index] = new Node();
 				}
-				Leaf = Leaf.Children[Index]!;
+				leaf = leaf._children[index]!;
 			}
 		}
 
 		/// <summary>
 		/// Searches for the given item in the trie
 		/// </summary>
-		/// <param name="Value">Value to add</param>
-		public bool Contains(ulong Value)
+		/// <param name="value">Value to add</param>
+		public bool Contains(ulong value)
 		{
 			// Loop through the tree until we've added the item
-			Node Leaf = Root;
-			for (int Shift = (sizeof(ulong) * 8) - 4; Shift >= 0; Shift -= 4)
+			Node leaf = _root;
+			for (int shift = (sizeof(ulong) * 8) - 4; shift >= 0; shift -= 4)
 			{
-				int Index = (int)(Value >> Shift) & 15;
-				if (Leaf.Children == null)
+				int index = (int)(value >> shift) & 15;
+				if (leaf._children == null)
 				{
 					return false;
 				}
-				if (Leaf.Children[Index] == null)
+				if (leaf._children[index] == null)
 				{
 					return false;
 				}
-				Leaf = Leaf.Children[Index]!;
+				leaf = leaf._children[index]!;
 			}
 			return true;
 		}
@@ -93,34 +83,34 @@ namespace Horde.Build.Logs
 		/// <returns></returns>
 		public ReadOnlyTrie Build()
 		{
-			List<ushort> Values = new List<ushort>();
+			List<ushort> values = new List<ushort>();
 
-			List<Node> Nodes = new List<Node>();
-			Nodes.Add(Root);
+			List<Node> nodes = new List<Node>();
+			nodes.Add(_root);
 
-			for(int Bits = 0; Bits < (sizeof(ulong) * 8); Bits += 4)
+			for(int bits = 0; bits < (sizeof(ulong) * 8); bits += 4)
 			{
-				List<Node> NextNodes = new List<Node>();
-				foreach(Node Node in Nodes)
+				List<Node> nextNodes = new List<Node>();
+				foreach(Node node in nodes)
 				{
-					ushort Value = 0;
-					if (Node.Children != null)
+					ushort value = 0;
+					if (node._children != null)
 					{
-						for (int Idx = 0; Idx < Node.Children.Length; Idx++)
+						for (int idx = 0; idx < node._children.Length; idx++)
 						{
-							if(Node.Children[Idx] != null)
+							if(node._children[idx] != null)
 							{
-								Value |= (ushort)(1 << Idx);
-								NextNodes.Add(Node.Children[Idx]!);
+								value |= (ushort)(1 << idx);
+								nextNodes.Add(node._children[idx]!);
 							}
 						}
 					}
-					Values.Add(Value);
+					values.Add(value);
 				}
-				Nodes = NextNodes;
+				nodes = nextNodes;
 			}
 
-			return new ReadOnlyTrie(Values.ToArray());
+			return new ReadOnlyTrie(values.ToArray());
 		}
 	}
 }
diff --git a/Engine/Source/Programs/Horde/Horde.Build/Utilities/RedisConnectionPool.cs b/Engine/Source/Programs/Horde/Horde.Build/Utilities/RedisConnectionPool.cs
--- a/Engine/Source/Programs/Horde/Horde.Build/Utilities/RedisConnectionPool.cs
+++ b/Engine/Source/Programs/Horde/Horde.Build/Utilities/RedisConnectionPool.cs
@@ -16,22 +16,22 @@ namespace Horde.Build.Utilities
 	/// </summary>
 	public class RedisConnectionPool
 	{
-		private readonly ConcurrentBag<Lazy<ConnectionMultiplexer>> Connections;
-		private readonly int DefaultDatabaseIndex;
+		private readonly ConcurrentBag<Lazy<ConnectionMultiplexer>> _connections;
+		private readonly int _defaultDatabaseIndex;
 
 		/// <summary>
 		/// Constructor
 		/// </summary>
-		/// <param name="PoolSize">Size of the pool (i.e max number of connections)</param>
-		/// <param name="RedisConfString">Configuration string for Redis</param>
-		/// <param name="DefaultDatabaseIndex">The Redis database index to use. Use -1 for the default one</param>
-		public RedisConnectionPool(int PoolSize, string RedisConfString, int DefaultDatabaseIndex = -1)
+		/// <param name="poolSize">Size of the pool (i.e max number of connections)</param>
+		/// <param name="redisConfString">Configuration string for Redis</param>
+		/// <param name="defaultDatabaseIndex">The Redis database index to use. Use -1 for the default one</param>
+		public RedisConnectionPool(int poolSize, string redisConfString, int defaultDatabaseIndex = -1)
 		{
-			this.DefaultDatabaseIndex = DefaultDatabaseIndex;
-			Connections = new ConcurrentBag<Lazy<ConnectionMultiplexer>>();
-			for (int i = 0; i < PoolSize; i++)
+			_defaultDatabaseIndex = defaultDatabaseIndex;
+			_connections = new ConcurrentBag<Lazy<ConnectionMultiplexer>>();
+			for (int i = 0; i < poolSize; i++)
 			{
-				Connections.Add(new Lazy<ConnectionMultiplexer>(() => ConnectionMultiplexer.Connect(RedisConfString)));
+				_connections.Add(new Lazy<ConnectionMultiplexer>(() => ConnectionMultiplexer.Connect(redisConfString)));
 			}
 		}
 
@@ -43,21 +43,21 @@ namespace Horde.Build.Utilities
 		/// <returns>A Redis database connection</returns>
 		public IConnectionMultiplexer GetConnection()
 		{
-			Lazy<ConnectionMultiplexer> LazyConnection;
-			var LazyConnections = Connections.Where(x => x.IsValueCreated);
+			Lazy<ConnectionMultiplexer> lazyConnection;
+			IEnumerable<Lazy<ConnectionMultiplexer>>? lazyConnections = _connections.Where(x => x.IsValueCreated);
 
-			if (LazyConnections.Count() == Connections.Count)
+			if (lazyConnections.Count() == _connections.Count)
 			{
 				// No more new connections can be created, pick the least loaded one
-				LazyConnection = Connections.OrderBy(x => x.Value.GetCounters().TotalOutstanding).First();
+				lazyConnection = _connections.OrderBy(x => x.Value.GetCounters().TotalOutstanding).First();
 			}
 			else
 			{
 				// Create a new connection by picking a not yet initialized lazy value 
-				LazyConnection = Connections.First(x => !x.IsValueCreated);
+				lazyConnection = _connections.First(x => !x.IsValueCreated);
 			}
 
-			return LazyConnection.Value;
+			return lazyConnection.Value;
 		}
 
 		/// <summary>
@@ -66,8 +66,7 @@ namespace Horde.Build.Utilities
 		/// <returns>A Redis database</returns>
 		public IDatabase GetDatabase()
 		{
-			return GetConnection().GetDatabase(DefaultDatabaseIndex);
+			return GetConnection().GetDatabase(_defaultDatabaseIndex);
 		}
-
 	}
 }
\ No newline at end of file
diff --git a/Engine/Source/Programs/Horde/Horde.Build/Utilities/Schemas.cs b/Engine/Source/Programs/Horde/Horde.Build/Utilities/Schemas.cs
--- a/Engine/Source/Programs/Horde/Horde.Build/Utilities/Schemas.cs
+++ b/Engine/Source/Programs/Horde/Horde.Build/Utilities/Schemas.cs
@@ -1,16 +1,10 @@
 // Copyright Epic Games, Inc. All Rights Reserved.
 
-using EpicGames.Core;
 using System;
 using System.Collections.Concurrent;
-using System.Collections.Generic;
-using System.Diagnostics.CodeAnalysis;
-using System.IO;
-using System.Linq;
 using System.Reflection;
-using System.Text.Json;
-using System.Text.Json.Serialization;
 using System.Xml;
+using EpicGames.Core;
 
 namespace Horde.Build.Utilities
 {
@@ -19,35 +13,35 @@ namespace Horde.Build.Utilities
 	/// </summary>
 	static class Schemas
 	{
-		static ConcurrentDictionary<Assembly, XmlDocument?> CachedDocumentation = new ConcurrentDictionary<Assembly, XmlDocument?>();
-		static ConcurrentDictionary<Type, JsonSchema> CachedSchemas = new ConcurrentDictionary<Type, JsonSchema>();
+		static readonly ConcurrentDictionary<Assembly, XmlDocument?> s_cachedDocumentation = new ConcurrentDictionary<Assembly, XmlDocument?>();
+		static readonly ConcurrentDictionary<Type, JsonSchema> s_cachedSchemas = new ConcurrentDictionary<Type, JsonSchema>();
 
 		/// <summary>
 		/// Create a Json schema (or retrieve a cached schema)
 		/// </summary>
-		/// <param name="Type"></param>
+		/// <param name="type"></param>
 		/// <returns></returns>
-		public static JsonSchema CreateSchema(Type Type)
+		public static JsonSchema CreateSchema(Type type)
 		{
-			JsonSchema? Schema;
-			if (!CachedSchemas.TryGetValue(Type, out Schema))
+			JsonSchema? schema;
+			if (!s_cachedSchemas.TryGetValue(type, out schema))
 			{
-				XmlDocument? Documentation;
-				if (!CachedDocumentation.TryGetValue(Type.Assembly, out Documentation))
+				XmlDocument? documentation;
+				if (!s_cachedDocumentation.TryGetValue(type.Assembly, out documentation))
 				{
-					FileReference InputDocumentationFile = new FileReference(Type.Assembly.Location).ChangeExtension(".xml");
-					if (FileReference.Exists(InputDocumentationFile))
+					FileReference inputDocumentationFile = new FileReference(type.Assembly.Location).ChangeExtension(".xml");
+					if (FileReference.Exists(inputDocumentationFile))
 					{
-						Documentation = new XmlDocument();
-						Documentation.Load(InputDocumentationFile.FullName);
+						documentation = new XmlDocument();
+						documentation.Load(inputDocumentationFile.FullName);
 					}
-					CachedDocumentation.TryAdd(Type.Assembly, Documentation);
+					s_cachedDocumentation.TryAdd(type.Assembly, documentation);
 				}
 
-				Schema = JsonSchema.FromType(Type, Documentation);
-				CachedSchemas.TryAdd(Type, Schema);
+				schema = JsonSchema.FromType(type, documentation);
+				s_cachedSchemas.TryAdd(type, schema);
 			}
-			return Schema;
+			return schema;
 		}
 	}
 }
diff --git a/Engine/Source/Programs/Horde/Horde.Build/Utilities/SearchText.cs b/Engine/Source/Programs/Horde/Horde.Build/Utilities/SearchText.cs
--- a/Engine/Source/Programs/Horde/Horde.Build/Utilities/SearchText.cs
+++ b/Engine/Source/Programs/Horde/Horde.Build/Utilities/SearchText.cs
@@ -2,9 +2,7 @@
 
 using System;
 using System.Collections.Generic;
-using System.Linq;
 using System.Text;
-using System.Threading.Tasks;
 
 namespace Horde.Build.Logs
 {
@@ -26,53 +24,53 @@ namespace Horde.Build.Logs
 		/// <summary>
 		/// Normalized (lowercase) utf-8 bytes to search for
 		/// </summary>
-		byte[] SearchBytes;
+		readonly byte[] _searchBytes;
 
 		/// <summary>
 		/// Skip table for comparisons
 		/// </summary>
-		byte[] SkipTable;
+		readonly byte[] _skipTable;
 
 		/// <summary>
 		/// Constructor
 		/// </summary>
-		/// <param name="Text">The text to search for</param>
-		public SearchText(string Text)
+		/// <param name="text">The text to search for</param>
+		public SearchText(string text)
 		{
-			byte[] Bytes = Encoding.UTF8.GetBytes(Text);
-			this.Bytes = Bytes;
-			this.Text = Text;
+			byte[] bytes = Encoding.UTF8.GetBytes(text);
+			Bytes = bytes;
+			Text = text;
 
 			// Find the byte sequence to search for, in lowercase
-			SearchBytes = new byte[Bytes.Length];
-			for (int Idx = 0; Idx < SearchBytes.Length; Idx++)
+			_searchBytes = new byte[bytes.Length];
+			for (int idx = 0; idx < _searchBytes.Length; idx++)
 			{
-				if (Bytes[Idx] >= 'A' && Bytes[Idx] <= 'Z')
+				if (bytes[idx] >= 'A' && bytes[idx] <= 'Z')
 				{
-					SearchBytes[Idx] = (byte)('a' + (Bytes[Idx] - 'A'));
+					_searchBytes[idx] = (byte)('a' + (bytes[idx] - 'A'));
 				}
 				else
 				{
-					SearchBytes[Idx] = Bytes[Idx];
+					_searchBytes[idx] = bytes[idx];
 				}
 			}
 
 			// Build a table indicating how many characters to skip before attempting the next comparison
-			SkipTable = new byte[256];
-			for (int Idx = 0; Idx < 256; Idx++)
+			_skipTable = new byte[256];
+			for (int idx = 0; idx < 256; idx++)
 			{
-				SkipTable[Idx] = (byte)SearchBytes.Length;
+				_skipTable[idx] = (byte)_searchBytes.Length;
 			}
-			for (int Idx = 0; Idx < SearchBytes.Length - 1; Idx++)
+			for (int idx = 0; idx < _searchBytes.Length - 1; idx++)
 			{
-				byte Character = SearchBytes[Idx];
+				byte character = _searchBytes[idx];
 
-				byte SkipBytes = (byte)(SearchBytes.Length - 1 - Idx);
-				SkipTable[Character] = SkipBytes;
+				byte skipBytes = (byte)(_searchBytes.Length - 1 - idx);
+				_skipTable[character] = skipBytes;
 
-				if (Character >= 'a' && Character <= 'z')
+				if (character >= 'a' && character <= 'z')
 				{
-					SkipTable['A' + (Character - 'a')] = SkipBytes;
+					_skipTable['A' + (character - 'a')] = skipBytes;
 				}
 			}
 		}
@@ -80,40 +78,40 @@ namespace Horde.Build.Logs
 		/// <summary>
 		/// Find all ocurrences of the text in the given buffer
 		/// </summary>
-		/// <param name="Buffer">The buffer to search</param>
-		/// <param name="Text">The text to search for</param>
+		/// <param name="buffer">The buffer to search</param>
+		/// <param name="text">The text to search for</param>
 		/// <returns>Sequence of offsets within the buffer</returns>
-		public static IEnumerable<int> FindOcurrences(ReadOnlyMemory<byte> Buffer, SearchText Text)
+		public static IEnumerable<int> FindOcurrences(ReadOnlyMemory<byte> buffer, SearchText text)
 		{
-			for(int Offset = 0; ;Offset++)
+			for(int offset = 0; ;offset++)
 			{
-				Offset = FindNextOcurrence(Buffer.Span, Offset, Text);
-				if(Offset == -1)
+				offset = FindNextOcurrence(buffer.Span, offset, text);
+				if(offset == -1)
 				{
 					break;
 				}
-				yield return Offset;
+				yield return offset;
 			}
 		}
 
 		/// <summary>
 		/// Perform a case sensitive search for the next occurerence of the search term in a given buffer
 		/// </summary>
-		/// <param name="Buffer">The buffer to search</param>
-		/// <param name="Offset">Starting offset for the search</param>
-		/// <param name="Text">The text to search for</param>
+		/// <param name="buffer">The buffer to search</param>
+		/// <param name="offset">Starting offset for the search</param>
+		/// <param name="text">The text to search for</param>
 		/// <returns>Offset of the next occurence, or -1</returns>
-		public static int FindNextOcurrence(ReadOnlySpan<byte> Buffer, int Offset, SearchText Text)
+		public static int FindNextOcurrence(ReadOnlySpan<byte> buffer, int offset, SearchText text)
 		{
-			while (Offset + Text.SearchBytes.Length <= Buffer.Length)
+			while (offset + text._searchBytes.Length <= buffer.Length)
 			{
-				if (Matches(Buffer, Offset, Text))
+				if (Matches(buffer, offset, text))
 				{
-					return Offset;
+					return offset;
 				}
 				else
 				{
-					Offset += Text.SkipTable[Buffer[Offset + Text.SearchBytes.Length - 1]];
+					offset += text._skipTable[buffer[offset + text._searchBytes.Length - 1]];
 				}
 			}
 			return -1;
@@ -122,20 +120,20 @@ namespace Horde.Build.Logs
 		/// <summary>
 		/// Compare the search term against the given buffer
 		/// </summary>
-		/// <param name="Buffer">The buffer to search</param>
-		/// <param name="Offset">Starting offset for the search</param>
-		/// <param name="Text">The text to search for</param>
+		/// <param name="buffer">The buffer to search</param>
+		/// <param name="offset">Starting offset for the search</param>
+		/// <param name="text">The text to search for</param>
 		/// <returns>True if the text matches, false otherwise</returns>
-		public static bool Matches(ReadOnlySpan<byte> Buffer, int Offset, SearchText Text)
+		public static bool Matches(ReadOnlySpan<byte> buffer, int offset, SearchText text)
 		{
-			for (int Idx = Text.SearchBytes.Length - 1; Idx >= 0; Idx--)
+			for (int idx = text._searchBytes.Length - 1; idx >= 0; idx--)
 			{
-				byte Char = Buffer[Offset + Idx];
-				if (Char >= 'A' && Char <= 'Z')
+				byte character = buffer[offset + idx];
+				if (character >= 'A' && character <= 'Z')
 				{
-					Char = (byte)('a' + (Char - 'A'));
+					character = (byte)('a' + (character - 'A'));
 				}
-				if (Char != Text.SearchBytes[Idx])
+				if (character != text._searchBytes[idx])
 				{
 					return false;
 				}
@@ -152,36 +150,36 @@ namespace Horde.Build.Logs
 		/// <summary>
 		/// Find all ocurrences of the text in the given buffer
 		/// </summary>
-		/// <param name="Buffer">The buffer to search</param>
-		/// <param name="Text">The text to search for</param>
+		/// <param name="buffer">The buffer to search</param>
+		/// <param name="text">The text to search for</param>
 		/// <returns>Sequence of offsets within the buffer</returns>
-		public static IEnumerable<int> FindOcurrences(this ReadOnlyMemory<byte> Buffer, SearchText Text)
+		public static IEnumerable<int> FindOcurrences(this ReadOnlyMemory<byte> buffer, SearchText text)
 		{
-			return SearchText.FindOcurrences(Buffer, Text);
+			return SearchText.FindOcurrences(buffer, text);
 		}
 
 		/// <summary>
 		/// Perform a case sensitive search for the next occurerence of the search term in a given buffer
 		/// </summary>
-		/// <param name="Buffer">The buffer to search</param>
-		/// <param name="Offset">Starting offset for the search</param>
-		/// <param name="Text">The text to search for</param>
+		/// <param name="buffer">The buffer to search</param>
+		/// <param name="offset">Starting offset for the search</param>
+		/// <param name="text">The text to search for</param>
 		/// <returns>Offset of the next occurence, or -1</returns>
-		public static int FindNextOcurrence(this ReadOnlySpan<byte> Buffer, int Offset, SearchText Text)
+		public static int FindNextOcurrence(this ReadOnlySpan<byte> buffer, int offset, SearchText text)
 		{
-			return SearchText.FindNextOcurrence(Buffer, Offset, Text);
+			return SearchText.FindNextOcurrence(buffer, offset, text);
 		}
 
 		/// <summary>
 		/// Compare the search term against the given buffer
 		/// </summary>
-		/// <param name="Buffer">The buffer to search</param>
-		/// <param name="Offset">Starting offset for the search</param>
-		/// <param name="Text">The text to search for</param>
+		/// <param name="buffer">The buffer to search</param>
+		/// <param name="offset">Starting offset for the search</param>
+		/// <param name="text">The text to search for</param>
 		/// <returns>True if the text matches, false otherwise</returns>
-		public static bool Matches(this ReadOnlySpan<byte> Buffer, int Offset, SearchText Text)
+		public static bool Matches(this ReadOnlySpan<byte> buffer, int offset, SearchText text)
 		{
-			return SearchText.Matches(Buffer, Offset, Text);
+			return SearchText.Matches(buffer, offset, text);
 		}
 	}
 }
diff --git a/Engine/Source/Programs/Horde/Horde.Build/Utilities/SemaphoreSlimExtensions.cs b/Engine/Source/Programs/Horde/Horde.Build/Utilities/SemaphoreSlimExtensions.cs
--- a/Engine/Source/Programs/Horde/Horde.Build/Utilities/SemaphoreSlimExtensions.cs
+++ b/Engine/Source/Programs/Horde/Horde.Build/Utilities/SemaphoreSlimExtensions.cs
@@ -1,8 +1,6 @@
 // Copyright Epic Games, Inc. All Rights Reserved.
 
 using System;
-using System.Collections.Generic;
-using System.Linq;
 using System.Threading;
 using System.Threading.Tasks;
 
@@ -16,13 +14,13 @@ namespace Horde.Build.Utilities
 		/// <summary>
 		/// Returns a disposable semaphore
 		/// </summary>
-		/// <param name="Semaphore">the semaphore to wrap</param>
-		/// <param name="CancelToken">optional cancellation token</param>
+		/// <param name="semaphore">the semaphore to wrap</param>
+		/// <param name="cancelToken">optional cancellation token</param>
 		/// <returns></returns>
-		public static async Task<IDisposable> UseWaitAsync(this SemaphoreSlim Semaphore, CancellationToken CancelToken = default(CancellationToken))
+		public static async Task<IDisposable> UseWaitAsync(this SemaphoreSlim semaphore, CancellationToken cancelToken = default)
 		{
-			await Semaphore.WaitAsync(CancelToken).ConfigureAwait(false);
-			return new ReleaseWrapper(Semaphore);
+			await semaphore.WaitAsync(cancelToken).ConfigureAwait(false);
+			return new ReleaseWrapper(semaphore);
 		}
 
 		/// <summary>
@@ -33,20 +31,20 @@ namespace Horde.Build.Utilities
 			/// <summary>
 			/// The semaphore to wrap
 			/// </summary>
-			private readonly SemaphoreSlim Semaphore;
+			private readonly SemaphoreSlim _semaphore;
 
 			/// <summary>
 			/// Whether this has been disposed
 			/// </summary>
-			private bool bIsDisposed;
+			private bool _bIsDisposed;
 
 			/// <summary>
 			/// Constructor
 			/// </summary>
-			/// <param name="Semaphore">The semaphore to wrap</param>
-			public ReleaseWrapper(SemaphoreSlim Semaphore)
+			/// <param name="semaphore">The semaphore to wrap</param>
+			public ReleaseWrapper(SemaphoreSlim semaphore)
 			{
-				this.Semaphore = Semaphore;
+				_semaphore = semaphore;
 			}
 
 			/// <summary>
@@ -54,13 +52,13 @@ namespace Horde.Build.Utilities
 			/// </summary>
 			public void Dispose()
 			{
-				if (bIsDisposed)
+				if (_bIsDisposed)
 				{
 					return;
 				}
 
-				Semaphore.Release();
-				bIsDisposed = true;
+				_semaphore.Release();
+				_bIsDisposed = true;
 			}
 		}
 	}
diff --git a/Engine/Source/Programs/Horde/Horde.Build/Utilities/SingletonDocument.cs b/Engine/Source/Programs/Horde/Horde.Build/Utilities/SingletonDocument.cs
--- a/Engine/Source/Programs/Horde/Horde.Build/Utilities/SingletonDocument.cs
+++ b/Engine/Source/Programs/Horde/Horde.Build/Utilities/SingletonDocument.cs
@@ -1,13 +1,11 @@
 // Copyright Epic Games, Inc. All Rights Reserved.
 
-using Horde.Build.Models;
-using Horde.Build.Services;
-using MongoDB.Bson;
 using System;
-using System.Collections.Generic;
-using System.Linq;
 using System.Reflection;
 using System.Threading.Tasks;
+using Horde.Build.Models;
+using Horde.Build.Services;
+using MongoDB.Bson;
 
 namespace Horde.Build.Utilities
 {
@@ -25,10 +23,10 @@ namespace Horde.Build.Utilities
 		/// <summary>
 		/// Constructor
 		/// </summary>
-		/// <param name="Id">Unique id for the singleton document</param>
-		public SingletonDocumentAttribute(string Id)
+		/// <param name="id">Unique id for the singleton document</param>
+		public SingletonDocumentAttribute(string id)
 		{
-			this.Id = Id;
+			Id = id;
 		}
 	}
 
@@ -47,9 +45,9 @@ namespace Horde.Build.Utilities
 		/// <summary>
 		/// Attempts to update the document
 		/// </summary>
-		/// <param name="Value">New state of the document</param>
+		/// <param name="value">New state of the document</param>
 		/// <returns>True if the document was updated, false otherwise</returns>
-		Task<bool> TryUpdateAsync(T Value);
+		Task<bool> TryUpdateAsync(T value);
 	}
 
 	/// <summary>
@@ -61,53 +59,53 @@ namespace Horde.Build.Utilities
 		/// <summary>
 		/// The database service instance
 		/// </summary>
-		DatabaseService DatabaseService;
+		readonly DatabaseService _databaseService;
 
 		/// <summary>
 		/// Unique id for the singleton document
 		/// </summary>
-		ObjectId ObjectId;
+		readonly ObjectId _objectId;
 
 		/// <summary>
 		/// Constructor
 		/// </summary>
-		/// <param name="DatabaseService">The database service instance</param>
-		public SingletonDocument(DatabaseService DatabaseService)
+		/// <param name="databaseService">The database service instance</param>
+		public SingletonDocument(DatabaseService databaseService)
 		{
-			this.DatabaseService = DatabaseService;
+			_databaseService = databaseService;
 
-			SingletonDocumentAttribute? Attribute = typeof(T).GetCustomAttribute<SingletonDocumentAttribute>();
-			if (Attribute == null)
+			SingletonDocumentAttribute? attribute = typeof(T).GetCustomAttribute<SingletonDocumentAttribute>();
+			if (attribute == null)
 			{
 				throw new Exception($"Type {typeof(T).Name} is missing a {nameof(SingletonDocumentAttribute)} annotation");
 			}
 
-			ObjectId = new ObjectId(Attribute.Id);
+			_objectId = new ObjectId(attribute.Id);
 		}
 
 		/// <summary>
 		/// Constructor
 		/// </summary>
-		/// <param name="DatabaseService">The database service instance</param>
-		/// <param name="ObjectId">The singleton document object id</param>
-		public SingletonDocument(DatabaseService DatabaseService, ObjectId ObjectId)
+		/// <param name="databaseService">The database service instance</param>
+		/// <param name="objectId">The singleton document object id</param>
+		public SingletonDocument(DatabaseService databaseService, ObjectId objectId)
 		{
-			this.DatabaseService = DatabaseService;
-			this.ObjectId = ObjectId;
+			_databaseService = databaseService;
+			_objectId = objectId;
 		}
 
 		/// <inheritdoc/>
 		public async Task<T> GetAsync()
 		{
-			T Value = await DatabaseService.GetSingletonAsync<T>(ObjectId);
-			Value.Id = ObjectId;
-			return Value;
+			T value = await _databaseService.GetSingletonAsync<T>(_objectId);
+			value.Id = _objectId;
+			return value;
 		}
 
 		/// <inheritdoc/>
-		public Task<bool> TryUpdateAsync(T Value)
+		public Task<bool> TryUpdateAsync(T value)
 		{
-			return DatabaseService.TryUpdateSingletonAsync<T>(Value);
+			return _databaseService.TryUpdateSingletonAsync<T>(value);
 		}
 	}
 
@@ -120,19 +118,19 @@ namespace Horde.Build.Utilities
 		/// Update a singleton
 		/// </summary>
 		/// <typeparam name="T"></typeparam>
-		/// <param name="Singleton"></param>
-		/// <param name="UpdateAction"></param>
+		/// <param name="singleton"></param>
+		/// <param name="updateAction"></param>
 		/// <returns></returns>
-		public static async Task<T> UpdateAsync<T>(this ISingletonDocument<T> Singleton, Action<T> UpdateAction) where T : SingletonBase, new()
+		public static async Task<T> UpdateAsync<T>(this ISingletonDocument<T> singleton, Action<T> updateAction) where T : SingletonBase, new()
 		{
 			for (; ; )
 			{
-				T Value = await Singleton.GetAsync();
-				UpdateAction(Value);
+				T value = await singleton.GetAsync();
+				updateAction(value);
 
-				if (await Singleton.TryUpdateAsync(Value))
+				if (await singleton.TryUpdateAsync(value))
 				{
-					return Value;
+					return value;
 				}
 			}
 		}
diff --git a/Engine/Source/Programs/Horde/Horde.Build/Utilities/StreamExtensions.cs b/Engine/Source/Programs/Horde/Horde.Build/Utilities/StreamExtensions.cs
--- a/Engine/Source/Programs/Horde/Horde.Build/Utilities/StreamExtensions.cs
+++ b/Engine/Source/Programs/Horde/Horde.Build/Utilities/StreamExtensions.cs
@@ -1,9 +1,7 @@
 // Copyright Epic Games, Inc. All Rights Reserved.
 
 using System;
-using System.Collections.Generic;
 using System.IO;
-using System.Linq;
 using System.Threading.Tasks;
 
 namespace Horde.Build.Utilities
@@ -16,32 +14,32 @@ namespace Horde.Build.Utilities
 		/// <summary>
 		/// Reads data of a fixed size into the buffer. Throws an exception if the whole block cannot be read.
 		/// </summary>
-		/// <param name="Stream">The stream to read from</param>
-		/// <param name="Data">Buffer to receive the read data</param>
-		/// <param name="Offset">Offset within the buffer to read the new data</param>
-		/// <param name="Length">Length of the data to read</param>
+		/// <param name="stream">The stream to read from</param>
+		/// <param name="data">Buffer to receive the read data</param>
+		/// <param name="offset">Offset within the buffer to read the new data</param>
+		/// <param name="length">Length of the data to read</param>
 		/// <returns>Async task</returns>
-		public static Task ReadFixedSizeDataAsync(this Stream Stream, byte[] Data, int Offset, int Length)
+		public static Task ReadFixedSizeDataAsync(this Stream stream, byte[] data, int offset, int length)
 		{
-			return ReadFixedSizeDataAsync(Stream, Data.AsMemory(Offset, Length));
+			return ReadFixedSizeDataAsync(stream, data.AsMemory(offset, length));
 		}
 
 		/// <summary>
 		/// Reads data of a fixed size into the buffer. Throws an exception if the whole block cannot be read.
 		/// </summary>
-		/// <param name="Stream">The stream to read from</param>
-		/// <param name="Data">Buffer to receive the read data</param>
+		/// <param name="stream">The stream to read from</param>
+		/// <param name="data">Buffer to receive the read data</param>
 		/// <returns>Async task</returns>
-		public static async Task ReadFixedSizeDataAsync(this Stream Stream, Memory<byte> Data)
+		public static async Task ReadFixedSizeDataAsync(this Stream stream, Memory<byte> data)
 		{
-			while (Data.Length > 0)
+			while (data.Length > 0)
 			{
-				int Count = await Stream.ReadAsync(Data);
-				if (Count == 0)
+				int count = await stream.ReadAsync(data);
+				if (count == 0)
 				{
 					throw new EndOfStreamException();
 				}
-				Data = Data.Slice(Count);
+				data = data.Slice(count);
 			}
 		}
 	}
diff --git a/Engine/Source/Programs/Horde/Horde.Build/Utilities/StringId.cs b/Engine/Source/Programs/Horde/Horde.Build/Utilities/StringId.cs
--- a/Engine/Source/Programs/Horde/Horde.Build/Utilities/StringId.cs
+++ b/Engine/Source/Programs/Horde/Horde.Build/Utilities/StringId.cs
@@ -1,22 +1,16 @@
 // Copyright Epic Games, Inc. All Rights Reserved.
 
-using EpicGames.Core;
-using EpicGames.Serialization;
-using EpicGames.Serialization.Converters;
-using Horde.Build.Models;
-using MongoDB.Bson.Serialization;
-using MongoDB.Bson.Serialization.Attributes;
-using MongoDB.Bson.Serialization.Serializers;
 using System;
-using System.Collections.Generic;
 using System.ComponentModel;
 using System.Diagnostics.CodeAnalysis;
 using System.Globalization;
-using System.Linq;
 using System.Text;
 using System.Text.Json;
 using System.Text.Json.Serialization;
-using System.Threading.Tasks;
+using EpicGames.Core;
+using EpicGames.Serialization;
+using MongoDB.Bson.Serialization;
+using MongoDB.Bson.Serialization.Serializers;
 
 namespace Horde.Build.Utilities
 {
@@ -37,40 +31,40 @@ namespace Horde.Build.Utilities
 		/// <summary>
 		/// The text representing this id
 		/// </summary>
-		readonly string Text;
+		readonly string _text;
 
 		/// <summary>
 		/// Constructor
 		/// </summary>
-		/// <param name="Input">Unique id for the string</param>
+		/// <param name="input">Unique id for the string</param>
 		[SuppressMessage("Globalization", "CA1308:Normalize strings to uppercase", Justification = "<Pending>")]
-		public StringId(string Input)
+		public StringId(string input)
 		{
-			this.Text = Input;
+			_text = input;
 
-			if (Text.Length == 0)
+			if (_text.Length == 0)
 			{
 //				throw new ArgumentException("String id may not be empty");
 			}
 
 			const int MaxLength = 64;
-			if (Text.Length > MaxLength)
+			if (_text.Length > MaxLength)
 			{
 				throw new ArgumentException($"String id may not be longer than {MaxLength} characters");
 			}
 
-			for (int Idx = 0; Idx < Text.Length; Idx++)
+			for (int idx = 0; idx < _text.Length; idx++)
 			{
-				char Character = Text[Idx];
-				if (!IsValidCharacter(Character))
+				char character = _text[idx];
+				if (!IsValidCharacter(character))
 				{
-					if (Character >= 'A' && Character <= 'Z')
+					if (character >= 'A' && character <= 'Z')
 					{
-						Text = Text.ToLowerInvariant();
+						_text = _text.ToLowerInvariant();
 					}
 					else
 					{
-						throw new ArgumentException($"{Text} is not a valid string id");
+						throw new ArgumentException($"{_text} is not a valid string id");
 					}
 				}
 			}
@@ -79,77 +73,74 @@ namespace Horde.Build.Utilities
 		/// <summary>
 		/// Constructs from a nullable string
 		/// </summary>
-		/// <param name="Text">The text to construct from</param>
+		/// <param name="text">The text to construct from</param>
 		/// <returns></returns>
 		[SuppressMessage("Design", "CA1000:Do not declare static members on generic types")]
-		public static StringId<T>? FromNullable(string? Text)
+		public static StringId<T>? FromNullable(string? text)
 		{
-			if (String.IsNullOrEmpty(Text))
+			if (String.IsNullOrEmpty(text))
 			{
 				return null;
 			}
 			else
 			{
-				return new StringId<T>(Text);
+				return new StringId<T>(text);
 			}
 		}
 
 		/// <summary>
 		/// Generates a new string id from the given text
 		/// </summary>
-		/// <param name="Text">Text to generate from</param>
+		/// <param name="text">Text to generate from</param>
 		/// <returns>New string id</returns>
 		[SuppressMessage("Design", "CA1000:Do not declare static members on generic types", Justification = "<Pending>")]
-		public static StringId<T> Sanitize(string Text)
+		public static StringId<T> Sanitize(string text)
 		{
-			StringBuilder Result = new StringBuilder();
-			for (int Idx = 0; Idx < Text.Length; Idx++)
+			StringBuilder result = new StringBuilder();
+			for (int idx = 0; idx < text.Length; idx++)
 			{
-				char Character = Text[Idx];
-				if (Character >= 'A' && Character <= 'Z')
+				char character = text[idx];
+				if (character >= 'A' && character <= 'Z')
 				{
-					Result.Append((char)('a' + (Character - 'A')));
+					result.Append((char)('a' + (character - 'A')));
 				}
-				else if (IsValidCharacter(Character))
+				else if (IsValidCharacter(character))
 				{
-					Result.Append(Character);
+					result.Append(character);
 				}
-				else if (Result.Length > 0 && Result[Result.Length - 1] != '-')
+				else if (result.Length > 0 && result[^1] != '-')
 				{
-					Result.Append('-');
+					result.Append('-');
 				}
 			}
-			while(Result.Length > 0 && Result[Result.Length - 1] == '-')
+			while(result.Length > 0 && result[^1] == '-')
 			{
-				Result.Remove(Result.Length - 1, 1);
+				result.Remove(result.Length - 1, 1);
 			}
-			return new StringId<T>(Result.ToString());
+			return new StringId<T>(result.ToString());
 		}
 
 		/// <summary>
 		/// Checks whether this StringId is set
 		/// </summary>
-		public bool IsEmpty
-		{
-			get { return String.IsNullOrEmpty(Text); }
-		}
+		public bool IsEmpty => String.IsNullOrEmpty(_text);
 
 		/// <summary>
 		/// Checks whether the given character is valid within a string id
 		/// </summary>
-		/// <param name="Character">The character to check</param>
+		/// <param name="character">The character to check</param>
 		/// <returns>True if the character is valid</returns>
-		static bool IsValidCharacter(char Character)
+		static bool IsValidCharacter(char character)
 		{
-			if (Character >= 'a' && Character <= 'z')
+			if (character >= 'a' && character <= 'z')
 			{
 				return true;
 			}
-			if (Character >= '0' && Character <= '9')
+			if (character >= '0' && character <= '9')
 			{
 				return true;
 			}
-			if (Character == '-' || Character == '_' || Character == '.')
+			if (character == '-' || character == '_' || character == '.')
 			{
 				return true;
 			}
@@ -157,49 +148,49 @@ namespace Horde.Build.Utilities
 		}
 
 		/// <inheritdoc/>
-		public override bool Equals(object? Obj)
+		public override bool Equals(object? obj)
 		{
-			return Obj is StringId<T> && Equals((StringId<T>)Obj);
+			return obj is StringId<T> id && Equals(id);
 		}
 
 		/// <inheritdoc/>
 		public override int GetHashCode()
 		{
-			return Text.GetHashCode(StringComparison.Ordinal);
+			return _text.GetHashCode(StringComparison.Ordinal);
 		}
 
 		/// <inheritdoc/>
-		public bool Equals(StringId<T> Other)
+		public bool Equals(StringId<T> other)
 		{
-			return Text.Equals(Other.Text, StringComparison.Ordinal);
+			return _text.Equals(other._text, StringComparison.Ordinal);
 		}
 
 		/// <inheritdoc/>
 		public override string ToString()
 		{
-			return Text;
+			return _text;
 		}
 
 		/// <summary>
 		/// Compares two string ids for equality
 		/// </summary>
-		/// <param name="Left">The first string id</param>
-		/// <param name="Right">Second string id</param>
+		/// <param name="left">The first string id</param>
+		/// <param name="right">Second string id</param>
 		/// <returns>True if the two string ids are equal</returns>
-		public static bool operator ==(StringId<T> Left, StringId<T> Right)
+		public static bool operator ==(StringId<T> left, StringId<T> right)
 		{
-			return Left.Equals(Right);
+			return left.Equals(right);
 		}
 
 		/// <summary>
 		/// Compares two string ids for inequality
 		/// </summary>
-		/// <param name="Left">The first string id</param>
-		/// <param name="Right">Second string id</param>
+		/// <param name="left">The first string id</param>
+		/// <param name="right">Second string id</param>
 		/// <returns>True if the two string ids are not equal</returns>
-		public static bool operator !=(StringId<T> Left, StringId<T> Right)
+		public static bool operator !=(StringId<T> left, StringId<T> right)
 		{
-			return !Left.Equals(Right);
+			return !left.Equals(right);
 		}
 	}
 
@@ -209,15 +200,15 @@ namespace Horde.Build.Utilities
 	public class StringIdJsonConverter<T> : JsonConverter<StringId<T>>
 	{
 		/// <inheritdoc/>
-		public override StringId<T> Read(ref Utf8JsonReader Reader, Type TypeToConvert, JsonSerializerOptions Options)
+		public override StringId<T> Read(ref Utf8JsonReader reader, Type typeToConvert, JsonSerializerOptions options)
 		{
-			return new StringId<T>(Reader.GetString()!);
+			return new StringId<T>(reader.GetString()!);
 		}
 
 		/// <inheritdoc/>
-		public override void Write(Utf8JsonWriter Writer, StringId<T> Value, JsonSerializerOptions Options)
+		public override void Write(Utf8JsonWriter writer, StringId<T> value, JsonSerializerOptions options)
 		{
-			Writer.WriteStringValue(Value.ToString());
+			writer.WriteStringValue(value.ToString());
 		}
 	}
 
@@ -227,15 +218,15 @@ namespace Horde.Build.Utilities
 	public class JsonStringIdConverterFactory : JsonConverterFactory
 	{
 		/// <inheritdoc/>
-		public override bool CanConvert(Type TypeToConvert)
+		public override bool CanConvert(Type typeToConvert)
 		{
-			return TypeToConvert.IsGenericType && TypeToConvert.GetGenericTypeDefinition() == typeof(StringId<>);
+			return typeToConvert.IsGenericType && typeToConvert.GetGenericTypeDefinition() == typeof(StringId<>);
 		}
 
 		/// <inheritdoc/>
-		public override JsonConverter? CreateConverter(Type Type, JsonSerializerOptions Options)
+		public override JsonConverter? CreateConverter(Type type, JsonSerializerOptions options)
 		{
-			return (JsonConverter?)Activator.CreateInstance(typeof(StringIdJsonConverter<>).MakeGenericType(Type.GetGenericArguments()));
+			return (JsonConverter?)Activator.CreateInstance(typeof(StringIdJsonConverter<>).MakeGenericType(type.GetGenericArguments()));
 		}
 	}
 
@@ -245,24 +236,24 @@ namespace Horde.Build.Utilities
 	public sealed class StringIdBsonSerializer<T> : SerializerBase<StringId<T>>
 	{
 		/// <inheritdoc/>
-		public override StringId<T> Deserialize(BsonDeserializationContext Context, BsonDeserializationArgs Args)
+		public override StringId<T> Deserialize(BsonDeserializationContext context, BsonDeserializationArgs args)
 		{
-			string Argument;
-			if (Context.Reader.CurrentBsonType == MongoDB.Bson.BsonType.ObjectId)
+			string argument;
+			if (context.Reader.CurrentBsonType == MongoDB.Bson.BsonType.ObjectId)
 			{
-				Argument = Context.Reader.ReadObjectId().ToString();
+				argument = context.Reader.ReadObjectId().ToString();
 			}
 			else
 			{
-				Argument = Context.Reader.ReadString();
+				argument = context.Reader.ReadString();
 			}
-			return new StringId<T>(Argument);
+			return new StringId<T>(argument);
 		}
 
 		/// <inheritdoc/>
-		public override void Serialize(BsonSerializationContext Context, BsonSerializationArgs Args, StringId<T> Value)
+		public override void Serialize(BsonSerializationContext context, BsonSerializationArgs args, StringId<T> value)
 		{
-			Context.Writer.WriteString(Value.ToString());
+			context.Writer.WriteString(value.ToString());
 		}
 	}
 
@@ -287,19 +278,19 @@ namespace Horde.Build.Utilities
 
 	sealed class CbStringIdConverter<T> : CbConverterBase<StringId<T>>
 	{
-		public override StringId<T> Read(CbField Field)
+		public override StringId<T> Read(CbField field)
 		{
-			return new StringId<T>(Field.AsString().ToString());
+			return new StringId<T>(field.AsString().ToString());
 		}
 
-		public override void Write(CbWriter Writer, StringId<T> Value)
+		public override void Write(CbWriter writer, StringId<T> value)
 		{
-			Writer.WriteStringValue(Value.ToString());
+			writer.WriteStringValue(value.ToString());
 		}
 
-		public override void WriteNamed(CbWriter Writer, Utf8String Name, StringId<T> Value)
+		public override void WriteNamed(CbWriter writer, Utf8String name, StringId<T> value)
 		{
-			Writer.WriteString(Name, Value.ToString());
+			writer.WriteString(name, value.ToString());
 		}
 	}
 
@@ -308,65 +299,65 @@ namespace Horde.Build.Utilities
 	/// </summary>
 	sealed class StringIdTypeConverter : TypeConverter
 	{
-		Type Type;
+		readonly Type _type;
 
 		/// <summary>
 		/// Constructor
 		/// </summary>
-		/// <param name="Type"></param>
-		public StringIdTypeConverter(Type Type)
+		/// <param name="type"></param>
+		public StringIdTypeConverter(Type type)
 		{
-			this.Type = Type;
+			_type = type;
 		}
 
 		/// <inheritdoc/>
-		public override bool CanConvertFrom(ITypeDescriptorContext? Context, Type SourceType)
+		public override bool CanConvertFrom(ITypeDescriptorContext? context, Type sourceType)
 		{
-			return SourceType == typeof(string) || base.CanConvertFrom(Context, SourceType);
+			return sourceType == typeof(string) || base.CanConvertFrom(context, sourceType);
 		}
 
 		/// <inheritdoc/>
-		public override object ConvertFrom(ITypeDescriptorContext? Context, CultureInfo? Culture, object Value)
+		public override object ConvertFrom(ITypeDescriptorContext? context, CultureInfo? culture, object value)
 		{
-			return Activator.CreateInstance(Type, Value)!;
+			return Activator.CreateInstance(_type, value)!;
 		}
 
 		/// <inheritdoc/>
-		public override bool CanConvertTo(ITypeDescriptorContext? Context, Type? DestinationType)
+		public override bool CanConvertTo(ITypeDescriptorContext? context, Type? destinationType)
 		{
-			if (DestinationType == null)
+			if (destinationType == null)
 			{
 				return false;
 			}
-			if (DestinationType == typeof(string))
+			if (destinationType == typeof(string))
 			{
 				return true;
 			}
-			if (DestinationType.IsGenericType)
+			if (destinationType.IsGenericType)
 			{
-				Type GenericTypeDefinition = DestinationType.GetGenericTypeDefinition();
-				if (GenericTypeDefinition == typeof(StringId<>))
+				Type genericTypeDefinition = destinationType.GetGenericTypeDefinition();
+				if (genericTypeDefinition == typeof(StringId<>))
 				{
 					return true;
 				}
-				if (GenericTypeDefinition == typeof(Nullable<>))
+				if (genericTypeDefinition == typeof(Nullable<>))
 				{
-					return CanConvertTo(Context, GenericTypeDefinition.GetGenericArguments()[0]);
+					return CanConvertTo(context, genericTypeDefinition.GetGenericArguments()[0]);
 				}
 			}
 			return false;
 		}
 
 		/// <inheritdoc/>
-		public override object? ConvertTo(ITypeDescriptorContext? Context, CultureInfo? Culture, object? Value, Type DestinationType)
+		public override object? ConvertTo(ITypeDescriptorContext? context, CultureInfo? culture, object? value, Type destinationType)
 		{
-			if (DestinationType == typeof(string))
+			if (destinationType == typeof(string))
 			{
-				return Value?.ToString();
+				return value?.ToString();
 			}
 			else
 			{
-				return Activator.CreateInstance(DestinationType, Value);
+				return Activator.CreateInstance(destinationType, value);
 			}
 		}
 	}
diff --git a/Engine/Source/Programs/Horde/Horde.Build/Utilities/StructuredRpcException.cs b/Engine/Source/Programs/Horde/Horde.Build/Utilities/StructuredRpcException.cs
--- a/Engine/Source/Programs/Horde/Horde.Build/Utilities/StructuredRpcException.cs
+++ b/Engine/Source/Programs/Horde/Horde.Build/Utilities/StructuredRpcException.cs
@@ -1,19 +1,15 @@
 // Copyright Epic Games, Inc. All Rights Reserved.
 
-using Grpc.Core;
 using System;
-using System.Collections.Generic;
 using System.Globalization;
-using System.Linq;
 using System.Text;
-using System.Threading.Tasks;
+using Grpc.Core;
 
 namespace Horde.Build.Utilities
 {
 	/// <summary>
 	/// Exception class designed to allow logging structured log messages
 	/// </summary>
-	[System.Diagnostics.CodeAnalysis.SuppressMessage("Design", "CA1032:Implement standard exception constructors", Justification = "<Pending>")]
 	public class StructuredRpcException : RpcException
 	{
 		/// <summary>
@@ -29,54 +25,54 @@ namespace Horde.Build.Utilities
 		/// <summary>
 		/// Constructor
 		/// </summary>
-		/// <param name="StatusCode">Status code to return</param>
-		/// <param name="Format">The format string</param>
-		/// <param name="Args">Arguments for the format string</param>
-		public StructuredRpcException(StatusCode StatusCode, string Format, params object[] Args)
-			: base(new Status(StatusCode, FormatMessage(Format, Args)))
+		/// <param name="statusCode">Status code to return</param>
+		/// <param name="format">The format string</param>
+		/// <param name="args">Arguments for the format string</param>
+		public StructuredRpcException(StatusCode statusCode, string format, params object[] args)
+			: base(new Status(statusCode, FormatMessage(format, args)))
 		{
-			this.Format = Format;
-			this.Args = Args;
+			Format = format;
+			Args = args;
 		}
 
 		/// <summary>
 		/// Replace named arguments in the format message with their values
 		/// </summary>
-		/// <param name="Format"></param>
-		/// <param name="Args"></param>
+		/// <param name="format"></param>
+		/// <param name="args"></param>
 		/// <returns></returns>
-		static string FormatMessage(string Format, params object[] Args)
+		static string FormatMessage(string format, params object[] args)
 		{
-			string NewFormat = ConvertToFormatString(Format);
-			return String.Format(CultureInfo.CurrentCulture, NewFormat, Args);
+			string newFormat = ConvertToFormatString(format);
+			return String.Format(CultureInfo.CurrentCulture, newFormat, args);
 		}
 
 		/// <summary>
 		/// Converts a named parameter format string to a String.Format style string
 		/// </summary>
-		/// <param name="Format"></param>
+		/// <param name="format"></param>
 		/// <returns></returns>
-		static string ConvertToFormatString(string Format)
+		static string ConvertToFormatString(string format)
 		{
-			int ArgIdx = 0;
+			int argIdx = 0;
 
-			StringBuilder NewFormat = new StringBuilder();
-			for (int Idx = 0; Idx < Format.Length; Idx++)
+			StringBuilder newFormat = new StringBuilder();
+			for (int idx = 0; idx < format.Length; idx++)
 			{
-				char Character = Format[Idx];
-				NewFormat.Append(Character);
+				char character = format[idx];
+				newFormat.Append(character);
 
-				if (Character == '{' && Idx + 1 < Format.Length)
+				if (character == '{' && idx + 1 < format.Length)
 				{
-					char NextCharacter = Format[Idx + 1];
-					if ((NextCharacter >= 'a' && NextCharacter <= 'z') || (NextCharacter >= 'A' && NextCharacter <= 'Z') || NextCharacter == '_')
+					char nextCharacter = format[idx + 1];
+					if ((nextCharacter >= 'a' && nextCharacter <= 'z') || (nextCharacter >= 'A' && nextCharacter <= 'Z') || nextCharacter == '_')
 					{
-						for (int EndIdx = Idx + 2; EndIdx < Format.Length; EndIdx++)
+						for (int endIdx = idx + 2; endIdx < format.Length; endIdx++)
 						{
-							if (Format[EndIdx] == ':' || Format[EndIdx] == '}')
+							if (format[endIdx] == ':' || format[endIdx] == '}')
 							{
-								NewFormat.Append(ArgIdx++);
-								Idx = EndIdx - 1;
+								newFormat.Append(argIdx++);
+								idx = endIdx - 1;
 								break;
 							}
 						}
@@ -84,7 +80,7 @@ namespace Horde.Build.Utilities
 				}
 			}
 
-			return NewFormat.ToString();
+			return newFormat.ToString();
 		}
 	}
 }
