Horde: Fix coding conventions within Horde.Build.

diff --git a/Engine/Source/Programs/Horde/Horde.Build/Utilities/TracingExtensions.cs b/Engine/Source/Programs/Horde/Horde.Build/Utilities/TracingExtensions.cs
--- a/Engine/Source/Programs/Horde/Horde.Build/Utilities/TracingExtensions.cs
+++ b/Engine/Source/Programs/Horde/Horde.Build/Utilities/TracingExtensions.cs
@@ -4,7 +4,6 @@ using System;
 using System.Collections.Generic;
 using EpicGames.Core;
 using Horde.Build.Models;
-using Horde.Build.Utilities;
 using OpenTracing;
 
 namespace Horde.Build.Utilities
@@ -16,82 +15,100 @@ namespace Horde.Build.Utilities
 	{
 		/// <summary>Set a key:value tag on the span</summary>
 		/// <returns>This span instance, for chaining</returns>
-		public static ISpan SetTag<T>(this ISpan Span, string Key, StringId<T>? Value)
+		public static ISpan SetTag<T>(this ISpan span, string key, StringId<T>? value)
 		{
-			if (Value != null) Span.SetTag(Key, Value.ToString());
-			return Span;
+			if (value != null)
+			{
+				span.SetTag(key, value.ToString());
+			}
+			return span;
 		}
 		
 		/// <summary>Set a key:value tag on the span</summary>
 		/// <returns>This span instance, for chaining</returns>
-		public static ISpan SetTag<T>(this ISpan Span, string Key, StringId<T> Value)
+		public static ISpan SetTag<T>(this ISpan span, string key, StringId<T> value)
 		{
-			Span.SetTag(Key, Value.ToString());
-			return Span;
+			span.SetTag(key, value.ToString());
+			return span;
 		}
 
 		/// <summary>Set a key:value tag on the span</summary>
 		/// <returns>This span instance, for chaining</returns>
-		public static ISpan SetTag<T>(this ISpan Span, string Key, IEnumerable<StringId<T>>? Value)
+		public static ISpan SetTag<T>(this ISpan span, string key, IEnumerable<StringId<T>>? value)
 		{
-			if (Value != null) Span.SetTag(Key, string.Join(',', Value));
-			return Span;
+			if (value != null)
+			{
+				span.SetTag(key, String.Join(',', value));
+			}
+			return span;
 		}
 
 		/// <summary>Set a key:value tag on the span</summary>
 		/// <returns>This span instance, for chaining</returns>
-		public static ISpan SetTag<T>(this ISpan Span, string Key, ObjectId<T>? Value)
+		public static ISpan SetTag<T>(this ISpan span, string key, ObjectId<T>? value)
 		{
-			if (Value != null) Span.SetTag(Key, Value.ToString());
-			return Span;
+			if (value != null)
+			{
+				span.SetTag(key, value.ToString());
+			}
+			return span;
 		}
 		
 		/// <summary>Set a key:value tag on the span</summary>
 		/// <returns>This span instance, for chaining</returns>
-		public static ISpan SetTag<T>(this ISpan Span, string Key, ObjectId<T> Value)
+		public static ISpan SetTag<T>(this ISpan span, string key, ObjectId<T> value)
 		{
-			Span.SetTag(Key, Value.ToString());
-			return Span;
+			span.SetTag(key, value.ToString());
+			return span;
 		}
 		
 		/// <summary>Set a key:value tag on the span</summary>
 		/// <returns>This span instance, for chaining</returns>
-		public static ISpan SetTag(this ISpan Span, string Key, ContentHash Value)
+		public static ISpan SetTag(this ISpan span, string key, ContentHash value)
 		{
-			Span.SetTag(Key, Value.ToString());
-			return Span;
+			span.SetTag(key, value.ToString());
+			return span;
 		}
 		
 		/// <summary>Set a key:value tag on the span</summary>
 		/// <returns>This span instance, for chaining</returns>
-		public static ISpan SetTag(this ISpan Span, string Key, SubResourceId Value)
+		public static ISpan SetTag(this ISpan span, string key, SubResourceId value)
 		{
-			Span.SetTag(Key, Value.ToString());
-			return Span;
+			span.SetTag(key, value.ToString());
+			return span;
 		}
 		
 		/// <summary>Set a key:value tag on the span</summary>
 		/// <returns>This span instance, for chaining</returns>
-		public static ISpan SetTag<T>(this ISpan Span, string Key, IEnumerable<ObjectId<T>>? Value)
+		public static ISpan SetTag<T>(this ISpan span, string key, IEnumerable<ObjectId<T>>? value)
 		{
-			if (Value != null) Span.SetTag(Key, string.Join(',', Value));
-			return Span;
+			if (value != null)
+			{
+				span.SetTag(key, String.Join(',', value));
+			}
+			return span;
 		}
 		
 		/// <summary>Set a key:value tag on the span</summary>
 		/// <returns>This span instance, for chaining</returns>
-		public static ISpan SetTag(this ISpan Span, string Key, int? Value)
+		public static ISpan SetTag(this ISpan span, string key, int? value)
 		{
-			if (Value != null) Span.SetTag(Key, Value.Value);
-			return Span;
+			if (value != null)
+			{
+				span.SetTag(key, value.Value);
+			}
+			return span;
 		}
 		
 		/// <summary>Set a key:value tag on the span</summary>
 		/// <returns>This span instance, for chaining</returns>
-		public static ISpan SetTag(this ISpan Span, string Key, DateTimeOffset? Value)
+		public static ISpan SetTag(this ISpan span, string key, DateTimeOffset? value)
 		{
-			if (Value != null) Span.SetTag(Key, Value.ToString());
-			return Span;
+			if (value != null)
+			{
+				span.SetTag(key, value.ToString());
+			}
+			return span;
 		}
 	}
 }
\ No newline at end of file
diff --git a/Engine/Source/Programs/Horde/Horde.Build/Utilities/VersionedCollection.cs b/Engine/Source/Programs/Horde/Horde.Build/Utilities/VersionedCollection.cs
--- a/Engine/Source/Programs/Horde/Horde.Build/Utilities/VersionedCollection.cs
+++ b/Engine/Source/Programs/Horde/Horde.Build/Utilities/VersionedCollection.cs
@@ -1,5 +1,8 @@
 ﻿// Copyright Epic Games, Inc. All Rights Reserved.
 
+using System;
+using System.Collections.Generic;
+using System.Threading.Tasks;
 using Horde.Build.Services;
 using Horde.Build.Utilities;
 using MongoDB.Bson;
@@ -7,9 +10,6 @@ using MongoDB.Bson.Serialization;
 using MongoDB.Bson.Serialization.Attributes;
 using MongoDB.Driver;
 using StackExchange.Redis;
-using System;
-using System.Collections.Generic;
-using System.Threading.Tasks;
 
 namespace Horde.Build.Tools.Impl
 {
@@ -60,7 +60,6 @@ namespace Horde.Build.Tools.Impl
 		private static readonly RedisKey s_timeSuffix = new RedisKey("/ts");
 
 		private static FilterDefinitionBuilder<VersionedDocument<TId, TLatest>> FilterBuilder { get; } = Builders<VersionedDocument<TId, TLatest>>.Filter;
-		private static UpdateDefinitionBuilder<VersionedDocument<TId, TLatest>> UpdateBuilder { get; } = Builders<VersionedDocument<TId, TLatest>>.Update;
 
 		private readonly IMongoCollection<VersionedDocument<TId, TLatest>> _baseCollection;
 		private readonly IMongoCollection<TLatest> _latestCollection;
@@ -102,14 +101,14 @@ namespace Horde.Build.Tools.Impl
 
 		private static FilterDefinition<VersionedDocument<TId, TLatest>> GetFilter(VersionedDocument<TId, TLatest> doc)
 		{
-			FilterDefinitionBuilder<VersionedDocument<TId, TLatest>> Builder = Builders<VersionedDocument<TId, TLatest>>.Filter;
-			return Builder.Eq(x => x.Id, doc.Id) & Builder.Eq(x => x.LastUpdateTime, doc.LastUpdateTime);
+			FilterDefinitionBuilder<VersionedDocument<TId, TLatest>> builder = Builders<VersionedDocument<TId, TLatest>>.Filter;
+			return builder.Eq(x => x.Id, doc.Id) & builder.Eq(x => x.LastUpdateTime, doc.LastUpdateTime);
 		}
 
 		private static FilterDefinition<TLatest> GetFilter(TLatest doc)
 		{
-			FilterDefinitionBuilder<TLatest> Builder = Builders<TLatest>.Filter;
-			return Builder.Eq(x => x.Id, doc.Id) & Builder.Eq(x => x.LastUpdateTime, doc.LastUpdateTime);
+			FilterDefinitionBuilder<TLatest> builder = Builders<TLatest>.Filter;
+			return builder.Eq(x => x.Id, doc.Id) & builder.Eq(x => x.LastUpdateTime, doc.LastUpdateTime);
 		}
 
 		/// <summary>
@@ -132,11 +131,11 @@ namespace Horde.Build.Tools.Impl
 
 		private void AddCachedValue(RedisKey docKey, TLatest doc)
 		{
-			KeyValuePair<RedisKey, RedisValue>[] Pairs = new KeyValuePair<RedisKey, RedisValue>[2];
-			Pairs[0] = new KeyValuePair<RedisKey, RedisValue>(docKey, doc.ToBson());
-			Pairs[1] = new KeyValuePair<RedisKey, RedisValue>(docKey.Append(s_timeSuffix), doc.LastUpdateTime.Ticks);
+			KeyValuePair<RedisKey, RedisValue>[] pairs = new KeyValuePair<RedisKey, RedisValue>[2];
+			pairs[0] = new KeyValuePair<RedisKey, RedisValue>(docKey, doc.ToBson());
+			pairs[1] = new KeyValuePair<RedisKey, RedisValue>(docKey.Append(s_timeSuffix), doc.LastUpdateTime.Ticks);
 
-			_ = _redis.StringSetAsync(Pairs, When.NotExists, CommandFlags.FireAndForget);
+			_ = _redis.StringSetAsync(pairs, When.NotExists, CommandFlags.FireAndForget);
 		}
 
 		private async ValueTask<VersionedDocument<TId, TLatest>?> GetCachedValueAsync(RedisKey docKey)
@@ -161,11 +160,11 @@ namespace Horde.Build.Tools.Impl
 			ITransaction transaction = _redis.CreateTransaction();
 			transaction.AddCondition(Condition.StringEqual(docKey.Append(s_timeSuffix), prevDoc.LastUpdateTime.Ticks));
 
-			KeyValuePair<RedisKey, RedisValue>[] Pairs = new KeyValuePair<RedisKey, RedisValue>[2];
-			Pairs[0] = new KeyValuePair<RedisKey, RedisValue>(docKey, doc.ToBson());
-			Pairs[1] = new KeyValuePair<RedisKey, RedisValue>(docKey.Append(s_timeSuffix), doc.LastUpdateTime.Ticks);
+			KeyValuePair<RedisKey, RedisValue>[] pairs = new KeyValuePair<RedisKey, RedisValue>[2];
+			pairs[0] = new KeyValuePair<RedisKey, RedisValue>(docKey, doc.ToBson());
+			pairs[1] = new KeyValuePair<RedisKey, RedisValue>(docKey.Append(s_timeSuffix), doc.LastUpdateTime.Ticks);
 
-			_ = transaction.StringSetAsync(Pairs);
+			_ = transaction.StringSetAsync(pairs);
 
 			if (!await transaction.ExecuteAsync())
 			{
diff --git a/Engine/Source/Programs/Horde/Horde.Build/Utilities/ViewMap.cs b/Engine/Source/Programs/Horde/Horde.Build/Utilities/ViewMap.cs
--- a/Engine/Source/Programs/Horde/Horde.Build/Utilities/ViewMap.cs
+++ b/Engine/Source/Programs/Horde/Horde.Build/Utilities/ViewMap.cs
@@ -1,13 +1,11 @@
 // Copyright Epic Games, Inc. All Rights Reserved.
 
-using EpicGames.Core;
 using System;
 using System.Collections.Generic;
-using System.Diagnostics;
-using System.Diagnostics.CodeAnalysis;
 using System.Globalization;
 using System.Text;
 using System.Text.RegularExpressions;
+using EpicGames.Core;
 
 namespace Horde.Build.Utilities
 {
@@ -34,70 +32,70 @@ namespace Horde.Build.Utilities
 		/// <summary>
 		/// Constructor
 		/// </summary>
-		/// <param name="Other"></param>
-		public ViewMap(ViewMap Other)
+		/// <param name="other"></param>
+		public ViewMap(ViewMap other)
 		{
-			Entries = new List<ViewMapEntry>(Other.Entries);
+			Entries = new List<ViewMapEntry>(other.Entries);
 		}
 
 		/// <summary>
 		/// Constructor
 		/// </summary>
-		/// <param name="ViewMap">Peforce viewmap definition</param>
-		public ViewMap(P4.ViewMap ViewMap)
+		/// <param name="viewMap">Peforce viewmap definition</param>
+		public ViewMap(P4.ViewMap viewMap)
 			: this()
 		{
-			foreach (P4.MapEntry Entry in ViewMap)
+			foreach (P4.MapEntry entry in viewMap)
 			{
-				Entries.Add(new ViewMapEntry(Entry));
+				Entries.Add(new ViewMapEntry(entry));
 			}
 		}
 
 		/// <summary>
 		/// Determines if a file is included in the view
 		/// </summary>
-		/// <param name="File">The file to test</param>
-		/// <param name="Comparison">The comparison type</param>
+		/// <param name="file">The file to test</param>
+		/// <param name="comparison">The comparison type</param>
 		/// <returns>True if the file is included in the view</returns>
-		public bool MatchFile(Utf8String File, Utf8StringComparer Comparison)
+		public bool MatchFile(Utf8String file, Utf8StringComparer comparison)
 		{
-			bool Included = false;
-			foreach (ViewMapEntry Entry in Entries)
+			bool included = false;
+			foreach (ViewMapEntry entry in Entries)
 			{
-				if (Entry.MatchFile(File, Comparison))
+				if (entry.MatchFile(file, comparison))
 				{
-					Included = Entry.Include;
+					included = entry.Include;
 				}
 			}
-			return Included;
+			return included;
 		}
 
 		/// <summary>
 		/// Attempts to convert a source file to its target path
 		/// </summary>
-		/// <param name="SourceFile"></param>
-		/// <param name="Comparison">The comparison type</param>
-		/// <param name="TargetFile"></param>
+		/// <param name="sourceFile"></param>
+		/// <param name="comparison">The comparison type</param>
+		/// <param name="targetFile"></param>
 		/// <returns></returns>
-		public bool TryMapFile(Utf8String SourceFile, Utf8StringComparer Comparison, out Utf8String TargetFile)
+		public bool TryMapFile(Utf8String sourceFile, Utf8StringComparer comparison, out Utf8String targetFile)
 		{
-			ViewMapEntry? MapEntry = null;
-			foreach (ViewMapEntry Entry in Entries)
+			ViewMapEntry? mapEntry = null;
+			foreach (ViewMapEntry entry in Entries)
 			{
-				if (Entry.MatchFile(SourceFile, Comparison))
+				if (entry.MatchFile(sourceFile, comparison))
 				{
-					MapEntry = Entry;
+					mapEntry = entry;
 				}
 			}
 
-			if (MapEntry != null && MapEntry.Include)
+			if (mapEntry != null && mapEntry.Include)
 			{
-				TargetFile = MapEntry.MapFile(SourceFile);
+				targetFile = mapEntry.MapFile(sourceFile);
 				return true;
 			}
 			else
 			{
-				TargetFile = Utf8String.Empty;
+				targetFile = Utf8String.Empty;
 				return false;
 			}
 		}
@@ -106,36 +104,36 @@ namespace Horde.Build.Utilities
 		/// Gets the root paths from the view entries
 		/// </summary>
 		/// <returns></returns>
-		public List<Utf8String> GetRootPaths(Utf8StringComparer Comparison)
+		public List<Utf8String> GetRootPaths(Utf8StringComparer comparison)
 		{
-			List<Utf8String> RootPaths = new List<Utf8String>();
-			foreach (ViewMapEntry Entry in Entries)
+			List<Utf8String> rootPaths = new List<Utf8String>();
+			foreach (ViewMapEntry entry in Entries)
 			{
-				if (Entry.Include)
+				if (entry.Include)
 				{
-					int LastSlashIdx = Entry.SourcePrefix.LastIndexOf('/');
-					Utf8String RootPath = Entry.SourcePrefix.Slice(0, LastSlashIdx + 1);
+					int lastSlashIdx = entry.SourcePrefix.LastIndexOf('/');
+					Utf8String rootPath = entry.SourcePrefix.Slice(0, lastSlashIdx + 1);
 
-					for (int Idx = 0; ; Idx++)
+					for (int idx = 0; ; idx++)
 					{
-						if (Idx == RootPaths.Count)
+						if (idx == rootPaths.Count)
 						{
-							RootPaths.Add(RootPath);
+							rootPaths.Add(rootPath);
 							break;
 						}
-						else if (RootPaths[Idx].StartsWith(RootPath, Comparison))
+						else if (rootPaths[idx].StartsWith(rootPath, comparison))
 						{
-							RootPaths[Idx] = RootPath;
+							rootPaths[idx] = rootPath;
 							break;
 						}
-						else if (RootPath.StartsWith(RootPaths[Idx], Comparison))
+						else if (rootPath.StartsWith(rootPaths[idx], comparison))
 						{
 							break;
 						}
 					}
 				}
 			}
-			return RootPaths;
+			return rootPaths;
 		}
 	}
 
@@ -199,45 +197,45 @@ namespace Horde.Build.Utilities
 		/// <summary>
 		/// Constructor
 		/// </summary>
-		/// <param name="Other"></param>
-		public ViewMapEntry(ViewMapEntry Other)
-			: this(Other.Include, Other.Wildcard, Other.SourcePrefix, Other.SourceSuffix, Other.TargetPrefix, Other.TargetSuffix)
+		/// <param name="other"></param>
+		public ViewMapEntry(ViewMapEntry other)
+			: this(other.Include, other.Wildcard, other.SourcePrefix, other.SourceSuffix, other.TargetPrefix, other.TargetSuffix)
 		{
 		}
 
 		/// <summary>
 		/// Constructor
 		/// </summary>
-		/// <param name="Include"></param>
-		/// <param name="Source"></param>
-		/// <param name="Target"></param>
-		public ViewMapEntry(bool Include, string Source, string Target)
+		/// <param name="include"></param>
+		/// <param name="source"></param>
+		/// <param name="target"></param>
+		public ViewMapEntry(bool include, string source, string target)
 		{
-			this.Include = Include;
+			Include = include;
 
-			Match Match = Regex.Match(Source, @"^(.*)(\*|\.\.\.|%%1)(.*)$");
-			if (Match.Success)
+			Match match = Regex.Match(source, @"^(.*)(\*|\.\.\.|%%1)(.*)$");
+			if (match.Success)
 			{
-				string WildcardStr = Match.Groups[2].Value;
+				string wildcardStr = match.Groups[2].Value;
 
-				SourcePrefix = Match.Groups[1].Value;
-				SourceSuffix = Match.Groups[3].Value;
-				Wildcard = Match.Groups[2].Value;
+				SourcePrefix = match.Groups[1].Value;
+				SourceSuffix = match.Groups[3].Value;
+				Wildcard = match.Groups[2].Value;
 
-				int OtherIdx = Target.IndexOf(WildcardStr, StringComparison.Ordinal);
-				TargetPrefix = Target.Substring(0, OtherIdx);
-				TargetSuffix = Target.Substring(OtherIdx + Wildcard.Length);
+				int otherIdx = target.IndexOf(wildcardStr, StringComparison.Ordinal);
+				TargetPrefix = target.Substring(0, otherIdx);
+				TargetSuffix = target.Substring(otherIdx + Wildcard.Length);
 
-				if (WildcardStr.Equals("%%1", StringComparison.Ordinal))
+				if (wildcardStr.Equals("%%1", StringComparison.Ordinal))
 				{
 					Wildcard = "*";
 				}
 			}
 			else
 			{
-				SourcePrefix = Source;
+				SourcePrefix = source;
 				SourceSuffix = Utf8String.Empty;
-				TargetPrefix = Target;
+				TargetPrefix = target;
 				TargetSuffix = Utf8String.Empty;
 			}
 		}
@@ -245,61 +243,61 @@ namespace Horde.Build.Utilities
 		/// <summary>
 		/// Constructor
 		/// </summary>
-		/// <param name="Include"></param>
-		/// <param name="Wildcard"></param>
-		/// <param name="SourcePrefix"></param>
-		/// <param name="SourceSuffix"></param>
-		/// <param name="TargetPrefix"></param>
-		/// <param name="TargetSuffix"></param>
-		public ViewMapEntry(bool Include, Utf8String Wildcard, Utf8String SourcePrefix, Utf8String SourceSuffix, Utf8String TargetPrefix, Utf8String TargetSuffix)
+		/// <param name="include"></param>
+		/// <param name="wildcard"></param>
+		/// <param name="sourcePrefix"></param>
+		/// <param name="sourceSuffix"></param>
+		/// <param name="targetPrefix"></param>
+		/// <param name="targetSuffix"></param>
+		public ViewMapEntry(bool include, Utf8String wildcard, Utf8String sourcePrefix, Utf8String sourceSuffix, Utf8String targetPrefix, Utf8String targetSuffix)
 		{
-			this.Include = Include;
-			this.Wildcard = Wildcard;
-			this.SourcePrefix = SourcePrefix;
-			this.SourceSuffix = SourceSuffix;
-			this.TargetPrefix = TargetPrefix;
-			this.TargetSuffix = TargetSuffix;
+			Include = include;
+			Wildcard = wildcard;
+			SourcePrefix = sourcePrefix;
+			SourceSuffix = sourceSuffix;
+			TargetPrefix = targetPrefix;
+			TargetSuffix = targetSuffix;
 		}
 
 		/// <summary>
 		/// Constructor
 		/// </summary>
-		/// <param name="Entry"></param>
-		public ViewMapEntry(P4.MapEntry Entry)
-			: this(Entry.Type == P4.MapType.Include, Entry.Left.Path, Entry.Right.Path)
+		/// <param name="entry"></param>
+		public ViewMapEntry(P4.MapEntry entry)
+			: this(entry.Type == P4.MapType.Include, entry.Left.Path, entry.Right.Path)
 		{
 		}
 
 		/// <summary>
 		/// Maps a file to the target path
 		/// </summary>
-		/// <param name="SourceFile"></param>
+		/// <param name="sourceFile"></param>
 		/// <returns></returns>
-		public Utf8String MapFile(Utf8String SourceFile)
+		public Utf8String MapFile(Utf8String sourceFile)
 		{
-			int Count = SourceFile.Length - SourceSuffix.Length - SourcePrefix.Length;
-			return TargetPrefix + SourceFile.Slice(SourcePrefix.Length, Count) + TargetSuffix;
+			int count = sourceFile.Length - SourceSuffix.Length - SourcePrefix.Length;
+			return TargetPrefix + sourceFile.Slice(SourcePrefix.Length, count) + TargetSuffix;
 		}
 
 		/// <summary>
 		/// Determine if a file matches the current entry
 		/// </summary>
-		/// <param name="Path">Path to the file</param>
-		/// <param name="Comparison">The comparison type</param>
+		/// <param name="path">Path to the file</param>
+		/// <param name="comparison">The comparison type</param>
 		/// <returns>True if the path matches the entry</returns>
-		public bool MatchFile(Utf8String Path, Utf8StringComparer Comparison)
+		public bool MatchFile(Utf8String path, Utf8StringComparer comparison)
 		{
 			if (Wildcard.Length == 0)
 			{
-				return Comparison.Compare(Path, SourcePrefix) == 0;
+				return comparison.Compare(path, SourcePrefix) == 0;
 			}
 			else
 			{
-				if (!Path.StartsWith(SourcePrefix, Comparison) || !Path.EndsWith(SourceSuffix, Comparison))
+				if (!path.StartsWith(SourcePrefix, comparison) || !path.EndsWith(SourceSuffix, comparison))
 				{
 					return false;
 				}
-				if (IsFileWildcard() && Path.Slice(SourcePrefix.Length, Path.Length - SourceSuffix.Length - SourcePrefix.Length).IndexOf('/') != -1)
+				if (IsFileWildcard() && path.Slice(SourcePrefix.Length, path.Length - SourceSuffix.Length - SourcePrefix.Length).IndexOf('/') != -1)
 				{
 					return false;
 				}
@@ -310,13 +308,13 @@ namespace Horde.Build.Utilities
 		/// <inheritdoc/>
 		public override string ToString()
 		{
-			StringBuilder Builder = new StringBuilder();
+			StringBuilder builder = new StringBuilder();
 			if (!Include)
 			{
-				Builder.Append('-');
+				builder.Append('-');
 			}
-			Builder.Append(CultureInfo.InvariantCulture, $"{SourcePrefix}{Wildcard}{SourceSuffix} {TargetPrefix}{Wildcard}{TargetSuffix}");
-			return Builder.ToString();
+			builder.Append(CultureInfo.InvariantCulture, $"{SourcePrefix}{Wildcard}{SourceSuffix} {TargetPrefix}{Wildcard}{TargetSuffix}");
+			return builder.ToString();
 		}
 	}
 }
