Horde: Remove reference to Datadog tracing types.

diff --git a/Engine/Source/Programs/Horde/Horde.Build/Program.cs b/Engine/Source/Programs/Horde/Horde.Build/Program.cs
--- a/Engine/Source/Programs/Horde/Horde.Build/Program.cs
+++ b/Engine/Source/Programs/Horde/Horde.Build/Program.cs
@@ -158,11 +158,7 @@ namespace Horde.Build
 
 			if (hordeSettings.WithDatadog)
 			{
-				using Datadog.Trace.IScope? _ = Datadog.Trace.Tracer.Instance.StartActive("Trace Test");
-				Serilog.Log.Logger.Information("Enabling datadog tracing");
 				GlobalTracer.Register(Datadog.Trace.OpenTracing.OpenTracingTracerFactory.WrapTracer(Datadog.Trace.Tracer.Instance));
-				using IScope scope = GlobalTracer.Instance.BuildSpan("OpenTrace Test").StartActive();
-				scope.Span.SetTag("TestProp", "hello");
 				Serilog.Log.Logger.Information("Enabling datadog tracing (OpenTrace)");
 			}
 
