Horde: Show abort requests on dashboard
#rnx
#skipci

diff --git a/Engine/Source/Programs/Horde/HordeDashboard/src/backend/JobDetails.ts b/Engine/Source/Programs/Horde/HordeDashboard/src/backend/JobDetails.ts
--- a/Engine/Source/Programs/Horde/HordeDashboard/src/backend/JobDetails.ts
+++ b/Engine/Source/Programs/Horde/HordeDashboard/src/backend/JobDetails.ts
@@ -48,14 +48,14 @@ export class JobDetails {
 
             const targetArgs = detailArgs.map(arg => arg.trim()).filter(arg => arg.toLowerCase().startsWith("-target="));
 
-            targetArgs.forEach(t => { 
-                
+            targetArgs.forEach(t => {
+
                 const index = t.indexOf("=");
                 if (index === -1) {
                     return;
                 }
-                const target = t.slice(index+1)?.trim();
-                
+                const target = t.slice(index + 1)?.trim();
+
                 if (target) {
                     targets.push(target);
                 }
@@ -409,7 +409,7 @@ export class JobDetails {
             batch.agentRate = 5.0;
         }
         */
-                
+
         if (!step || !batch || !batch.agentRate || !step.startTime || !step.finishTime) {
             return undefined;
         }
@@ -418,7 +418,7 @@ export class JobDetails {
         const end = moment(step.finishTime);
         const hours = moment.duration(end.diff(start)).asHours();
         const price = hours * batch.agentRate;
-                                
+
         return price ? price : undefined;
     }
 
@@ -441,8 +441,8 @@ export class JobDetails {
                 const end = moment(b.finishTime);
                 const hours = moment.duration(end.diff(start)).asHours();
                 price += hours * b.agentRate;
-                                
-            } 
+
+            }
         });
 
         return price ? price : undefined;
@@ -1190,11 +1190,7 @@ export const getStepSummaryMarkdown = (jobDetails: JobDetails, stepId: string):
         }
 
     }
-
-    if (step.state === JobStepState.Skipped) {
-        eta.display = eta.server = "";
-        text.push("The step was skipped");
-    } else if (step.state === JobStepState.Aborted) {
+    if (step.abortRequested || step.state === JobStepState.Aborted) {
         eta.display = eta.server = "";
         let aborted = "";
         if (step.abortedByUserInfo) {
@@ -1207,11 +1203,14 @@ export const getStepSummaryMarkdown = (jobDetails: JobDetails, stepId: string):
             aborted = "The step was aborted";
         }
         text.push(aborted);
+    } else if (step.state === JobStepState.Skipped) {
+        eta.display = eta.server = "";
+        text.push("The step was skipped");
     } else if (step.state === JobStepState.Ready || step.state === JobStepState.Waiting) {
 
         text.push(batchText() ?? `The step is pending in ${step.state} state`);
     }
-    
+
     if (batch?.agentId) {
 
         if (step.startTime) {
@@ -1237,7 +1236,9 @@ export const getStepSummaryMarkdown = (jobDetails: JobDetails, stepId: string):
         }
 
     } else {
-        text.push(batchText() ?? "Step does not have a batch.");
+        if (!step.abortRequested) {
+            text.push(batchText() ?? "Step does not have a batch.");
+        }        
     }
 
     if (eta.display) {
