Horde: Better step abort handling
#skipci
#rnx

diff --git a/Engine/Source/Programs/Horde/Horde.Build/Jobs/Impl/JobCollection.cs b/Engine/Source/Programs/Horde/Horde.Build/Jobs/Impl/JobCollection.cs
--- a/Engine/Source/Programs/Horde/Horde.Build/Jobs/Impl/JobCollection.cs
+++ b/Engine/Source/Programs/Horde/Horde.Build/Jobs/Impl/JobCollection.cs
@@ -805,6 +805,32 @@ namespace Horde.Build.Collections.Impl
 						JobStepDocument step = batch.Steps[stepIdx];
 						if (step.Id == stepId)
 						{
+
+							// Update the request abort status
+							if (newAbortRequested != null && step.AbortRequested == false)
+							{
+								step.AbortRequested = newAbortRequested.Value;
+								updates.Add(updateBuilder.Set(x => x.Batches[batchIdx].Steps[stepIdx].AbortRequested, step.AbortRequested));
+
+								// If the step is pending, and not running on an agent, set to aborted
+								if (step.IsPending() && step.State != JobStepState.Running)
+								{
+									newState = JobStepState.Aborted;
+									newOutcome = JobStepOutcome.Failure;
+								}
+
+								bRefreshDependentJobSteps = true;
+							}
+
+							// Update the user that requested the abort
+							if (newAbortByUserId != null && step.AbortedByUserId == null)
+							{
+								step.AbortedByUserId = newAbortByUserId;
+								updates.Add(updateBuilder.Set(x => x.Batches[batchIdx].Steps[stepIdx].AbortedByUserId, step.AbortedByUserId));
+
+								bRefreshDependentJobSteps = true;
+							}
+
 							// Update the state
 							if (newState != JobStepState.Unspecified && step.State != newState)
 							{
@@ -834,24 +860,6 @@ namespace Horde.Build.Collections.Impl
 								bRefreshDependentJobSteps = true;
 							}
 
-							// Update the request abort status
-							if (newAbortRequested != null && step.AbortRequested == false)
-							{
-								step.AbortRequested = newAbortRequested.Value;
-								updates.Add(updateBuilder.Set(x => x.Batches[batchIdx].Steps[stepIdx].AbortRequested, step.AbortRequested));
-
-								bRefreshDependentJobSteps = true;
-							}
-
-							// Update the user that requested the abort
-							if (newAbortByUserId != null && step.AbortedByUserId == null)
-							{
-								step.AbortedByUserId = newAbortByUserId;
-								updates.Add(updateBuilder.Set(x => x.Batches[batchIdx].Steps[stepIdx].AbortedByUserId, step.AbortedByUserId));
-
-								bRefreshDependentJobSteps = true;
-							}
-
 							// Update the log id
 							if (newLogId != null && step.LogId != newLogId.Value)
 							{
@@ -1620,12 +1628,12 @@ namespace Horde.Build.Collections.Impl
 						if (newState == JobStepState.Waiting)
 						{
 							List<IJobStep> steps = GetDependentSteps(graph, node, stepForNode);
-							if (steps.Any(x => x.IsFailedOrSkipped()))
+							if (steps.Any(x => x.AbortRequested || x.IsFailedOrSkipped()))
 							{
 								newState = JobStepState.Skipped;
 								newOutcome = JobStepOutcome.Failure;
 							}
-							else if (!steps.Any(x => x.IsPending()))
+							else if (!steps.Any(x => !x.AbortRequested && x.IsPending()))
 							{
 								logger.LogDebug("Transitioning job {JobId}, batch {BatchId}, step {StepId} to ready state ({Dependencies})", job.Id, batch.Id, step.Id, String.Join(", ", steps.Select(x => x.Id.ToString())));
 								newState = JobStepState.Ready;
