Improve loading feeback for Automated Test Session global view and history

#rnx

diff --git a/Engine/Source/Programs/Horde/HordeDashboard/src/Plugins/AutomatedTestSession/models/AutomatedTestSession.ts b/Engine/Source/Programs/Horde/HordeDashboard/src/Plugins/AutomatedTestSession/models/AutomatedTestSession.ts
--- a/Engine/Source/Programs/Horde/HordeDashboard/src/Plugins/AutomatedTestSession/models/AutomatedTestSession.ts
+++ b/Engine/Source/Programs/Horde/HordeDashboard/src/Plugins/AutomatedTestSession/models/AutomatedTestSession.ts
@@ -95,6 +95,11 @@ type MetadataHolder = {
     Metadata?: Metadata
 }
 
+export type Loader = {
+    loadingProgress: number;
+    whenLoadUpdate: () => Promise<void> & {cancel: () => void};
+}
+
 const generatetHashFromString = (aggregate: string) => {
     let hash = 0;
     for (let i = 0; i < aggregate.length; i++) {
@@ -310,7 +315,7 @@ export class TestStats {
     }
 }
 
-export class TestCase {
+export class TestCase implements Loader {
     Name: string;
     DisplayName: string;
     TestUID: string;
@@ -344,6 +349,18 @@ export class TestCase {
         );
     }
 
+    @observable
+    historyTarget: number = 0;
+    @action
+    incrementHistoryTarget() {
+        this.historyTarget++;
+    }
+    get loadingProgress() {return this.history}
+    whenLoadUpdate() {
+        const previousVersion = this.history;
+        return when(() => this.version !== previousVersion);
+    }
+
     constructor(name: string, testuid: string, suite: string, sessionName: string, displayName?: string) {
         this.Name = name;
         this.DisplayName = displayName??this.toDislayName(name);
@@ -686,7 +703,7 @@ type SessionReferences = {
     ResultsId: string;
 }
 
-export class TestSessionCollection {
+export class TestSessionCollection implements Loader {
 
     @observable
     version: number = 0;
@@ -711,6 +728,11 @@ export class TestSessionCollection {
     decrementLoading() {
         this.loadingCount--;
     }
+    whenLoadUpdate() {
+        const previousCount = this.loadingCount;
+        return when(() => this.loadingCount !== previousCount)
+    }
+    get loadingProgress(): number { return this._datahandler.items?.size??0}
 
     get size(): number { return this._collection.size; }
     get hash(): string {
@@ -805,9 +827,9 @@ export class TestSessionCollection {
         try {
             if(streamId === this._datahandler.streamId
                 && this._datahandler.items
-                && this._datahandler.activated
                 && maxCount > this._datahandler.items.size) {
                     // Max count resize
+                    this._datahandler.setActive(true);
                     await this._datahandler.queueFetchItems(onFetch, SessionFields.initiate, maxCount);
             } else {
                 // Initial fetch or reload
@@ -934,7 +956,7 @@ export class TestSessionCollection {
                 );
                 if(!test) {
                     // trigger fetch history in background
-                    this.fetchSessionHistory(testcase.SessionName);
+                    this.fetchSessionHistory(testcase.SessionName, uid);
                 }
                 while(!test && this.loadingCount !== 0) {
                     const previousVersion = this.version;
@@ -1045,36 +1067,37 @@ export class TestSessionCollection {
         return testresultsByMeta;
     }
 
-    async fetchSessionHistory(sessionName: string, maxCount: number = 1000) {
+    async fetchSessionHistory(sessionName: string, testuid: string, maxCount: number = 1000) {
         if(this._datahandler.cursor) {
             // for single view mode we need to fetch the whole pool for context
-            const onFetch = (items: (TestDataWrapper|undefined)[]) => this.onInitialHistoryFetch(items, sessionName)
+            const onFetch = (items: (TestDataWrapper|undefined)[]) => this.onInitialHistoryFetch(items, sessionName, testuid)
             await this._datahandler.fetchCursorHistory(onFetch, SessionFields.initiate, maxCount);
             return;
         }
 
-        await this.fetchIncompleteSessionsForHistory(sessionName);
+        await this.fetchIncompleteSessionsForHistory(sessionName, testuid);
         while(this.loadingCount !== 0) {
             // wait for the collection to be fully loaded, in the meanwhile load the missing session history
             const previousVersion = this.version;
             await when(() => this.version !== previousVersion || this.loadingCount === 0);
-            await this.fetchIncompleteSessionsForHistory(sessionName);
+            await this.fetchIncompleteSessionsForHistory(sessionName, testuid);
         }
     }
 
-    private async fetchIncompleteSessionsForHistory(sessionName: string) {
+    private async fetchIncompleteSessionsForHistory(sessionName: string, testuid: string) {
         const missingSessions: string[] = [];
         this.forEachSessionHistoryByChange(sessionName,
             (session) => {
                 if(!session.isAllPropertiesLoaded()) {
                     missingSessions.push(session.Testdata.id);
+                    this._testcases.get(testuid)?.incrementHistoryTarget();
                 }
             }
         );
 
         if(missingSessions.length > 0) {
             await this._datahandler.fetchUpdateItems(
-                missingSessions, (item) => this.onFetchTestCaseHistory(item), SessionFields.remaining
+                missingSessions, (item) => this.onFetchTestCaseHistory(item, testuid), SessionFields.remaining
             )
         }
         return missingSessions.length;
@@ -1097,7 +1120,7 @@ export class TestSessionCollection {
         }
     }
 
-    private onInitialHistoryFetch(items: (TestDataWrapper|undefined)[], sessionName: string) {
+    private onInitialHistoryFetch(items: (TestDataWrapper|undefined)[], sessionName: string, testuid: string) {
         const toFetch: TestDataWrapper[] = [];
         items.forEach((item) => {
             if (item) {
@@ -1107,25 +1130,22 @@ export class TestSessionCollection {
                     item.setDataHandler(data);
                     if(data.Name === sessionName) {
                         toFetch.push(item);
+                        this._testcases.get(testuid)?.incrementHistoryTarget();
                     }
                 }
             }
         });
         if(toFetch.length > 0) {
             this._datahandler.fetchUpdateItems(
-                toFetch.map((item) => item.id), (items) => this.onFetchTestCaseHistory(items), SessionFields.remaining
+                toFetch.map((item) => item.id), (items) => this.onFetchTestCaseHistory(items, testuid), SessionFields.remaining
             )
         }
     }
 
-    private onFetchTestCaseHistory(items: (TestDataWrapper|undefined)[]) {
+    private onFetchTestCaseHistory(items: (TestDataWrapper|undefined)[], testuid: string) {
         items.forEach((item) => {
             if(item) {
-                const session = item.getDataHandler() as TestSessionWrapper;
-                const tests = session.TestSessionInfo?.Tests;
-                if (tests) {
-                    Object.keys(tests).forEach((testuid) => this._testcases.get(testuid)?.updateHistory());
-                }
+                this._testcases.get(testuid)?.updateHistory();
             }
         });
     }
