Horde: Adding Experimental Issue View V2
#rnx
#skipci

diff --git a/Engine/Source/Programs/Horde/HordeDashboard/src/backend/JobDetails.ts b/Engine/Source/Programs/Horde/HordeDashboard/src/backend/JobDetails.ts
--- a/Engine/Source/Programs/Horde/HordeDashboard/src/backend/JobDetails.ts
+++ b/Engine/Source/Programs/Horde/HordeDashboard/src/backend/JobDetails.ts
@@ -96,7 +96,9 @@ export class JobDetails {
 
                 this.stepId = stepId;
                 this.labelIdx = labelIdx;
-                requests.push(this.getIssues());
+                if (!this.suppressIssues) {
+                    requests.push(this.getIssues());
+                }
 
                 this.history = undefined;
 
@@ -881,7 +883,9 @@ export class JobDetails {
 
             if (!this.isLogView) {
                 requests.push(backend.getJobTestData(this.id!));
-                requests.push(this.getIssues());
+                if (!this.suppressIssues) {
+                    requests.push(this.getIssues());
+                }                
                 requests.push(this.queryReports());
             }
 
@@ -1031,6 +1035,8 @@ export class JobDetails {
 
     fatalError?: string;
 
+    suppressIssues:boolean = false;
+
     private static templates = new Map<string, GetTemplateResponse>();
 }
 
