Horde: Add public constructors for job template support classes, so the JSON serializer can construct them even when the settings don't match the names of parameters to the constructor.

diff --git a/Engine/Source/Programs/Horde/Horde.Build/Jobs/ITemplate.cs b/Engine/Source/Programs/Horde/Horde.Build/Jobs/ITemplate.cs
--- a/Engine/Source/Programs/Horde/Horde.Build/Jobs/ITemplate.cs
+++ b/Engine/Source/Programs/Horde/Horde.Build/Jobs/ITemplate.cs
@@ -52,7 +52,7 @@ namespace Horde.Build.Models
 		/// <summary>
 		/// Private constructor for serialization
 		/// </summary>
-		private GroupParameter()
+		public GroupParameter()
 		{
 			Label = null!;
 			Children = null!;
@@ -138,7 +138,7 @@ namespace Horde.Build.Models
 		/// <summary>
 		/// Private constructor for serialization
 		/// </summary>
-		private TextParameter()
+		public TextParameter()
 		{
 			Label = null!;
 			Argument = null!;
@@ -219,6 +219,14 @@ namespace Horde.Build.Models
 		[BsonIgnoreIfDefault, BsonDefaultValue(false)]
 		public bool Default { get; set; }
 
+		/// <summary>
+		/// Default constructor for JSON serializer
+		/// </summary>
+		public ListParameterItem()
+		{
+			Text = String.Empty;
+		}
+
 		/// <summary>
 		/// Constructor
 		/// </summary>
@@ -274,7 +282,7 @@ namespace Horde.Build.Models
 		/// <summary>
 		/// Private constructor for serialization
 		/// </summary>
-		private ListParameter()
+		public ListParameter()
 		{
 			Label = null!;
 			Items = null!;
@@ -366,9 +374,9 @@ namespace Horde.Build.Models
 		/// <summary>
 		/// Private constructor for serialization
 		/// </summary>
-		private BoolParameter()
+		public BoolParameter()
 		{
-			Label = null!;
+			Label = String.Empty;
 		}
 
 		/// <summary>
