Horde.Agent: Handle rooted paths for compute actions

#rnx

diff --git a/Engine/Source/Programs/Horde/Horde.Agent/Execution/ComputeTaskExecutor.cs b/Engine/Source/Programs/Horde/Horde.Agent/Execution/ComputeTaskExecutor.cs
--- a/Engine/Source/Programs/Horde/Horde.Agent/Execution/ComputeTaskExecutor.cs
+++ b/Engine/Source/Programs/Horde/Horde.Agent/Execution/ComputeTaskExecutor.cs
@@ -102,8 +102,9 @@ namespace Horde.Agent.Execution
 
 			using (ManagedProcessGroup processGroup = new ManagedProcessGroup())
 			{
-				string fileName = FileReference.Combine(sandboxDir, task.WorkingDirectory.ToString(), task.Executable.ToString()).FullName;
-				string workingDirectory = DirectoryReference.Combine(sandboxDir, task.WorkingDirectory.ToString()).FullName;
+				// TODO: Rooted paths should be opt-in to prevent access outside of the sandbox. Need to test MacOS compiles.
+				FileReference fileName = Path.IsPathRooted(task.Executable.ToString()) ? new FileReference(task.Executable.ToString()) : FileReference.Combine(sandboxDir, task.WorkingDirectory.ToString(), task.Executable.ToString());
+				DirectoryReference workingDirectory = Path.IsPathRooted(task.WorkingDirectory.ToString()) ? new DirectoryReference(task.WorkingDirectory.ToString()) : DirectoryReference.Combine(sandboxDir, task.WorkingDirectory.ToString());
 
 				Dictionary<string, string> newEnvironment = new Dictionary<string, string>();
 				foreach (System.Collections.DictionaryEntry? entry in Environment.GetEnvironmentVariables())
@@ -120,7 +121,7 @@ namespace Horde.Agent.Execution
 
 				TimeSpan timeout = TimeSpan.FromMinutes(5.0);// Action.Timeout == null ? TimeSpan.FromMinutes(5) : Action.Timeout.ToTimeSpan();
 				DateTimeOffset executionStartTime = DateTimeOffset.UtcNow;
-				using (ManagedProcess process = new ManagedProcess(processGroup, fileName, arguments, workingDirectory, newEnvironment, ProcessPriorityClass.Normal, ManagedProcessFlags.None))
+				using (ManagedProcess process = new ManagedProcess(processGroup, fileName.FullName, arguments, workingDirectory.FullName, newEnvironment, ProcessPriorityClass.Normal, ManagedProcessFlags.None))
 				{
 					(byte[] stdOutData, byte[] stdErrData) = await ReadProcessStreams(process, timeout, cancellationToken);
 
