Horde: Fix additional instances of JSON-serialized classes whose constructor parameter names no longer match property names. Just expose a parameterless constructor (NET 5 change).

diff --git a/Engine/Source/Programs/Horde/Horde.Build/Api/Templates.cs b/Engine/Source/Programs/Horde/Horde.Build/Api/Templates.cs
--- a/Engine/Source/Programs/Horde/Horde.Build/Api/Templates.cs
+++ b/Engine/Source/Programs/Horde/Horde.Build/Api/Templates.cs
@@ -62,7 +62,7 @@ namespace Horde.Build.Api
 		/// <summary>
 		/// Private constructor for serialization
 		/// </summary>
-		private GroupParameterData()
+		public GroupParameterData()
 		{
 			Label = null!;
 			Children = null!;
@@ -135,7 +135,7 @@ namespace Horde.Build.Api
 		/// <summary>
 		/// Private constructor for serialization
 		/// </summary>
-		private TextParameterData()
+		public TextParameterData()
 		{
 			Label = null!;
 			Argument = null!;
@@ -227,7 +227,7 @@ namespace Horde.Build.Api
 		/// <summary>
 		/// Private constructor for serialization
 		/// </summary>
-		private ListParameterItemData()
+		public ListParameterItemData()
 		{
 			Text = null!;
 		}
@@ -288,7 +288,7 @@ namespace Horde.Build.Api
 		/// <summary>
 		/// Private constructor
 		/// </summary>
-		private ListParameterData()
+		public ListParameterData()
 		{
 			Label = null!;
 			Items = null!;
@@ -353,7 +353,7 @@ namespace Horde.Build.Api
 		/// <summary>
 		/// Private constructor for serialization
 		/// </summary>
-		private BoolParameterData()
+		public BoolParameterData()
 		{
 			Label = null!;
 		}
