Horde: Remove boilerplate text from submit failure messages, and put into a preformatted block in Slack messages.

diff --git a/Engine/Source/Programs/Horde/Horde.Build/Jobs/Impl/JobService.cs b/Engine/Source/Programs/Horde/Horde.Build/Jobs/Impl/JobService.cs
--- a/Engine/Source/Programs/Horde/Horde.Build/Jobs/Impl/JobService.cs
+++ b/Engine/Source/Programs/Horde/Horde.Build/Jobs/Impl/JobService.cs
@@ -5,6 +5,7 @@ using System.Collections.Generic;
 using System.Globalization;
 using System.Linq;
 using System.Security.Claims;
+using System.Text.RegularExpressions;
 using System.Threading.Tasks;
 using EpicGames.Core;
 using Horde.Build.Acls;
@@ -1089,6 +1090,11 @@ namespace Horde.Build.Services
 
 					_logger.LogInformation("Attempt to submit {Change} (through {ChangeCopy}): {Message}", job.PreflightChange, clonedPreflightChange, message);
 
+					if (!String.IsNullOrEmpty(message))
+					{
+						message = Regex.Replace(message, @"^Submit validation failed.*\n(?:\s*'[^']*' validation failed:\s*\n)?", "");
+					}
+
 					if (ShouldClonePreflightChange(job.StreamId))
 					{
 						if (change != null && job.ClonedPreflightChange != 0)
diff --git a/Engine/Source/Programs/Horde/Horde.Build/Notifications/Impl/SlackNotificationSink.cs b/Engine/Source/Programs/Horde/Horde.Build/Notifications/Impl/SlackNotificationSink.cs
--- a/Engine/Source/Programs/Horde/Horde.Build/Notifications/Impl/SlackNotificationSink.cs
+++ b/Engine/Source/Programs/Horde/Horde.Build/Notifications/Impl/SlackNotificationSink.cs
@@ -424,7 +424,7 @@ namespace Horde.Build.Notifications.Impl
 						autoSubmitMessage = $"\n\n```{job.AutoSubmitMessage}```";
 					}
 
-					attachment.Blocks.Add(new SectionBlock($"Files in CL {job.PreflightChange} were *not submitted*. Please resolve the following issues and submit manually.{job.AutoSubmitMessage}"));
+					attachment.Blocks.Add(new SectionBlock($"Files in CL *{job.PreflightChange}* were *not submitted*. Please resolve the following issues and submit manually.{autoSubmitMessage}"));
 				}
 			}
 
