Horde: Add an annotation to Slack messages sent by the dev server, to make them easier to distinguish.

diff --git a/Engine/Source/Programs/Horde/Horde.Build/Notifications/Impl/SlackNotificationSink.cs b/Engine/Source/Programs/Horde/Horde.Build/Notifications/Impl/SlackNotificationSink.cs
--- a/Engine/Source/Programs/Horde/Horde.Build/Notifications/Impl/SlackNotificationSink.cs
+++ b/Engine/Source/Programs/Horde/Horde.Build/Notifications/Impl/SlackNotificationSink.cs
@@ -2,6 +2,7 @@
 
 using System;
 using System.Collections.Generic;
+using System.Diagnostics.CodeAnalysis;
 using System.Drawing;
 using System.Globalization;
 using System.Linq;
@@ -21,6 +22,7 @@ using Horde.Build.Services;
 using Horde.Build.Utilities;
 using Horde.Build.Utilities.Slack.BlockKit;
 using HordeCommon;
+using Microsoft.AspNetCore.Hosting;
 using Microsoft.Extensions.Caching.Memory;
 using Microsoft.Extensions.Hosting;
 using Microsoft.Extensions.Logging;
@@ -193,6 +195,7 @@ namespace Horde.Build.Notifications.Impl
 		readonly IIssueService _issueService;
 		readonly IUserCollection _userCollection;
 		readonly StreamService _streamService;
+		readonly IWebHostEnvironment _environment;
 		readonly ServerSettings _settings;
 		readonly IMongoCollection<MessageStateDocument> _messageStates;
 		readonly IMongoCollection<SlackUser> _slackUsers;
@@ -211,13 +214,15 @@ namespace Horde.Build.Notifications.Impl
 		/// <param name="issueService"></param>
 		/// <param name="userCollection">The user collection</param>
 		/// <param name="streamService"></param>
+		/// <param name="environment"></param>
 		/// <param name="settings">The current configuration settings</param>
 		/// <param name="logger">Logging device</param>
-		public SlackNotificationSink(DatabaseService databaseService, IIssueService issueService, IUserCollection userCollection, StreamService streamService, IOptions<ServerSettings> settings, ILogger<SlackNotificationSink> logger)
+		public SlackNotificationSink(DatabaseService databaseService, IIssueService issueService, IUserCollection userCollection, StreamService streamService, IWebHostEnvironment environment, IOptions<ServerSettings> settings, ILogger<SlackNotificationSink> logger)
 		{
 			_issueService = issueService;
 			_userCollection = userCollection;
 			_streamService = streamService;
+			_environment = environment;
 			_settings = settings.Value;
 			_messageStates = databaseService.Database.GetCollection<MessageStateDocument>("Slack");
 			_slackUsers = databaseService.Database.GetCollection<SlackUser>("Slack.UsersV2");
@@ -1168,6 +1173,27 @@ namespace Horde.Build.Notifications.Impl
 			public string? Ts { get; set; }
 		}
 
+		[return: NotNullIfNotNull("message")]
+		private string? AddEnvironmentAnnotation(string? message)
+		{
+			if (_environment.IsProduction())
+			{
+				return message;
+			}
+			else if (_environment.IsDevelopment())
+			{
+				return $"[DEV] {message}";
+			}
+			else if (_environment.IsStaging())
+			{
+				return $"[STAGING] {message}";
+			}
+			else
+			{
+				return $"[{_environment.EnvironmentName.ToUpperInvariant()}] {message}";
+			}
+		}
+
 		private async Task SendMessageAsync(string recipient, string? text = null, BlockBase[]? blocks = null, BlockKitAttachment[]? attachments = null)
 		{
 			if (_allowUsers != null && !_allowUsers.Contains(recipient))
@@ -1178,7 +1204,7 @@ namespace Horde.Build.Notifications.Impl
 
 			BlockKitMessage message = new BlockKitMessage();
 			message.Channel = recipient;
-			message.Text = text;
+			message.Text = AddEnvironmentAnnotation(text);
 			if (blocks != null)
 			{
 				message.Blocks.AddRange(blocks);
@@ -1200,7 +1226,8 @@ namespace Horde.Build.Notifications.Impl
 			}
 
 			BlockKitMessage message = new BlockKitMessage();
-			message.Text = text;
+			message.Text = AddEnvironmentAnnotation(text);
+
 			if (blocks != null)
 			{
 				message.Blocks.AddRange(blocks);
