Horde: Remove the log directory on agents before each run, to prevent re-uploading the same data for unrelated builds.

diff --git a/Engine/Source/Programs/Horde/Horde.Agent/Execution/BuildGraphExecutor.cs b/Engine/Source/Programs/Horde/Horde.Agent/Execution/BuildGraphExecutor.cs
--- a/Engine/Source/Programs/Horde/Horde.Agent/Execution/BuildGraphExecutor.cs
+++ b/Engine/Source/Programs/Horde/Horde.Agent/Execution/BuildGraphExecutor.cs
@@ -767,6 +767,7 @@ namespace Horde.Agent.Execution
 			newEnvironment["IsBuildMachine"] = "1";
 
 			DirectoryReference logDir = DirectoryReference.Combine(workspaceDir, "Engine", "Programs", "AutomationTool", "Saved", "Logs");
+			FileUtils.ForceDeleteDirectoryContents(logDir);
 			newEnvironment["uebp_LogFolder"] = logDir.FullName;
 
 			DirectoryReference telemetryDir = DirectoryReference.Combine(workspaceDir, "Engine", "Programs", "AutomationTool", "Saved", "Telemetry");
