Horde: Print the file/args for every UAT invocation.

diff --git a/Engine/Source/Programs/Horde/Horde.Agent/Execution/BuildGraphExecutor.cs b/Engine/Source/Programs/Horde/Horde.Agent/Execution/BuildGraphExecutor.cs
--- a/Engine/Source/Programs/Horde/Horde.Agent/Execution/BuildGraphExecutor.cs
+++ b/Engine/Source/Programs/Horde/Horde.Agent/Execution/BuildGraphExecutor.cs
@@ -693,7 +693,7 @@ namespace Horde.Agent.Execution
 			}
 		}
 
-		async Task ExecuteCleanupScriptAsync(FileReference cleanupScript, LogParser filter)
+		async Task ExecuteCleanupScriptAsync(FileReference cleanupScript, LogParser filter, ILogger logger)
 		{
 			if (FileReference.Exists(cleanupScript))
 			{
@@ -717,7 +717,7 @@ namespace Horde.Agent.Execution
 					using (CancellationTokenSource cancellationSource = new CancellationTokenSource())
 					{
 						cancellationSource.CancelAfter(TimeSpan.FromSeconds(30.0));
-						await ExecuteProcessAsync(fileName, arguments, null, filter, cancellationSource.Token);
+						await ExecuteProcessAsync(fileName, arguments, null, filter, logger, cancellationSource.Token);
 					}
 					FileUtils.ForceDeleteFile(cleanupScript);
 				}
@@ -728,8 +728,9 @@ namespace Horde.Agent.Execution
 			}
 		}
 
-		async Task<int> ExecuteProcessAsync(string fileName, string arguments, IReadOnlyDictionary<string, string>? newEnvironment, LogParser filter, CancellationToken cancellationToken)
+		async Task<int> ExecuteProcessAsync(string fileName, string arguments, IReadOnlyDictionary<string, string>? newEnvironment, LogParser filter, ILogger logger, CancellationToken cancellationToken)
 		{
+			logger.LogInformation("Executing {file} {arguments}", fileName.QuoteArgument(), arguments);
 			using (ManagedProcessGroup processGroup = new ManagedProcessGroup())
 			{
 				using (ManagedProcess process = new ManagedProcess(processGroup, fileName, arguments, null, newEnvironment, null, ProcessPriorityClass.Normal))
@@ -812,14 +813,14 @@ namespace Horde.Agent.Execution
 			int exitCode;
 			using (LogParser filter = new LogParser(logger, context, ignorePatterns))
 			{
-				await ExecuteCleanupScriptAsync(cleanupScript, filter);
+				await ExecuteCleanupScriptAsync(cleanupScript, filter, logger);
 				try
 				{
-					exitCode = await ExecuteProcessAsync(fileName, arguments, newEnvironment, filter, cancellationToken);
+					exitCode = await ExecuteProcessAsync(fileName, arguments, newEnvironment, filter, logger, cancellationToken);
 				}
 				finally
 				{
-					await ExecuteCleanupScriptAsync(cleanupScript, filter);
+					await ExecuteCleanupScriptAsync(cleanupScript, filter, logger);
 				}
 				filter.Flush();
 			}
