Horde Storage - Avoiding blob index consistency check on legacy namespaces as that coce will not update the refs aspect of the blob index correctly (and it doesnt need to).

diff --git a/Engine/Source/Programs/Horde/HordeStorage/Horde.Storage/Implementation/Blob/BlobIndexConsistencyCheckService.cs b/Engine/Source/Programs/Horde/HordeStorage/Horde.Storage/Implementation/Blob/BlobIndexConsistencyCheckService.cs
--- a/Engine/Source/Programs/Horde/HordeStorage/Horde.Storage/Implementation/Blob/BlobIndexConsistencyCheckService.cs
+++ b/Engine/Source/Programs/Horde/HordeStorage/Horde.Storage/Implementation/Blob/BlobIndexConsistencyCheckService.cs
@@ -5,8 +5,10 @@ using System.Linq;
 using System.Threading;
 using System.Threading.Tasks;
 using Datadog.Trace;
+using EpicGames.Horde.Storage;
 using Horde.Storage.Implementation.Blob;
 using Jupiter;
+using Jupiter.Common;
 using Jupiter.Implementation;
 using Microsoft.Extensions.Options;
 using Serilog;
@@ -21,6 +23,7 @@ namespace Horde.Storage.Implementation
         private readonly ILeaderElection _leaderElection;
         private readonly IBlobIndex _blobIndex;
         private readonly IBlobService _blobService;
+        private readonly INamespacePolicyResolver _namespacePolicyResolver;
         private readonly ILogger _logger = Log.ForContext<BlobIndexConsistencyCheckService>();
 
         public class ConsistencyState
@@ -32,7 +35,7 @@ namespace Horde.Storage.Implementation
             return _settings.CurrentValue.EnableBlobIndexChecks;
         }
 
-        public BlobIndexConsistencyCheckService(IOptionsMonitor<ConsistencyCheckSettings> settings, IOptionsMonitor<JupiterSettings> jupiterSettings, ILeaderElection leaderElection, IBlobIndex blobIndex, IBlobService blobService) :
+        public BlobIndexConsistencyCheckService(IOptionsMonitor<ConsistencyCheckSettings> settings, IOptionsMonitor<JupiterSettings> jupiterSettings, ILeaderElection leaderElection, IBlobIndex blobIndex, IBlobService blobService, INamespacePolicyResolver namespacePolicyResolver) :
             base(serviceName: nameof(BlobStoreConsistencyCheckService), TimeSpan.FromSeconds(settings.CurrentValue.ConsistencyCheckPollFrequencySeconds), new ConsistencyState())
         {
             _settings = settings;
@@ -40,6 +43,7 @@ namespace Horde.Storage.Implementation
             _leaderElection = leaderElection;
             _blobIndex = blobIndex;
             _blobService = blobService;
+            _namespacePolicyResolver = namespacePolicyResolver;
         }
 
         public override async Task<bool> OnPoll(ConsistencyState state, CancellationToken cancellationToken)
@@ -63,6 +67,13 @@ namespace Horde.Storage.Implementation
             return true;
         }
 
+        private bool NamespaceShouldBeCheckedForConsistency(NamespaceId ns)
+        {
+            NamespaceSettings.PerNamespaceSettings policy = _namespacePolicyResolver.GetPoliciesForNs(ns);
+
+            return policy.IsLegacyNamespace.HasValue && !policy.IsLegacyNamespace.Value;
+        }
+
         private async Task RunConsistencyCheck()
         {
             ulong countOfBlobsChecked = 0;
@@ -75,6 +86,11 @@ namespace Horde.Storage.Implementation
                 if (countOfBlobsChecked % 100 == 0)
                     _logger.Information("Consistency check running on blob index, count of blobs processed so far: {CountOfBlobs}", countOfBlobsChecked);
 
+
+                // skip namespace that we shouldn't consistency check, e.g. the legacy namespaces
+                if (NamespaceShouldBeCheckedForConsistency(blobInfo.Namespace))
+                    continue;
+                
                 try
                 {
                     if (!blobInfo.Regions.Any())
