Horde Storage - Fixed case were refs were not tracked properly for blobs of compact binary references that were missing
Also refactored the finalize call into a shared method between PUT and Finalize to ease maintaince.

diff --git a/Engine/Source/Programs/Horde/HordeStorage/Horde.Storage/Implementation/ObjectService.cs b/Engine/Source/Programs/Horde/HordeStorage/Horde.Storage/Implementation/ObjectService.cs
--- a/Engine/Source/Programs/Horde/HordeStorage/Horde.Storage/Implementation/ObjectService.cs
+++ b/Engine/Source/Programs/Horde/HordeStorage/Horde.Storage/Implementation/ObjectService.cs
@@ -102,57 +102,10 @@ namespace Horde.Storage.Implementation
             Task objectStorePut = _referencesStore.Put(ns, bucket, key, blobHash, payload.GetView().ToArray(), isFinalized);
 
             Task<BlobIdentifier> blobStorePut = _blobService.PutObject(ns, payload.GetView().ToArray(), blobHash);
-            ContentId[] missingReferences = Array.Empty<ContentId>();
-            BlobIdentifier[] missingBlobs = Array.Empty<BlobIdentifier>();
-
-            if (hasReferences)
-            {
-                using IScope _ = Tracer.Instance.StartActive("ObjectService.ResolveReferences");
-                try
-                {
-                    IAsyncEnumerable<BlobIdentifier> references = _referenceResolver.ResolveReferences(ns, payload);
-                    // TODO: we no longer use this as a async enumerable, which will slow down the resolve reference a little bit
-                    BlobIdentifier[] referencesArray = await references.ToArrayAsync();
-                    Task addRefsTask = _blobIndex.AddRefToBlobs(ns, bucket, key, referencesArray);
-                    Task<BlobIdentifier[]> checkReferencesTask = _blobService.FilterOutKnownBlobs(ns, referencesArray);
-                    await Task.WhenAll(addRefsTask, checkReferencesTask);
-                    missingBlobs = await checkReferencesTask;
-                }
-                catch (PartialReferenceResolveException e)
-                {
-                    missingReferences = e.UnresolvedReferences.ToArray();
-                }
-                catch (ReferenceIsMissingBlobsException e)
-                {
-                    missingBlobs = e.MissingBlobs.ToArray();
-                }
-            }
-
+            
             await Task.WhenAll(objectStorePut, blobStorePut);
 
-            Task addRefToBlobsTask = _blobIndex.AddRefToBlobs(ns, bucket, key, new [] {blobHash});
-
-            if (missingReferences.Length == 0 && missingBlobs.Length == 0)
-            {
-                try
-                {
-                    await _referencesStore.Finalize(ns, bucket, key, blobHash);
-                }
-                catch (PartialReferenceResolveException e)
-                {
-                    missingReferences = e.UnresolvedReferences.ToArray();
-                }
-                catch (ReferenceIsMissingBlobsException e)
-                {
-                    missingBlobs = e.MissingBlobs.ToArray();
-                }
-
-                await _replicationLog.InsertAddEvent(ns, bucket, key, blobHash);
-            }
-
-            await addRefToBlobsTask;
-
-            return (missingReferences, missingBlobs);
+            return await DoFinalize(ns, bucket, key, blobHash, payload);
         }
 
         public async Task<(ContentId[], BlobIdentifier[])> Finalize(NamespaceId ns, BucketId bucket, IoHashKey key, BlobIdentifier blobHash)
@@ -166,17 +119,27 @@ namespace Horde.Storage.Implementation
             if (!o.BlobIdentifier.Equals(blobHash))
                 throw new ObjectHashMismatchException(ns, bucket, key, blobHash, o.BlobIdentifier);
 
-            bool hasReferences = payload.Any(field => field.IsAttachment());
+            return await DoFinalize(ns, bucket, key, blobHash, payload);
+        }
+
+        
+        private async Task<(ContentId[], BlobIdentifier[])> DoFinalize(NamespaceId ns, BucketId bucket, IoHashKey key, BlobIdentifier blobHash, CbObject payload)
+        {
+            Task addRefToBlobsTask = _blobIndex.AddRefToBlobs(ns, bucket, key, new [] {blobHash});
 
             ContentId[] missingReferences = Array.Empty<ContentId>();
             BlobIdentifier[] missingBlobs = Array.Empty<BlobIdentifier>();
+            bool hasReferences = payload.Any(field => field.IsAttachment());
             if (hasReferences)
             {
                 using IScope _ = Tracer.Instance.StartActive("ObjectService.ResolveReferences");
                 try
                 {
                     IAsyncEnumerable<BlobIdentifier> references = _referenceResolver.ResolveReferences(ns, payload);
-                    missingBlobs = await _blobService.FilterOutKnownBlobs(ns, references);
+                    BlobIdentifier[] referencesArray = await references.ToArrayAsync();
+                    Task addRefsTask = _blobIndex.AddRefToBlobs(ns, bucket, key, referencesArray);
+                    missingBlobs = await _blobService.FilterOutKnownBlobs(ns, referencesArray);
+                    await addRefsTask;
                 }
                 catch (PartialReferenceResolveException e)
                 {
@@ -188,6 +151,8 @@ namespace Horde.Storage.Implementation
                 }
             }
 
+            await addRefToBlobsTask;
+
             if (missingReferences.Length == 0 && missingBlobs.Length == 0)
             {
                 await _referencesStore.Finalize(ns, bucket, key, blobHash);
