Horde Storage - Fixed peer latency list to allow for peers with the same latency to them

diff --git a/Engine/Source/Programs/Horde/HordeStorage/Horde.Storage/Implementation/Blob/IBlobService.cs b/Engine/Source/Programs/Horde/HordeStorage/Horde.Storage/Implementation/Blob/IBlobService.cs
--- a/Engine/Source/Programs/Horde/HordeStorage/Horde.Storage/Implementation/Blob/IBlobService.cs
+++ b/Engine/Source/Programs/Horde/HordeStorage/Horde.Storage/Implementation/Blob/IBlobService.cs
@@ -271,7 +271,7 @@ public class BlobService : IBlobService
             throw new BlobReplicationException(ns, blob, "Blob not found in any region");
 
         _logger.Information("On-demand replicating blob {Blob} in Namespace {Namespace}", blob, ns);
-        SortedList<int, string> possiblePeers = new SortedList<int, string>(_peerStatusService.GetPeersByLatency(blobInfo.Regions.ToList()));
+        List<(int, string)> possiblePeers = new List<(int, string)>(_peerStatusService.GetPeersByLatency(blobInfo.Regions.ToList()));
 
         bool replicated = false;
         foreach ((int latency, string? region) in possiblePeers)
diff --git a/Engine/Source/Programs/Horde/HordeStorage/Horde.Storage/Implementation/PeerStatusService.cs b/Engine/Source/Programs/Horde/HordeStorage/Horde.Storage/Implementation/PeerStatusService.cs
--- a/Engine/Source/Programs/Horde/HordeStorage/Horde.Storage/Implementation/PeerStatusService.cs
+++ b/Engine/Source/Programs/Horde/HordeStorage/Horde.Storage/Implementation/PeerStatusService.cs
@@ -6,6 +6,7 @@ using System.Diagnostics;
 using System.Net.Http;
 using System.Threading;
 using System.Threading.Tasks;
+using EpicGames.Core;
 using Jupiter;
 using Microsoft.Extensions.Options;
 
@@ -27,7 +28,12 @@ namespace Horde.Storage.Implementation
 
         PeerStatus? GetPeerStatus(string regionName);
 
-        SortedList<int, string> GetPeersByLatency(IEnumerable<string> peerNames);
+        /// <summary>
+        /// Returns possible peers sorted by latency to them
+        /// </summary>
+        /// <param name="peerNames"></param>
+        /// <returns></returns>
+        List<(int, string)> GetPeersByLatency(IEnumerable<string> peerNames);
     }
 
     public class PeerStatusService : PollingService<PeerStatusService.PeerStatusServiceState>, IPeerStatusService
@@ -52,9 +58,9 @@ namespace Horde.Storage.Implementation
             return null;
         }
 
-        public SortedList<int, string> GetPeersByLatency(IEnumerable<string> peerNames)
+        public List<(int, string)> GetPeersByLatency(IEnumerable<string> peerNames)
         {
-            SortedList<int, string> sortedPeers = new();
+            List<(int, string)> peers = new();
             foreach (string peerName in peerNames)
             {
                 // skip the local site
@@ -64,10 +70,11 @@ namespace Horde.Storage.Implementation
                 IPeerStatusService.PeerStatus? peerStatus = GetPeerStatus(peerName);
                 if (peerStatus == null)
                     continue;
-                sortedPeers.Add(peerStatus.Latency, peerName);
+                peers.Add((peerStatus.Latency, peerName));
             }
 
-            return sortedPeers;
+            peers.SortBy(pair => pair.Item1);
+            return peers;
         }
 
         public PeerStatusService(IOptionsMonitor<ClusterSettings> clusterSettings, IOptionsMonitor<JupiterSettings> jupiterSettings, IHttpClientFactory clientFactory) : base("PeerStatus", TimeSpan.FromMinutes(15), new PeerStatusServiceState())
