Horde Storage - Fixed potential issue were a stale file info could be used to determine the file size of a buffered object. Also added the size used by the object to some datadog tracking to better be able to see how large the request being handled is.

diff --git a/Engine/Source/Programs/Horde/HordeStorage/Horde.Storage/Implementation/Blob/IBlobService.cs b/Engine/Source/Programs/Horde/HordeStorage/Horde.Storage/Implementation/Blob/IBlobService.cs
--- a/Engine/Source/Programs/Horde/HordeStorage/Horde.Storage/Implementation/Blob/IBlobService.cs
+++ b/Engine/Source/Programs/Horde/HordeStorage/Horde.Storage/Implementation/Blob/IBlobService.cs
@@ -129,6 +129,8 @@ public class BlobService : IBlobService
     {
         using IScope scope = Tracer.Instance.StartActive("put_blob");
         scope.Span.ResourceName = identifier.ToString();
+        scope.Span.SetTag("Content-Length", payload.Length.ToString());
+
         await using Stream hashStream = payload.GetStream();
         BlobIdentifier id = BlobIdentifier.FromContentHash(await VerifyContentMatchesHash(hashStream, identifier));
 
@@ -143,6 +145,8 @@ public class BlobService : IBlobService
     {
         using IScope scope = Tracer.Instance.StartActive("put_blob");
         scope.Span.ResourceName = identifier.ToString();
+        scope.Span.SetTag("Content-Length", payload.Length.ToString());
+
         await using Stream hashStream = new MemoryStream(payload);
         BlobIdentifier id = BlobIdentifier.FromContentHash(await VerifyContentMatchesHash(hashStream, identifier));
 
@@ -156,6 +160,8 @@ public class BlobService : IBlobService
     {
         using IScope scope = Tracer.Instance.StartActive("put_blob");
         scope.Span.ResourceName = identifier.ToString();
+        scope.Span.SetTag("Content-Length", content.Length.ToString());
+
         BlobIdentifier objectStoreIdentifier = await PutObjectToStores(ns, content, identifier);
         await _blobIndex.AddBlobToIndex(ns, identifier);
 
diff --git a/Engine/Source/Programs/Horde/HordeStorage/Jupiter.Common/Implementation/StreamingBodyPayload.cs b/Engine/Source/Programs/Horde/HordeStorage/Jupiter.Common/Implementation/StreamingBodyPayload.cs
--- a/Engine/Source/Programs/Horde/HordeStorage/Jupiter.Common/Implementation/StreamingBodyPayload.cs
+++ b/Engine/Source/Programs/Horde/HordeStorage/Jupiter.Common/Implementation/StreamingBodyPayload.cs
@@ -73,8 +73,10 @@ namespace Jupiter.Common.Implementation
             {
                 await using FileStream fs = payload._tempFile.OpenWrite();
                 await s.CopyToAsync(fs);
-                payload._length = payload._tempFile.Length;
             }
+            
+            payload._tempFile.Refresh();
+            payload._length = payload._tempFile.Length;
 
             return payload;
         }
