Horde Storage -  Attempt to handle "reduce your rate" errors from S3 a bit better by actually returning a seperate reponse code for this (429, same reduce your rate code as returned by S3).
The cooker already retries if this happens, and in practise it doesnt happen in cooks. This does in fact happen for the benchmarks though, which is the intention (to run so many requests that s3 returns this error message) and this change will help clarify its not an error.
Most of the time when this happens AWS SDK will automatically attempt to retry the upload operation, and most of the time it succeeds. So I dont expect to see this very often.

diff --git a/Engine/Source/Programs/Horde/HordeStorage/Horde.Storage/Controllers/StorageController.cs b/Engine/Source/Programs/Horde/HordeStorage/Horde.Storage/Controllers/StorageController.cs
--- a/Engine/Source/Programs/Horde/HordeStorage/Horde.Storage/Controllers/StorageController.cs
+++ b/Engine/Source/Programs/Horde/HordeStorage/Horde.Storage/Controllers/StorageController.cs
@@ -181,12 +181,18 @@ namespace Horde.Storage.Controllers
             _diagnosticContext.Set("Content-Length", Request.ContentLength ?? -1);
             using IBufferedPayload payload = await _bufferedPayloadFactory.CreateFromRequest(Request);
 
-            BlobIdentifier identifier = await _storage.PutObject(ns, payload, id);
-
-            return Ok(new
+            try
             {
-                Identifier = identifier.ToString()
-            });
+                BlobIdentifier identifier = await _storage.PutObject(ns, payload, id);
+                return Ok(new
+                {
+                    Identifier = identifier.ToString()
+                });
+            }
+            catch (ResourceHasToManyRequestsException)
+            {
+                return StatusCode(StatusCodes.Status429TooManyRequests);
+            }
         }
 
         [HttpPost("{ns}")]
diff --git a/Engine/Source/Programs/Horde/HordeStorage/Horde.Storage/Implementation/Blob/AmazonS3Store.cs b/Engine/Source/Programs/Horde/HordeStorage/Horde.Storage/Implementation/Blob/AmazonS3Store.cs
--- a/Engine/Source/Programs/Horde/HordeStorage/Horde.Storage/Implementation/Blob/AmazonS3Store.cs
+++ b/Engine/Source/Programs/Horde/HordeStorage/Horde.Storage/Implementation/Blob/AmazonS3Store.cs
@@ -112,6 +112,9 @@ namespace Horde.Storage.Implementation
                 {
                     return objectName;
                 }
+
+                if (e.StatusCode == HttpStatusCode.TooManyRequests)
+                    throw new ResourceHasToManyRequestsException(e);
                 throw;
             }
 
diff --git a/Engine/Source/Programs/Horde/HordeStorage/Horde.Storage/Implementation/Blob/IBlobStore.cs b/Engine/Source/Programs/Horde/HordeStorage/Horde.Storage/Implementation/Blob/IBlobStore.cs
--- a/Engine/Source/Programs/Horde/HordeStorage/Horde.Storage/Implementation/Blob/IBlobStore.cs
+++ b/Engine/Source/Programs/Horde/HordeStorage/Horde.Storage/Implementation/Blob/IBlobStore.cs
@@ -80,4 +80,12 @@ namespace Horde.Storage.Implementation
 
     }
 
+    public class ResourceHasToManyRequestsException : Exception
+    {
+        public ResourceHasToManyRequestsException(Exception originalException) : base($"To many requests to resource", originalException)
+        {
+        }
+    }
+
+
 }
