Horde Storage - Changed some safe errors into warnings to avoid spamming logs.

diff --git a/Engine/Source/Programs/Horde/HordeStorage/Horde.Storage/Implementation/GC/OrphanBlobCleanupRefs.cs b/Engine/Source/Programs/Horde/HordeStorage/Horde.Storage/Implementation/GC/OrphanBlobCleanupRefs.cs
--- a/Engine/Source/Programs/Horde/HordeStorage/Horde.Storage/Implementation/GC/OrphanBlobCleanupRefs.cs
+++ b/Engine/Source/Programs/Horde/HordeStorage/Horde.Storage/Implementation/GC/OrphanBlobCleanupRefs.cs
@@ -148,9 +148,17 @@ namespace Horde.Storage.Implementation
 
         private bool NamespaceShouldBeCleaned(NamespaceId ns)
         {
-            NamespaceSettings.PerNamespaceSettings policy = _namespacePolicyResolver.GetPoliciesForNs(ns);
+            try
+            {
+                NamespaceSettings.PerNamespaceSettings policy = _namespacePolicyResolver.GetPoliciesForNs(ns);
 
-            return policy.IsLegacyNamespace.HasValue && !policy.IsLegacyNamespace.Value;
+                return policy.IsLegacyNamespace.HasValue && !policy.IsLegacyNamespace.Value;
+            }
+            catch (UnknownNamespaceException)
+            {
+                _logger.Warning("Namespace {Namespace} does not configure any policy, not running cleanup on it.", ns);
+                return false;
+            }
         }
 
         private IAsyncEnumerable<NamespaceId> ListNamespaces()
diff --git a/Engine/Source/Programs/Horde/HordeStorage/Horde.Storage/Implementation/TransactionLog/ReplicationSnapshotService.cs b/Engine/Source/Programs/Horde/HordeStorage/Horde.Storage/Implementation/TransactionLog/ReplicationSnapshotService.cs
--- a/Engine/Source/Programs/Horde/HordeStorage/Horde.Storage/Implementation/TransactionLog/ReplicationSnapshotService.cs
+++ b/Engine/Source/Programs/Horde/HordeStorage/Horde.Storage/Implementation/TransactionLog/ReplicationSnapshotService.cs
@@ -83,7 +83,7 @@ namespace Horde.Storage.Implementation
                 }
                 catch (IncrementalLogNotAvailableException)
                 {
-                    _logger.Error("Unable to generate a snapshot for {Namespace} as there was no incremental state available", ns);
+                    _logger.Warning("Unable to generate a snapshot for {Namespace} as there was no incremental state available", ns);
                 }
 
                 ran = true;
