Horde: Fix case issue with gauntlet device reservations
#skipci
#rnx

diff --git a/Engine/Source/Programs/Horde/Horde.Build/Api/Devices.cs b/Engine/Source/Programs/Horde/Horde.Build/Api/Devices.cs
--- a/Engine/Source/Programs/Horde/Horde.Build/Api/Devices.cs
+++ b/Engine/Source/Programs/Horde/Horde.Build/Api/Devices.cs
@@ -667,7 +667,7 @@ namespace Horde.Build.Api
 		/// The IP or hostname of device
 		/// </summary>
 		[Required]
-		public string IpOrHostName { get; set; } = null!;
+		public string IPOrHostName { get; set; } = null!;
 
 		/// <summary>
 		/// The (legacy) perf spec of the device
diff --git a/Engine/Source/Programs/Horde/Horde.Build/Controllers/DevicesController.cs b/Engine/Source/Programs/Horde/Horde.Build/Controllers/DevicesController.cs
--- a/Engine/Source/Programs/Horde/Horde.Build/Controllers/DevicesController.cs
+++ b/Engine/Source/Programs/Horde/Horde.Build/Controllers/DevicesController.cs
@@ -921,7 +921,7 @@ namespace Horde.Build.Controllers
 
 			response.Name = device.Name;
 			response.Type = platformName;
-			response.IpOrHostName = device.Address ?? "";
+			response.IPOrHostName = device.Address ?? "";
 			response.AvailableStartTime = "00:00:00";
 			response.AvailableEndTime = "00:00:00";
 			response.Enabled = true;
