Horde.Storage: Batch exist from body api for objects and compressed-blobs, to match blobs

#rnx

diff --git a/Engine/Source/Programs/Horde/HordeStorage/Horde.Storage/Controllers/CompressedBlobController.cs b/Engine/Source/Programs/Horde/HordeStorage/Horde.Storage/Controllers/CompressedBlobController.cs
--- a/Engine/Source/Programs/Horde/HordeStorage/Horde.Storage/Controllers/CompressedBlobController.cs
+++ b/Engine/Source/Programs/Horde/HordeStorage/Horde.Storage/Controllers/CompressedBlobController.cs
@@ -168,6 +168,50 @@ namespace Horde.Storage.Controllers
             return Ok(new ExistCheckMultipleContentIdResponse { Needs = partialContentIds.ToArray()});
         }
 
+        [HttpPost("{ns}/exist")]
+        [Authorize("Storage.read")]
+        [ProducesDefaultResponseType]
+        public async Task<IActionResult> ExistsBody(
+            [Required] NamespaceId ns,
+            [FromBody] ContentId[] bodyIds)
+        {
+            AuthorizationResult authorizationResult = await _authorizationService.AuthorizeAsync(User, ns, NamespaceAccessRequirement.Name);
+
+            if (!authorizationResult.Succeeded)
+            {
+                return Forbid();
+            }
+
+            ConcurrentBag<ContentId> partialContentIds = new ConcurrentBag<ContentId>();
+            ConcurrentBag<ContentId> invalidContentIds = new ConcurrentBag<ContentId>();
+
+            IEnumerable<Task> tasks = bodyIds.Select(async blob =>
+            {
+                BlobIdentifier[]? chunks = await _contentIdStore.Resolve(ns, blob, mustBeContentId: false);
+
+                if (chunks == null)
+                {
+                    invalidContentIds.Add(blob);
+                    return;
+                }
+
+                foreach (BlobIdentifier chunk in chunks)
+                {
+                    if (!await _storage.Exists(ns, chunk))
+                    {
+                        partialContentIds.Add(blob);
+                        break;
+                    }
+                }
+            });
+            await Task.WhenAll(tasks);
+
+            if (invalidContentIds.Count != 0)
+                return NotFound(new ValidationProblemDetails { Title = $"Missing content ids {string.Join(" ,", invalidContentIds)}" });
+
+            return Ok(new ExistCheckMultipleContentIdResponse { Needs = partialContentIds.ToArray() });
+        }
+
         private async Task<(BlobContents, string)> GetImpl(NamespaceId ns, ContentId contentId)
         {
             BlobIdentifier[]? chunks = await _contentIdStore.Resolve(ns, contentId, mustBeContentId: false);
diff --git a/Engine/Source/Programs/Horde/HordeStorage/Horde.Storage/Controllers/ObjectController.cs b/Engine/Source/Programs/Horde/HordeStorage/Horde.Storage/Controllers/ObjectController.cs
--- a/Engine/Source/Programs/Horde/HordeStorage/Horde.Storage/Controllers/ObjectController.cs
+++ b/Engine/Source/Programs/Horde/HordeStorage/Horde.Storage/Controllers/ObjectController.cs
@@ -125,6 +125,32 @@ namespace Horde.Storage.Controllers
             return Ok(new HeadMultipleResponse {Needs = missingBlobs.ToArray()});
         }
 
+        [HttpPost("{ns}/exist")]
+        [Authorize("Storage.read")]
+        [ProducesDefaultResponseType]
+        public async Task<IActionResult> ExistsBody(
+            [Required] NamespaceId ns,
+            [FromBody] BlobIdentifier[] bodyIds)
+        {
+            AuthorizationResult authorizationResult = await _authorizationService.AuthorizeAsync(User, ns, NamespaceAccessRequirement.Name);
+
+            if (!authorizationResult.Succeeded)
+            {
+                return Forbid();
+            }
+
+            ConcurrentBag<BlobIdentifier> missingBlobs = new ConcurrentBag<BlobIdentifier>();
+
+            IEnumerable<Task> tasks = bodyIds.Select(async blob =>
+            {
+                if (!await _storage.Exists(ns, blob))
+                    missingBlobs.Add(blob);
+            });
+            await Task.WhenAll(tasks);
+
+            return Ok(new HeadMultipleResponse { Needs = missingBlobs.ToArray() });
+        }
+
         [HttpPut("{ns}/{id}")]
         [Authorize("Storage.write")]
         [RequiredContentType(CustomMediaTypeNames.UnrealCompactBinary)]
