Horde: Fix race condition in lease assignment caused by task sources which need to execute async operations before being able to determine if a lease is available.

The AssignLeaseAsync() method now returns a Task<Task<AgentLease>>. The outer task is awaited on each task source in priority order, and all inner tasks are awaited in parallell. Utility methods in TaskSourceBase (Skip/DrainAsync/Lease) exist to simplify creation of task objects to return.

diff --git a/Engine/Source/Programs/Horde/Horde.Build/Agents/IAgent.cs b/Engine/Source/Programs/Horde/Horde.Build/Agents/IAgent.cs
--- a/Engine/Source/Programs/Horde/Horde.Build/Agents/IAgent.cs
+++ b/Engine/Source/Programs/Horde/Horde.Build/Agents/IAgent.cs
@@ -178,11 +178,6 @@ namespace Horde.Build.Models
 	/// </summary>
 	public class AgentLease
 	{
-		/// <summary>
-		/// Lease returned when a task source wants to run, but needs to wait for other tasks to finish first.
-		/// </summary>
-		public static AgentLease Drain { get; } = new AgentLease(LeaseId.Empty, String.Empty, null, null, null, LeaseState.Pending, null, true, null);
-
 		/// <summary>
 		/// Name of this lease
 		/// </summary>
diff --git a/Engine/Source/Programs/Horde/Horde.Build/Agents/ITaskSource.cs b/Engine/Source/Programs/Horde/Horde.Build/Agents/ITaskSource.cs
--- a/Engine/Source/Programs/Horde/Horde.Build/Agents/ITaskSource.cs
+++ b/Engine/Source/Programs/Horde/Horde.Build/Agents/ITaskSource.cs
@@ -8,6 +8,7 @@ using System.Reflection;
 using System.Text;
 using System.Threading;
 using System.Threading.Tasks;
+using EpicGames.Core;
 using Google.Protobuf;
 using Google.Protobuf.Reflection;
 using Google.Protobuf.WellKnownTypes;
@@ -63,12 +64,22 @@ namespace Horde.Build.Tasks
 		TaskSourceFlags Flags { get; }
 
 		/// <summary>
-		/// Assigns a lease or waits for one to be available
+		/// Assigns a lease or waits for one to be available.
+		/// 
+		/// This method returns a wrapped task object. Issuing of tasks is done in two phases; all task sources are queried sequentially for a lease with immediate 
+		/// availability, and if one isn't found we wait for leases to become available in parallel. Task sources are specificially ordered to allow one to take precedence
+		/// over another.
+		/// 
+		/// Performing the lease allocation in two steps allows for the following patterns:
+		/// - A lease cannot be assigned immediately (due to dependencies on some limited external resource, or due to other leases already running on an agent), 
+		///   but we should not assign any other leases to the agent either. In this situation, the method can just block immediately until the lease criteria are 
+		///   satisfied.
+		/// - The determination of a lease's availability requires some async setup, which can be done before the secondary task returns.
 		/// </summary>
 		/// <param name="agent">The agent to assign a lease to</param>
 		/// <param name="cancellationToken">Cancellation token for the wait</param>
-		/// <returns>New lease object</returns>
-		Task<AgentLease?> AssignLeaseAsync(IAgent agent, CancellationToken cancellationToken);
+		/// <returns>Task returning a new lease object.</returns>
+		Task<Task<AgentLease>> AssignLeaseAsync(IAgent agent, CancellationToken cancellationToken);
 
 		/// <summary>
 		/// Cancel a lease that was previously assigned to an agent, allowing it to be assigned out again
@@ -162,7 +173,7 @@ namespace Horde.Build.Tasks
 		public MessageDescriptor Descriptor => s_message.Descriptor;
 
 		/// <inheritdoc/>
-		public abstract Task<AgentLease?> AssignLeaseAsync(IAgent agent, CancellationToken cancellationToken);
+		public abstract Task<Task<AgentLease>> AssignLeaseAsync(IAgent agent, CancellationToken cancellationToken);
 
 		/// <inheritdoc/>
 		public Task CancelLeaseAsync(IAgent agent, LeaseId leaseId, Any payload) => CancelLeaseAsync(agent, leaseId, payload.Unpack<TMessage>());
@@ -213,5 +224,36 @@ namespace Horde.Build.Tasks
 				}
 			}
 		}
-	}
+
+        /// <summary>
+        /// Creates a lease task which will wait until the given cancellation token is signalled.
+        /// </summary>
+        /// <param name="token">The cancellation token</param>
+        /// <returns>Lease task</returns>
+        protected static Task<AgentLease> Skip(CancellationToken token)
+        {
+            return token.AsTask<AgentLease>();
+        }
+
+        /// <summary>
+        /// Waits until the cancellation token is signalled, then return an cancelled lease task.
+        /// </summary>
+        /// <param name="token">The cancellation token</param>
+        /// <returns>Lease task</returns>
+        protected static async Task<Task<AgentLease>> DrainAsync(CancellationToken token)
+        {
+            await token.AsTask();
+            return Task.FromCanceled<AgentLease>(token);
+        }
+
+        /// <summary>
+        /// Creates a lease task from a given lease
+        /// </summary>
+        /// <param name="lease">Lease to create the task from</param>
+        /// <returns></returns>
+        protected static Task<AgentLease> Lease(AgentLease lease)
+        {
+            return Task.FromResult<AgentLease>(lease);
+        }
+    }
 }
diff --git a/Engine/Source/Programs/Horde/Horde.Build/Agents/Impl/AgentService.cs b/Engine/Source/Programs/Horde/Horde.Build/Agents/Impl/AgentService.cs
--- a/Engine/Source/Programs/Horde/Horde.Build/Agents/Impl/AgentService.cs
+++ b/Engine/Source/Programs/Horde/Horde.Build/Agents/Impl/AgentService.cs
@@ -334,24 +334,6 @@ namespace Horde.Build.Services
 			return agent;
 		}
 
-		async Task<(ITaskSource, AgentLease)?> GuardedWaitForLeaseAsync(ITaskSource source, IAgent agent, CancellationToken cancellationToken)
-		{
-			try
-			{
-				AgentLease? lease = await source.AssignLeaseAsync(agent, cancellationToken);
-				return (lease != null) ? (source, lease) : null;
-			}
-			catch (TaskCanceledException)
-			{
-				return null;
-			}
-			catch (Exception ex)
-			{
-				_logger.LogError(ex, "Exception while trying to assign lease"); 
-				return null;
-			}
-		}
-
 		/// <summary>
 		/// Determines whether a task source can currently issue tasks
 		/// </summary>
@@ -407,12 +389,13 @@ namespace Horde.Build.Services
 					}
 
 					// Create all the tasks to wait for
-					List<Task<(ITaskSource, AgentLease)?>> tasks = new List<Task<(ITaskSource, AgentLease)?>>();
+					List<Task<(ITaskSource source, AgentLease lease)>> tasks = new List<Task<(ITaskSource, AgentLease)>>();
 					foreach (ITaskSource taskSource in _taskSources)
 					{
-						if (CanUseTaskSource(agent, taskSource))
+						if (CanUseTaskSource(agent, taskSource) && !cancellationSource.IsCancellationRequested)
 						{
-							tasks.Add(GuardedWaitForLeaseAsync(taskSource, agent, cancellationSource.Token));
+							Task<(ITaskSource, AgentLease)> task = await GuardedAssignLeaseAsync(taskSource, agent, cancellationSource);
+							tasks.Add(task);
 						}
 					}
 
@@ -423,35 +406,23 @@ namespace Horde.Build.Services
 						break;
 					}
 
-					// Wait for a lease to be available
-					while (tasks.Count > 0)
-					{
-						await Task.WhenAny(tasks);
+					// Wait for all the tasks to complete. Once the first task completes it will set the cancellation source, triggering the 
+					// others to terminate.
+					await Task.WhenAll(tasks);
 
-						for (int idx = 0; idx < tasks.Count; idx++)
+					// Find the first result
+					foreach (Task<(ITaskSource, AgentLease)> task in tasks)
+					{
+						(ITaskSource source, AgentLease lease) taskResult;
+						if (task.TryGetResult(out taskResult))
 						{
-							(ITaskSource, AgentLease)? taskResult;
-							if (!tasks[idx].TryGetResult(out taskResult))
-							{
-								continue;
-							}
-
-							tasks.RemoveAt(idx--);
-
-							if (!taskResult.HasValue)
-							{
-								continue;
-							}
-
 							if (result == null)
 							{
 								result = taskResult;
-								cancellationSource.Cancel();
 							}
-							else if (taskResult.Value.Item2 != AgentLease.Drain)
+							else
 							{
-								(ITaskSource taskSource, AgentLease taskLease) = taskResult.Value;
-								await taskSource.CancelLeaseAsync(agent, taskLease.Id, Any.Parser.ParseFrom(taskLease.Payload));
+								await taskResult.source.CancelLeaseAsync(agent, taskResult.lease.Id, Any.Parser.ParseFrom(taskResult.lease.Payload));
 							}
 						}
 					}
@@ -476,11 +447,6 @@ namespace Horde.Build.Services
 
 				// Get the resulting lease
 				(ITaskSource source, AgentLease lease) = result.Value;
-				if (lease == AgentLease.Drain)
-				{
-					await AsyncUtils.DelayNoThrow(maxWaitTime, cancellationToken);
-					break;
-				}
 
 				// Add the new lease to the agent
 				IAgent? newAgent = await Agents.TryAddLeaseAsync(agent, lease);
@@ -497,6 +463,31 @@ namespace Horde.Build.Services
 			return agent;
 		}
 
+		async Task<Task<(ITaskSource, AgentLease)>> GuardedAssignLeaseAsync(ITaskSource source, IAgent agent, CancellationTokenSource cancellationSource)
+		{
+			CancellationToken cancellationToken = cancellationSource.Token;
+			try
+			{
+				Task<AgentLease> task = await source.AssignLeaseAsync(agent, cancellationToken);
+				return task.ContinueWith(x => WrapAssignedLease(source, x.Result, cancellationSource), CancellationToken.None, TaskContinuationOptions.OnlyOnRanToCompletion, TaskScheduler.Default);
+			}
+			catch (TaskCanceledException) when (cancellationToken.IsCancellationRequested)
+			{
+				return Task.FromCanceled<(ITaskSource, AgentLease)>(cancellationToken);
+			}
+			catch (Exception ex)
+			{
+				_logger.LogError(ex, "Exception while trying to assign lease");
+				return Task.FromException<(ITaskSource, AgentLease)>(ex);
+			}
+		}
+
+		static (ITaskSource, AgentLease) WrapAssignedLease(ITaskSource source, AgentLease lease, CancellationTokenSource cancellationSource)
+		{
+			cancellationSource.Cancel();
+			return (source, lease);
+		}
+
 		/// <summary>
 		/// Cancels the specified agent lease
 		/// </summary>
diff --git a/Engine/Source/Programs/Horde/Horde.Build/Agents/Impl/Tasks/ConformTaskSource.cs b/Engine/Source/Programs/Horde/Horde.Build/Agents/Impl/Tasks/ConformTaskSource.cs
--- a/Engine/Source/Programs/Horde/Horde.Build/Agents/Impl/Tasks/ConformTaskSource.cs
+++ b/Engine/Source/Programs/Horde/Horde.Build/Agents/Impl/Tasks/ConformTaskSource.cs
@@ -149,12 +149,12 @@ namespace Horde.Build.Tasks.Impl
 		}
 
 		/// <inheritdoc/>
-		public override async Task<AgentLease?> AssignLeaseAsync(IAgent agent, CancellationToken cancellationToken)
+		public override async Task<Task<AgentLease>> AssignLeaseAsync(IAgent agent, CancellationToken cancellationToken)
 		{
 			DateTime utcNow = DateTime.UtcNow;
 			if (!await IsConformPendingAsync(agent, utcNow))
 			{
-				return null;
+				return Skip(cancellationToken);
 			}
 
 			if (agent.Leases.Count == 0)
@@ -171,17 +171,19 @@ namespace Horde.Build.Tasks.Impl
 
 						byte[] payload = Any.Pack(task).ToByteArray();
 
-						return new AgentLease(leaseId, "Updating workspaces", null, null, log.Id, LeaseState.Pending, null, true, payload);
+						return Lease(new AgentLease(leaseId, "Updating workspaces", null, null, log.Id, LeaseState.Pending, null, true, payload));
 					}
 				}
 			}
 
 			if (agent.RequestConform || agent.RequestFullConform)
 			{
-				return AgentLease.Drain;
+				return await DrainAsync(cancellationToken);
+			}
+			else
+			{
+				return Skip(cancellationToken);
 			}
-
-			return null;
 		}
 
 		/// <inheritdoc/>
diff --git a/Engine/Source/Programs/Horde/Horde.Build/Agents/Impl/Tasks/RestartTaskSource.cs b/Engine/Source/Programs/Horde/Horde.Build/Agents/Impl/Tasks/RestartTaskSource.cs
--- a/Engine/Source/Programs/Horde/Horde.Build/Agents/Impl/Tasks/RestartTaskSource.cs
+++ b/Engine/Source/Programs/Horde/Horde.Build/Agents/Impl/Tasks/RestartTaskSource.cs
@@ -1,7 +1,10 @@
 // Copyright Epic Games, Inc. All Rights Reserved.
 
+using System.Collections.Generic;
+using System.Runtime.CompilerServices;
 using System.Threading;
 using System.Threading.Tasks;
+using EpicGames.Core;
 using Google.Protobuf;
 using Google.Protobuf.WellKnownTypes;
 using Horde.Build.Api;
@@ -32,15 +35,15 @@ namespace Horde.Build.Tasks.Impl
 			OnLeaseStartedProperties.Add(nameof(RestartTask.LogId), x => new LogId(x.LogId));
 		}
 
-		public override async Task<AgentLease?> AssignLeaseAsync(IAgent agent, CancellationToken cancellationToken)
+		public override async Task<Task<AgentLease>> AssignLeaseAsync(IAgent agent, CancellationToken cancellationToken)
 		{
 			if (!agent.RequestRestart)
 			{
-				return null;
+				return Skip(cancellationToken);
 			}
 			if (agent.Leases.Count > 0)
 			{
-				return AgentLease.Drain;
+				return await DrainAsync(cancellationToken);
 			}
 
 			ILogFile log = await _logService.CreateLogFileAsync(JobId.Empty, agent.SessionId, LogType.Json);
@@ -50,7 +53,7 @@ namespace Horde.Build.Tasks.Impl
 
 			byte[] payload = Any.Pack(task).ToByteArray();
 
-			return new AgentLease(LeaseId.GenerateNewId(), "Restart", null, null, log.Id, LeaseState.Pending, null, true, payload);
+			return Lease(new AgentLease(LeaseId.GenerateNewId(), "Restart", null, null, log.Id, LeaseState.Pending, null, true, payload));
 		}
 	}
 }
diff --git a/Engine/Source/Programs/Horde/Horde.Build/Agents/Impl/Tasks/ShutdownTaskSource.cs b/Engine/Source/Programs/Horde/Horde.Build/Agents/Impl/Tasks/ShutdownTaskSource.cs
--- a/Engine/Source/Programs/Horde/Horde.Build/Agents/Impl/Tasks/ShutdownTaskSource.cs
+++ b/Engine/Source/Programs/Horde/Horde.Build/Agents/Impl/Tasks/ShutdownTaskSource.cs
@@ -31,15 +31,15 @@ namespace Horde.Build.Tasks.Impl
 			OnLeaseStartedProperties.Add(nameof(ShutdownTask.LogId), x => new LogId(x.LogId));
 		}
 
-		public override async Task<AgentLease?> AssignLeaseAsync(IAgent agent, CancellationToken cancellationToken)
+		public override async Task<Task<AgentLease>> AssignLeaseAsync(IAgent agent, CancellationToken cancellationToken)
 		{
 			if (!agent.RequestShutdown)
 			{
-				return null;
+				return Skip(cancellationToken);
 			}
 			if (agent.Leases.Count > 0)
 			{
-				return AgentLease.Drain;
+				return await DrainAsync(cancellationToken);
 			}
 
 			ILogFile log = await _logService.CreateLogFileAsync(JobId.Empty, agent.SessionId, LogType.Json);
@@ -49,7 +49,7 @@ namespace Horde.Build.Tasks.Impl
 
 			byte[] payload = Any.Pack(task).ToByteArray();
 
-			return new AgentLease(LeaseId.GenerateNewId(), "Shutdown", null, null, log.Id, LeaseState.Pending, null, true, payload);
+			return Lease(new AgentLease(LeaseId.GenerateNewId(), "Shutdown", null, null, log.Id, LeaseState.Pending, null, true, payload));
 		}
 	}
 }
diff --git a/Engine/Source/Programs/Horde/Horde.Build/Agents/Impl/Tasks/UpgradeTaskSource.cs b/Engine/Source/Programs/Horde/Horde.Build/Agents/Impl/Tasks/UpgradeTaskSource.cs
--- a/Engine/Source/Programs/Horde/Horde.Build/Agents/Impl/Tasks/UpgradeTaskSource.cs
+++ b/Engine/Source/Programs/Horde/Horde.Build/Agents/Impl/Tasks/UpgradeTaskSource.cs
@@ -1,6 +1,8 @@
 // Copyright Epic Games, Inc. All Rights Reserved.
 
 using System;
+using System.Collections.Generic;
+using System.Runtime.CompilerServices;
 using System.Threading;
 using System.Threading.Tasks;
 using Google.Protobuf;
@@ -38,38 +40,38 @@ namespace Horde.Build.Tasks.Impl
 			OnLeaseStartedProperties.Add(nameof(UpgradeTask.LogId), x => new LogId(x.LogId));
 		}
 
-		public override async Task<AgentLease?> AssignLeaseAsync(IAgent agent, CancellationToken cancellationToken)
+		public override async Task<Task<AgentLease>> AssignLeaseAsync(IAgent agent, CancellationToken cancellationToken)
 		{
 			string? requiredVersion = await GetRequiredSoftwareVersion(agent);
-			if (requiredVersion != null && agent.Version != requiredVersion)
+			if (requiredVersion == null || agent.Version == requiredVersion)
 			{
-				AgentLease? lease = null;
-				if (agent.Leases.Count == 0 && (agent.LastUpgradeTime == null || agent.LastUpgradeTime.Value + TimeSpan.FromMinutes(5.0) < _clock.UtcNow || agent.LastUpgradeVersion != requiredVersion.ToString()))
-				{
-					ILogFile logFile = await _logService.CreateLogFileAsync(JobId.Empty, agent.SessionId, LogType.Json);
+				return Skip(cancellationToken);
+			}
+			if (agent.Leases.Count > 0 || !(agent.LastUpgradeTime == null || agent.LastUpgradeTime.Value + TimeSpan.FromMinutes(5.0) < _clock.UtcNow || agent.LastUpgradeVersion != requiredVersion.ToString()))
+			{
+				return await DrainAsync(cancellationToken);
+			}
 
-					UpgradeTask task = new UpgradeTask();
-					task.SoftwareId = requiredVersion.ToString();
-					task.LogId = logFile.Id.ToString();
+			ILogFile logFile = await _logService.CreateLogFileAsync(JobId.Empty, agent.SessionId, LogType.Json);
 
-					byte[] payload;
-					if (agent.Version == "5.0.0-17425336" || agent.Version == "5.0.0-17448746")
-					{
-						Any any = new Any();
-						any.TypeUrl = "type.googleapis.com/Horde.UpgradeTask";
-						any.Value = task.ToByteString();
-						payload = any.ToByteArray();
-					}
-					else
-					{
-						payload = Any.Pack(task).ToByteArray();
-					}
+			UpgradeTask task = new UpgradeTask();
+			task.SoftwareId = requiredVersion.ToString();
+			task.LogId = logFile.Id.ToString();
 
-					lease = new AgentLease(LeaseId.GenerateNewId(), $"Upgrade to {requiredVersion}", null, null, logFile.Id, LeaseState.Pending, null, true, payload);
-				}
-				return lease;
+			byte[] payload;
+			if (agent.Version == "5.0.0-17425336" || agent.Version == "5.0.0-17448746")
+			{
+				Any any = new Any();
+				any.TypeUrl = "type.googleapis.com/Horde.UpgradeTask";
+				any.Value = task.ToByteString();
+				payload = any.ToByteArray();
 			}
-			return null;
+			else
+			{
+				payload = Any.Pack(task).ToByteArray();
+			}
+
+			return Lease(new AgentLease(LeaseId.GenerateNewId(), $"Upgrade to {requiredVersion}", null, null, logFile.Id, LeaseState.Pending, null, true, payload));
 		}
 
 		/// <summary>
diff --git a/Engine/Source/Programs/Horde/Horde.Build/Compute/Impl/ComputeService.cs b/Engine/Source/Programs/Horde/Horde.Build/Compute/Impl/ComputeService.cs
--- a/Engine/Source/Programs/Horde/Horde.Build/Compute/Impl/ComputeService.cs
+++ b/Engine/Source/Programs/Horde/Horde.Build/Compute/Impl/ComputeService.cs
@@ -272,57 +272,65 @@ namespace Horde.Build.Compute.Impl
 		}
 
 		/// <inheritdoc/>
-		public override async Task<AgentLease?> AssignLeaseAsync(IAgent agent, CancellationToken cancellationToken)
+		public override async Task<Task<AgentLease>> AssignLeaseAsync(IAgent agent, CancellationToken cancellationToken)
 		{
-			// If the agent is disabled, just block until we cancel
-			if (!agent.Enabled)
-			{
-				using CancellationTask task = new CancellationTask(cancellationToken);
-				await task.Task;
-				return null;
-			}
+			Task<(QueueKey, ComputeTaskInfo)> task = await _taskScheduler.DequeueAsync(queueKey => CheckRequirements(agent, queueKey), cancellationToken);
+			return WaitForLeaseAsync(agent, task, cancellationToken);
+		}
 
-			// Find a task to execute
-			(QueueKey, ComputeTaskInfo)? entry = await _taskScheduler.DequeueAsync(queueKey => CheckRequirements(agent, queueKey), cancellationToken);
-			if (entry != null)
+		private async Task<AgentLease> WaitForLeaseAsync(IAgent agent, Task<(QueueKey, ComputeTaskInfo)> task, CancellationToken cancellationToken)
+		{
+			for (; ; )
 			{
-				(QueueKey queueKey, ComputeTaskInfo taskInfo) = entry.Value;
+				(QueueKey, ComputeTaskInfo) entry = await task;
 
-				ComputeClusterConfig? cluster = await GetClusterAsync(taskInfo.ClusterId);
-				if (cluster == null)
+				AgentLease? lease = await CreateLeaseForEntryAsync(agent, await task);
+				if (lease != null)
 				{
-					_logger.LogWarning("Invalid cluster '{ClusterId}'; failing task {TaskRefId}", taskInfo.ClusterId, taskInfo.TaskRefId);
-					ComputeTaskStatus status = new ComputeTaskStatus(taskInfo.TaskRefId, ComputeTaskState.Complete, agent.Id, null) { Detail = $"Invalid cluster '{taskInfo.ClusterId}'" };
-					await PostStatusMessageAsync(taskInfo, status);
-					return null;
+					return lease;
 				}
 
-				Requirements? requirements = await GetCachedRequirementsAsync(queueKey);
-				if (requirements == null)
-				{
-					_logger.LogWarning("Unable to fetch requirements {RequirementsHash}", queueKey);
-					ComputeTaskStatus status = new ComputeTaskStatus(taskInfo.TaskRefId, ComputeTaskState.Complete, agent.Id, null) { Detail = $"Unable to retrieve requirements '{queueKey}'" };
-					await PostStatusMessageAsync(taskInfo, status);
-					return null;
-				}
+				task = await _taskScheduler.DequeueAsync(queueKey => CheckRequirements(agent, queueKey), cancellationToken);
+			}
+		}
+
+		private async Task<AgentLease?> CreateLeaseForEntryAsync(IAgent agent, (QueueKey, ComputeTaskInfo) entry)
+		{
+			(QueueKey queueKey, ComputeTaskInfo taskInfo) = entry;
+
+			ComputeClusterConfig? cluster = await GetClusterAsync(taskInfo.ClusterId);
+			if (cluster == null)
+			{
+				_logger.LogWarning("Invalid cluster '{ClusterId}'; failing task {TaskRefId}", taskInfo.ClusterId, taskInfo.TaskRefId);
+				ComputeTaskStatus status = new ComputeTaskStatus(taskInfo.TaskRefId, ComputeTaskState.Complete, agent.Id, null) { Detail = $"Invalid cluster '{taskInfo.ClusterId}'" };
+				await PostStatusMessageAsync(taskInfo, status);
+				return null;
+			}
 
-				ComputeTaskMessage computeTask = new ComputeTaskMessage();
-				computeTask.ClusterId = taskInfo.ClusterId.ToString();
-				computeTask.ChannelId = taskInfo.ChannelId.ToString();
-				computeTask.NamespaceId = cluster.NamespaceId.ToString();
-				computeTask.InputBucketId = cluster.RequestBucketId.ToString();
-				computeTask.OutputBucketId = cluster.ResponseBucketId.ToString();
-				computeTask.RequirementsHash = queueKey.RequirementsHash;
-				computeTask.TaskRefId = taskInfo.TaskRefId;
-
-				string leaseName = $"Remote action ({taskInfo.TaskRefId})";
-				byte[] payload = Any.Pack(computeTask).ToByteArray();
-
-				AgentLease lease = new AgentLease(LeaseId.GenerateNewId(), leaseName, null, null, null, LeaseState.Pending, requirements.Resources, requirements.Exclusive, payload);
-				_logger.LogDebug("Created lease {LeaseId} for channel {ChannelId} task {TaskHash} req {RequirementsHash}", lease.Id, computeTask.ChannelId, computeTask.TaskRefId, computeTask.RequirementsHash);
-				return lease;
+			Requirements? requirements = await GetCachedRequirementsAsync(queueKey);
+			if (requirements == null)
+			{
+				_logger.LogWarning("Unable to fetch requirements {RequirementsHash}", queueKey);
+				ComputeTaskStatus status = new ComputeTaskStatus(taskInfo.TaskRefId, ComputeTaskState.Complete, agent.Id, null) { Detail = $"Unable to retrieve requirements '{queueKey}'" };
+				await PostStatusMessageAsync(taskInfo, status);
+				return null;
 			}
-			return null;
+
+			ComputeTaskMessage computeTask = new ComputeTaskMessage();
+			computeTask.ClusterId = taskInfo.ClusterId.ToString();
+			computeTask.ChannelId = taskInfo.ChannelId.ToString();
+			computeTask.NamespaceId = cluster.NamespaceId.ToString();
+			computeTask.InputBucketId = cluster.RequestBucketId.ToString();
+			computeTask.OutputBucketId = cluster.ResponseBucketId.ToString();
+			computeTask.RequirementsHash = queueKey.RequirementsHash;
+			computeTask.TaskRefId = taskInfo.TaskRefId;
+
+			string leaseName = $"Remote action ({taskInfo.TaskRefId})";
+			byte[] payload = Any.Pack(computeTask).ToByteArray();
+
+			AgentLease lease = new AgentLease(LeaseId.GenerateNewId(), leaseName, null, null, null, LeaseState.Pending, requirements.Resources, requirements.Exclusive, payload);
+			_logger.LogDebug("Created lease {LeaseId} for channel {ChannelId} task {TaskHash} req {RequirementsHash}", lease.Id, computeTask.ChannelId, computeTask.TaskRefId, computeTask.RequirementsHash);
+			return lease;
 		}
 
 		/// <inheritdoc/>
diff --git a/Engine/Source/Programs/Horde/Horde.Build/Jobs/Impl/JobTaskSource.cs b/Engine/Source/Programs/Horde/Horde.Build/Jobs/Impl/JobTaskSource.cs
--- a/Engine/Source/Programs/Horde/Horde.Build/Jobs/Impl/JobTaskSource.cs
+++ b/Engine/Source/Programs/Horde/Horde.Build/Jobs/Impl/JobTaskSource.cs
@@ -162,12 +162,12 @@ namespace Horde.Build.Tasks.Impl
 			/// <summary>
 			/// Task to wait for a lease to be assigned
 			/// </summary>
-			public Task<AgentLease?> Task => LeaseSource.Task;
+			public Task<AgentLease> Task => LeaseSource.Task;
 
 			/// <summary>
 			/// Completion source for the waiting agent. If a new queue item becomes available, the result will be passed through 
 			/// </summary>
-			public TaskCompletionSource<AgentLease?> LeaseSource { get; } = new TaskCompletionSource<AgentLease?>();
+			public TaskCompletionSource<AgentLease> LeaseSource { get; } = new TaskCompletionSource<AgentLease>();
 
 			/// <summary>
 			/// Constructor
@@ -295,7 +295,7 @@ namespace Horde.Build.Tasks.Impl
 			}
 			if(waiter != null)
 			{
-				waiter.LeaseSource.TrySetResult(null);
+				waiter.LeaseSource.TrySetCanceled();
 			}
 		}
 
@@ -661,7 +661,7 @@ namespace Horde.Build.Tasks.Impl
 		}
 
 		/// <inheritdoc/>
-		public override async Task<AgentLease?> AssignLeaseAsync(IAgent agent, CancellationToken cancellationToken)
+		public override async Task<Task<AgentLease>> AssignLeaseAsync(IAgent agent, CancellationToken cancellationToken)
 		{
 			QueueWaiter waiter = new QueueWaiter(agent);
 			lock (_lockObject)
@@ -669,13 +669,17 @@ namespace Horde.Build.Tasks.Impl
 				AssignAnyQueueItemToWaiter(waiter);
 				if (waiter.LeaseSource.Task.TryGetResult(out AgentLease? result))
 				{
-					return result;
+					if (result == null)
+					{
+						return Skip(cancellationToken);
+					}
+					return Lease(result);
 				}
 				_waiters.Add(waiter);
 			}
 
 			AgentLease? lease;
-			using (cancellationToken.Register(() => waiter.LeaseSource.TrySetResult(null)))
+			using (cancellationToken.Register(() => waiter.LeaseSource.TrySetCanceled()))
 			{
 				lease = await waiter.Task;
 			}
@@ -685,7 +689,7 @@ namespace Horde.Build.Tasks.Impl
 				_waiters.Remove(waiter);
 			}
 
-			return lease;
+			return Lease(lease);
 		}
 
 		/// <inheritdoc/>
diff --git a/Engine/Source/Programs/Shared/EpicGames.Core/AsyncUtils.cs b/Engine/Source/Programs/Shared/EpicGames.Core/AsyncUtils.cs
--- a/Engine/Source/Programs/Shared/EpicGames.Core/AsyncUtils.cs
+++ b/Engine/Source/Programs/Shared/EpicGames.Core/AsyncUtils.cs
@@ -22,6 +22,16 @@ namespace EpicGames.Core
 			return Task.Delay(-1, token).ContinueWith(x => { });
 		}
 
+		/// <summary>
+		/// Converts a cancellation token to a waitable task
+		/// </summary>
+		/// <param name="token">Cancellation token</param>
+		/// <returns></returns>
+		public static Task<T> AsTask<T>(this CancellationToken token)
+		{
+			return Task.Delay(-1, token).ContinueWith(_ => Task.FromCanceled<T>(token)).Unwrap();
+		}
+
 		/// <summary>
 		/// Attempts to get the result of a task, if it has finished
 		/// </summary>
