Horde: Missing file.

diff --git a/Engine/Source/Programs/Horde/Horde.Build/Compute/Impl/RedisTaskScheduler.cs b/Engine/Source/Programs/Horde/Horde.Build/Compute/Impl/RedisTaskScheduler.cs
--- a/Engine/Source/Programs/Horde/Horde.Build/Compute/Impl/RedisTaskScheduler.cs
+++ b/Engine/Source/Programs/Horde/Horde.Build/Compute/Impl/RedisTaskScheduler.cs
@@ -4,6 +4,7 @@ using System;
 using System.Collections.Generic;
 using System.Diagnostics;
 using System.Linq;
+using System.Runtime.CompilerServices;
 using System.Threading;
 using System.Threading.Channels;
 using System.Threading.Tasks;
@@ -46,7 +47,7 @@ namespace Horde.Build.Compute.Impl
 		/// <param name="predicate">Predicate for determining which queues can be removed from</param>
 		/// <param name="token">Cancellation token for the operation. Will return a null entry rather than throwing an exception.</param>
 		/// <returns>Information about the task to be executed</returns>
-		Task<(TQueueId, TTask)?> DequeueAsync(Func<TQueueId, ValueTask<bool>> predicate, CancellationToken token = default);
+		Task<Task<(TQueueId, TTask)>> DequeueAsync(Func<TQueueId, ValueTask<bool>> predicate, CancellationToken token = default);
 
 		/// <summary>
 		/// Gets hashes of all the inactive task queues
@@ -67,12 +68,12 @@ namespace Horde.Build.Compute.Impl
 		class Listener
 		{
 			public Func<TQueueId, ValueTask<bool>> _predicate;
-			public TaskCompletionSource<(TQueueId, TTask)?> CompletionSource { get; }
+			public TaskCompletionSource<(TQueueId, TTask)> CompletionSource { get; }
 
 			public Listener(Func<TQueueId, ValueTask<bool>> predicate)
 			{
 				_predicate = predicate;
-				CompletionSource = new TaskCompletionSource<(TQueueId, TTask)?>();
+				CompletionSource = new TaskCompletionSource<(TQueueId, TTask)>();
 			}
 		}
 
@@ -197,51 +198,57 @@ namespace Horde.Build.Compute.Impl
 		/// <param name="predicate">Predicate for queues that tasks can be removed from</param>
 		/// <param name="token">Cancellation token for waiting for an item</param>
 		/// <returns>The dequeued item, or null if no item is available</returns>
-		public async Task<(TQueueId, TTask)?> DequeueAsync(Func<TQueueId, ValueTask<bool>> predicate, CancellationToken token = default)
+		public async Task<Task<(TQueueId, TTask)>> DequeueAsync(Func<TQueueId, ValueTask<bool>> predicate, CancellationToken token = default)
 		{
 			// Compare against all the list of cached queues to see if we can dequeue something from any of them
-			Listener? listener = null;
+			TQueueId[] queues = await _queueIndex.MembersAsync();
+
+			// Try to dequeue an item from the list
+			(TQueueId, TTask)? entry = await TryAssignToLocalAgentAsync(queues, predicate);
+			if (entry != null)
+			{
+				return Task.FromResult(entry.Value);
+			}
+
+			// Otherwise create a new task to do the wait
+			return WaitToDequeueAsync(queues, predicate, token);
+		}
+
+		private async Task<(TQueueId, TTask)> WaitToDequeueAsync(TQueueId[] queues, Func<TQueueId, ValueTask<bool>> predicate, CancellationToken token)
+		{
+			Listener listener = new Listener(predicate);
 			try
 			{
-				TQueueId[] queues = await _queueIndex.MembersAsync();
-				while (!token.IsCancellationRequested)
+				// Add the listener to the global list
+				lock (_listeners)
 				{
-					// Try to dequeue an item from the list
-					(TQueueId, TTask)? entry = await TryAssignToLocalAgentAsync(queues, predicate);
-					if (entry != null)
-					{
-						return entry;
-					}
+					_listeners.Add(listener);
+				}
 
-					// Create and register a listener for this waiter 
-					if (listener == null)
-					{
-						listener = new Listener(predicate);
-						lock (_listeners)
-						{
-							_listeners.Add(listener);
-						}
-					}
-					else
+				// Try to dequeue an item from the list again. An item may have been made available 
+				(TQueueId, TTask)? entry = await TryAssignToLocalAgentAsync(queues, predicate);
+				if (entry != null)
+				{
+					if (!listener.CompletionSource.TrySetResult(entry.Value))
 					{
-						using (IDisposable registration = token.Register(() => listener.CompletionSource.TrySetResult(null)))
-						{
-							return await listener.CompletionSource.Task;
-						}
+						(TQueueId queue, TTask task) = entry.Value;
+						await EnqueueAsync(queue, task, true);
 					}
 				}
+
+				// Wait for an item to be available
+				using (IDisposable registration = token.Register(() => listener.CompletionSource.TrySetCanceled()))
+				{
+					return await listener.CompletionSource.Task;
+				}
 			}
 			finally
 			{
-				if (listener != null)
+				lock (_listeners)
 				{
-					lock (_listeners)
-					{
-						_listeners.Remove(listener);
-					}
+					_listeners.Remove(listener);
 				}
 			}
-			return null;
 		}
 
 		/// <summary>
@@ -438,7 +445,7 @@ namespace Horde.Build.Compute.Impl
 					}
 
 					// Assign it to the listener
-					if (listener.CompletionSource.TrySetResult(entry))
+					if (listener.CompletionSource.TrySetResult(entry.Value))
 					{
 						entry = null;
 					}
