Horde: Prevent cancelled tasks throwing exceptions during lease assignment.

diff --git a/Engine/Source/Programs/Horde/Horde.Build/Agents/Impl/AgentService.cs b/Engine/Source/Programs/Horde/Horde.Build/Agents/Impl/AgentService.cs
--- a/Engine/Source/Programs/Horde/Horde.Build/Agents/Impl/AgentService.cs
+++ b/Engine/Source/Programs/Horde/Horde.Build/Agents/Impl/AgentService.cs
@@ -389,12 +389,12 @@ namespace Horde.Build.Services
 					}
 
 					// Create all the tasks to wait for
-					List<Task<(ITaskSource source, AgentLease lease)>> tasks = new List<Task<(ITaskSource, AgentLease)>>();
+					List<Task<(ITaskSource, AgentLease)?>> tasks = new List<Task<(ITaskSource, AgentLease)?>>();
 					foreach (ITaskSource taskSource in _taskSources)
 					{
 						if (CanUseTaskSource(agent, taskSource) && !cancellationSource.IsCancellationRequested)
 						{
-							Task<(ITaskSource, AgentLease)> task = await GuardedAssignLeaseAsync(taskSource, agent, cancellationSource);
+							Task<(ITaskSource, AgentLease)?> task = await GuardedAssignLeaseAsync(taskSource, agent, cancellationSource);
 							tasks.Add(task);
 						}
 					}
@@ -411,18 +411,19 @@ namespace Horde.Build.Services
 					await Task.WhenAll(tasks);
 
 					// Find the first result
-					foreach (Task<(ITaskSource, AgentLease)> task in tasks)
+					foreach (Task<(ITaskSource, AgentLease)?> task in tasks)
 					{
-						(ITaskSource source, AgentLease lease) taskResult;
-						if (task.TryGetResult(out taskResult))
+						(ITaskSource, AgentLease)? taskResult;
+						if (task.TryGetResult(out taskResult) && taskResult != null)
 						{
+							(ITaskSource taskSource, AgentLease taskLease) = taskResult.Value;
 							if (result == null)
 							{
-								result = taskResult;
+								result = (taskSource, taskLease);
 							}
 							else
 							{
-								await taskResult.source.CancelLeaseAsync(agent, taskResult.lease.Id, Any.Parser.ParseFrom(taskResult.lease.Payload));
+								await taskSource.CancelLeaseAsync(agent, taskLease.Id, Any.Parser.ParseFrom(taskLease.Payload));
 							}
 						}
 					}
@@ -463,29 +464,46 @@ namespace Horde.Build.Services
 			return agent;
 		}
 
-		async Task<Task<(ITaskSource, AgentLease)>> GuardedAssignLeaseAsync(ITaskSource source, IAgent agent, CancellationTokenSource cancellationSource)
+		async Task<Task<(ITaskSource, AgentLease)?>> GuardedAssignLeaseAsync(ITaskSource source, IAgent agent, CancellationTokenSource cancellationSource)
 		{
 			CancellationToken cancellationToken = cancellationSource.Token;
 			try
 			{
 				Task<AgentLease> task = await source.AssignLeaseAsync(agent, cancellationToken);
-				return task.ContinueWith(x => WrapAssignedLease(source, x.Result, cancellationSource), CancellationToken.None, TaskContinuationOptions.OnlyOnRanToCompletion, TaskScheduler.Default);
+				return task.ContinueWith(x => WrapAssignedLease(source, x, cancellationSource), TaskScheduler.Default);
 			}
 			catch (TaskCanceledException) when (cancellationToken.IsCancellationRequested)
 			{
-				return Task.FromCanceled<(ITaskSource, AgentLease)>(cancellationToken);
+				return Task.FromResult<(ITaskSource, AgentLease)?>(null);
 			}
 			catch (Exception ex)
 			{
 				_logger.LogError(ex, "Exception while trying to assign lease");
-				return Task.FromException<(ITaskSource, AgentLease)>(ex);
+				return Task.FromException<(ITaskSource, AgentLease)?>(ex);
 			}
 		}
 
-		static (ITaskSource, AgentLease) WrapAssignedLease(ITaskSource source, AgentLease lease, CancellationTokenSource cancellationSource)
+		(ITaskSource, AgentLease)? WrapAssignedLease(ITaskSource source, Task<AgentLease> task, CancellationTokenSource cancellationSource)
 		{
-			cancellationSource.Cancel();
-			return (source, lease);
+			if (task.IsCanceled)
+			{
+				return null;
+			}
+			else if (task.IsCompletedSuccessfully)
+			{
+				cancellationSource.Cancel();
+				return (source, task.Result);
+			}
+			else if (task.IsFaulted)
+			{
+				_logger.LogError(task.Exception, "Exception while trying to assign lease");
+				return null;
+			}
+			else
+			{
+				_logger.LogWarning("Unhandled task state: {Status}", task.Status);
+				return null;
+			}
 		}
 
 		/// <summary>
