Horde Storage - Ignore running blob index consistency check of namespaces without policies

diff --git a/Engine/Source/Programs/Horde/HordeStorage/Horde.Storage/Implementation/Blob/BlobIndexConsistencyCheckService.cs b/Engine/Source/Programs/Horde/HordeStorage/Horde.Storage/Implementation/Blob/BlobIndexConsistencyCheckService.cs
--- a/Engine/Source/Programs/Horde/HordeStorage/Horde.Storage/Implementation/Blob/BlobIndexConsistencyCheckService.cs
+++ b/Engine/Source/Programs/Horde/HordeStorage/Horde.Storage/Implementation/Blob/BlobIndexConsistencyCheckService.cs
@@ -69,9 +69,18 @@ namespace Horde.Storage.Implementation
 
         private bool NamespaceShouldBeCheckedForConsistency(NamespaceId ns)
         {
-            NamespaceSettings.PerNamespaceSettings policy = _namespacePolicyResolver.GetPoliciesForNs(ns);
+            try
+            {
+                NamespaceSettings.PerNamespaceSettings policy = _namespacePolicyResolver.GetPoliciesForNs(ns);
+
+                return policy.IsLegacyNamespace.HasValue && !policy.IsLegacyNamespace.Value;
+            }
+            catch (UnknownNamespaceException)
+            {
+                _logger.Warning("Namespace {Namespace} does not configure any policy, not running blob index consistency checks on it.", ns);
+                return false;
+            }
 
-            return policy.IsLegacyNamespace.HasValue && !policy.IsLegacyNamespace.Value;
         }
 
         private async Task RunConsistencyCheck()
