Horde Storage - Ignore cleaning namespaces that lack a policy configuration

diff --git a/Engine/Source/Programs/Horde/HordeStorage/Horde.Storage/Implementation/GC/BlobCleanupService.cs b/Engine/Source/Programs/Horde/HordeStorage/Horde.Storage/Implementation/GC/BlobCleanupService.cs
--- a/Engine/Source/Programs/Horde/HordeStorage/Horde.Storage/Implementation/GC/BlobCleanupService.cs
+++ b/Engine/Source/Programs/Horde/HordeStorage/Horde.Storage/Implementation/GC/BlobCleanupService.cs
@@ -79,8 +79,6 @@ namespace Horde.Storage.Implementation
 
         public async Task Cleanup(BlobCleanupState state, CancellationToken cancellationToken)
         {
-            List<BlobIdentifier> removedBlobs = new List<BlobIdentifier>();
-            
             foreach (IBlobCleanup blobCleanup in state.BlobCleanups)
             {
                 if (!blobCleanup.ShouldRun())
diff --git a/Engine/Source/Programs/Horde/HordeStorage/Horde.Storage/Implementation/GC/CacheCleanup.cs b/Engine/Source/Programs/Horde/HordeStorage/Horde.Storage/Implementation/GC/CacheCleanup.cs
--- a/Engine/Source/Programs/Horde/HordeStorage/Horde.Storage/Implementation/GC/CacheCleanup.cs
+++ b/Engine/Source/Programs/Horde/HordeStorage/Horde.Storage/Implementation/GC/CacheCleanup.cs
@@ -41,7 +41,16 @@ namespace Horde.Storage.Implementation
 
         public Task<List<OldRecord>> Cleanup(NamespaceId ns, CancellationToken cancellationToken)
         {
-            NamespaceSettings.PerNamespaceSettings policies = _namespacePolicyResolver.GetPoliciesForNs(ns);
+            NamespaceSettings.PerNamespaceSettings policies;
+            try
+            {
+                policies = _namespacePolicyResolver.GetPoliciesForNs(ns);
+            }
+            catch (UnknownNamespaceException)
+            {
+                _logger.Warning("Namespace {Namespace} does not configure any policy, not running ref cleanup on it.", ns);
+                return Task.FromResult(new List<OldRecord>());
+            }
 
             using IScope scope = Tracer.Instance.StartActive("gc.refs.namespace");
             scope.Span.ResourceName = ns.ToString();
