Horde Storage - Fixed failing unit tests with latest

Converting a dictonary of objects to compact binary doesnt really work as its undefined how to serialize the object (it should be a CbObject if we just wanted to serialize a hirearchy of compact binary).

This means that the metadata field will not be serialized for compact binary representations in the legacy api, this isnt used anyway and didnt used to work before our migration to the new cb serializer so this should be fine.

diff --git a/Engine/Source/Programs/Horde/HordeStorage/Horde.Storage/Implementation/DDCRefService.cs b/Engine/Source/Programs/Horde/HordeStorage/Horde.Storage/Implementation/DDCRefService.cs
--- a/Engine/Source/Programs/Horde/HordeStorage/Horde.Storage/Implementation/DDCRefService.cs
+++ b/Engine/Source/Programs/Horde/HordeStorage/Horde.Storage/Implementation/DDCRefService.cs
@@ -41,7 +41,6 @@ namespace Horde.Storage.Implementation
         [CbField("lastAccessTime")]
         public DateTime? LastAccessTime { get; set; }
 
-        [CbField("metadata")]
         public Dictionary<string, object>? Metadata { get; set; }
 
         [CbField("contentHash")]
diff --git a/Engine/Source/Programs/Horde/HordeStorage/Horde.Storage/Tests/Functional/RefTests.cs b/Engine/Source/Programs/Horde/HordeStorage/Horde.Storage/Tests/Functional/RefTests.cs
--- a/Engine/Source/Programs/Horde/HordeStorage/Horde.Storage/Tests/Functional/RefTests.cs
+++ b/Engine/Source/Programs/Horde/HordeStorage/Horde.Storage/Tests/Functional/RefTests.cs
@@ -343,9 +343,8 @@ namespace Horde.Storage.FunctionalTests.Ref
             byte[] buffer = await response.Content.ReadAsByteArrayAsync();
             CbObject o = new CbObject(buffer);
             List<CbField> fields = o.ToList();
-            Assert.AreEqual(6, fields.Count);
+            Assert.AreEqual(5, fields.Count);
             Assert.AreEqual($"{TestNamespace}.bucket.testObjectWithMetadata", o["name"]!.AsString());
-            Assert.IsNotNull(o["metadata"]);
 
             Assert.IsNotNull(o["contentHash"]);
 
