Horde: Add a fail-safe exception handler for session updates.

diff --git a/Engine/Source/Programs/Horde/Horde.Build/Services/RpcService.cs b/Engine/Source/Programs/Horde/Horde.Build/Services/RpcService.cs
--- a/Engine/Source/Programs/Horde/Horde.Build/Services/RpcService.cs
+++ b/Engine/Source/Programs/Horde/Horde.Build/Services/RpcService.cs
@@ -345,7 +345,14 @@ namespace Horde.Build.Services
 					}
 
 					// Update the session
-					agent = await _agentService.UpdateSessionWithWaitAsync(agent, SessionId.Parse(request.SessionId), request.Status, properties, resources, request.Leases, cancellationSource.Token);
+					try
+					{
+						agent = await _agentService.UpdateSessionWithWaitAsync(agent, SessionId.Parse(request.SessionId), request.Status, properties, resources, request.Leases, cancellationSource.Token);
+					}
+					catch (Exception ex)
+					{
+						_logger.LogError(ex, "Swallowed exception while updating session for {AgentId}.", request.AgentId);
+					}
 				}
 
 				// Handle the invalid agent case
