Horde: Commit replication changes

- Fix uncapped growth of channel for stream updates.
- Prevent no-op changes unnecessarily triggering update iterations.
- Add functionality to AsyncEvent to support latching the event on, rather than just pulsing it.

diff --git a/Engine/Source/Programs/Horde/Horde.Build/Compute/Impl/RedisTaskScheduler.cs b/Engine/Source/Programs/Horde/Horde.Build/Compute/Impl/RedisTaskScheduler.cs
--- a/Engine/Source/Programs/Horde/Horde.Build/Compute/Impl/RedisTaskScheduler.cs
+++ b/Engine/Source/Programs/Horde/Horde.Build/Compute/Impl/RedisTaskScheduler.cs
@@ -47,7 +47,7 @@ namespace Horde.Build.Compute.Impl
 		/// <param name="predicate">Predicate for determining which queues can be removed from</param>
 		/// <param name="token">Cancellation token for the operation. Will return a null entry rather than throwing an exception.</param>
 		/// <returns>Information about the task to be executed</returns>
-		Task<Task<(TQueueId, TTask)>> DequeueAsync(Func<TQueueId, ValueTask<bool>> predicate, CancellationToken token = default);
+		Task<Task<(TQueueId, TTask)?>> DequeueAsync(Func<TQueueId, ValueTask<bool>> predicate, CancellationToken token = default);
 
 		/// <summary>
 		/// Gets hashes of all the inactive task queues
@@ -68,12 +68,12 @@ namespace Horde.Build.Compute.Impl
 		class Listener
 		{
 			public Func<TQueueId, ValueTask<bool>> _predicate;
-			public TaskCompletionSource<(TQueueId, TTask)> CompletionSource { get; }
+			public TaskCompletionSource<(TQueueId, TTask)?> CompletionSource { get; }
 
 			public Listener(Func<TQueueId, ValueTask<bool>> predicate)
 			{
 				_predicate = predicate;
-				CompletionSource = new TaskCompletionSource<(TQueueId, TTask)>();
+				CompletionSource = new TaskCompletionSource<(TQueueId, TTask)?>();
 			}
 		}
 
@@ -198,7 +198,7 @@ namespace Horde.Build.Compute.Impl
 		/// <param name="predicate">Predicate for queues that tasks can be removed from</param>
 		/// <param name="token">Cancellation token for waiting for an item</param>
 		/// <returns>The dequeued item, or null if no item is available</returns>
-		public async Task<Task<(TQueueId, TTask)>> DequeueAsync(Func<TQueueId, ValueTask<bool>> predicate, CancellationToken token = default)
+		public async Task<Task<(TQueueId, TTask)?>> DequeueAsync(Func<TQueueId, ValueTask<bool>> predicate, CancellationToken token = default)
 		{
 			// Compare against all the list of cached queues to see if we can dequeue something from any of them
 			TQueueId[] queues = await _queueIndex.MembersAsync();
@@ -207,14 +207,14 @@ namespace Horde.Build.Compute.Impl
 			(TQueueId, TTask)? entry = await TryAssignToLocalAgentAsync(queues, predicate);
 			if (entry != null)
 			{
-				return Task.FromResult(entry.Value);
+				return Task.FromResult(entry);
 			}
 
 			// Otherwise create a new task to do the wait
 			return WaitToDequeueAsync(queues, predicate, token);
 		}
 
-		private async Task<(TQueueId, TTask)> WaitToDequeueAsync(TQueueId[] queues, Func<TQueueId, ValueTask<bool>> predicate, CancellationToken token)
+		private async Task<(TQueueId, TTask)?> WaitToDequeueAsync(TQueueId[] queues, Func<TQueueId, ValueTask<bool>> predicate, CancellationToken token)
 		{
 			Listener listener = new Listener(predicate);
 			try
@@ -237,7 +237,7 @@ namespace Horde.Build.Compute.Impl
 				}
 
 				// Wait for an item to be available
-				using (IDisposable registration = token.Register(() => listener.CompletionSource.TrySetCanceled()))
+				using (IDisposable registration = token.Register(() => listener.CompletionSource.TrySetResult(null)))
 				{
 					return await listener.CompletionSource.Task;
 				}
diff --git a/Engine/Source/Programs/Horde/Horde.Build/Perforce/Impl/CommitService.cs b/Engine/Source/Programs/Horde/Horde.Build/Perforce/Impl/CommitService.cs
--- a/Engine/Source/Programs/Horde/Horde.Build/Perforce/Impl/CommitService.cs
+++ b/Engine/Source/Programs/Horde/Horde.Build/Perforce/Impl/CommitService.cs
@@ -344,6 +344,7 @@ namespace Horde.Build.Commits.Impl
 			}
 
 			// Update the database with all the commits
+			HashSet<StreamId> newStreamIds = new HashSet<StreamId>();
 			await foreach (NewCommit newCommit in FindCommitsForClusterAsync(clusterName, streamToFirstChange))
 			{
 				ICommit commit = await _commitCollection.AddOrReplaceAsync(newCommit);
@@ -352,8 +353,14 @@ namespace Horde.Build.Commits.Impl
 				RedisList<int> streamCommitsKey = RedisStreamChanges(commit.StreamId);
 				await streamCommitsKey.RightPushAsync(commit.Change);
 
-				await _redisDirtyStreams.AddAsync(commit.StreamId);
-				await _redis.PublishAsync(RedisUpdateChannel, commit.StreamId);
+				newStreamIds.Add(commit.StreamId);
+			}
+
+			// Signal to any listeners that we have new data to process
+			foreach (StreamId newStreamId in newStreamIds)
+			{
+				await _redisDirtyStreams.AddAsync(newStreamId);
+				await _redis.PublishAsync(RedisUpdateChannel, newStreamId);
 			}
 		}
 
@@ -609,9 +616,10 @@ namespace Horde.Build.Commits.Impl
 
 		RedisKey RedisBaseKey { get; } = new RedisKey("commits/");
 		RedisChannel<StreamId> RedisUpdateChannel { get; } = new RedisChannel<StreamId>("commits/streams");
+		bool _stopping;
+		readonly AsyncEvent _updateStreamsEvent = new AsyncEvent();
 		readonly RedisSet<StreamId> _redisDirtyStreams;
 		readonly RedisSortedSet<StreamId> _redisReservations;
-		Channel<StreamId>? _redisUpdateStreams;
 		RedisChannelSubscription<StreamId>? _redisUpdateSubscription;
 		Task? _streamUpdateTask;
 
@@ -619,18 +627,16 @@ namespace Horde.Build.Commits.Impl
 
 		async Task StartContentReplicationAsync()
 		{
-			_redisUpdateStreams = Channel.CreateUnbounded<StreamId>();
-			_redisUpdateSubscription = await _redis.Multiplexer.GetSubscriber().SubscribeAsync(RedisUpdateChannel, (_, streamId) => _redisUpdateStreams.Writer.TryWrite(streamId));
+			_stopping = false;
+			_updateStreamsEvent.Reset();
+			_redisUpdateSubscription = await _redis.Multiplexer.GetSubscriber().SubscribeAsync(RedisUpdateChannel, (_, _) => _updateStreamsEvent.Pulse());
 			_streamUpdateTask = Task.Run(() => UpdateContentAsync());
 		}
 
 		async Task StopContentReplicationAsync()
 		{
-			if (_redisUpdateStreams != null)
-			{
-				_redisUpdateStreams.Writer.Complete();
-				_redisUpdateStreams = null;
-			}
+			_stopping = true;
+			_updateStreamsEvent.Latch();
 			if (_redisUpdateSubscription != null)
 			{
 				await _redisUpdateSubscription.DisposeAsync();
@@ -642,7 +648,7 @@ namespace Horde.Build.Commits.Impl
 		async Task UpdateContentAsync()
 		{
 			List<(StreamId, Task)> backgroundTasks = new List<(StreamId, Task)>(MaxBackgroundTasks);
-			while (!_redisUpdateStreams!.Reader.Completion.IsCompleted || backgroundTasks.Count > 0)
+			while (!_stopping || backgroundTasks.Count > 0)
 			{
 				try
 				{
@@ -673,11 +679,10 @@ namespace Horde.Build.Commits.Impl
 
 			// Create a list of events to wait for
 			List<Task> waitTasks = new List<Task>(backgroundTasks.Select(x => x.Item2));
-			if (!_redisUpdateStreams!.Reader.Completion.IsCompleted)
+			if (!_stopping)
 			{
-				waitTasks.Add(_redisUpdateStreams.Reader.WaitToReadAsync().AsTask());
-
 				// If we have spare slots for executing background tasks, check if there are any dirty streams
+				Task newStreamTask = _updateStreamsEvent.Task;
 				if (backgroundTasks.Count < MaxBackgroundTasks)
 				{
 					// Expire any reservations that are no longer valid
@@ -725,6 +730,12 @@ namespace Horde.Build.Commits.Impl
 						}
 					}
 				}
+
+				// If we still have bandwidth to process more tasks, wait for new streams to become available
+				if (backgroundTasks.Count < MaxBackgroundTasks)
+				{
+					waitTasks.Add(newStreamTask);
+				}
 			}
 
 			// Wait until any task has completed
diff --git a/Engine/Source/Programs/Shared/EpicGames.Core/AsyncEvent.cs b/Engine/Source/Programs/Shared/EpicGames.Core/AsyncEvent.cs
--- a/Engine/Source/Programs/Shared/EpicGames.Core/AsyncEvent.cs
+++ b/Engine/Source/Programs/Shared/EpicGames.Core/AsyncEvent.cs
@@ -16,12 +16,56 @@ namespace EpicGames.Core
 		TaskCompletionSource<bool> _source = new TaskCompletionSource<bool>();
 
 		/// <summary>
-		/// Signal the event
+		/// Legacy name for <see cref="Pulse"/>.
 		/// </summary>
-		public void Set()
+		public void Set() => Pulse();
+
+		/// <summary>
+		/// Pulse the event, allowing any captured copies of the task to continue.
+		/// </summary>
+		public void Pulse()
+		{
+			for (; ; )
+			{
+				TaskCompletionSource<bool> prevSource = _source;
+				if (prevSource.Task.IsCompleted)
+				{
+					// The task source is latched or has been set from another thread. Either way behaves identically wrt the contract of this method.
+					break;
+				}
+				if (Interlocked.CompareExchange(ref _source, new TaskCompletionSource<bool>(), prevSource) == prevSource)
+				{
+					prevSource.TrySetResult(true);
+					break;
+				}
+			}
+		}
+
+		/// <summary>
+		/// Resets an event which is latched to the set state.
+		/// </summary>
+		public void Reset()
+		{
+			for (; ; )
+			{
+				TaskCompletionSource<bool> prevSource = _source;
+				if (!prevSource.Task.IsCompleted)
+				{
+					break;
+				}
+				else if (Interlocked.CompareExchange(ref _source, new TaskCompletionSource<bool>(), prevSource) == prevSource)
+				{
+					break;
+				}
+			}
+		}
+
+		/// <summary>
+		/// Move the task to completed without returning it to an unset state. A call to Reset() is required to clear it.
+		/// </summary>
+		public void Latch()
 		{
-			TaskCompletionSource<bool> prevSource = Interlocked.Exchange(ref _source, new TaskCompletionSource<bool>());
-			prevSource.SetResult(true);
+			_source.TrySetResult(true);
 		}
 
 		/// <summary>
