Horde: Avoid cancelled tasks throwing exceptions in scheduler; just return null leases instead.

diff --git a/Engine/Source/Programs/Horde/Horde.Build/Agents/ITaskSource.cs b/Engine/Source/Programs/Horde/Horde.Build/Agents/ITaskSource.cs
--- a/Engine/Source/Programs/Horde/Horde.Build/Agents/ITaskSource.cs
+++ b/Engine/Source/Programs/Horde/Horde.Build/Agents/ITaskSource.cs
@@ -79,7 +79,7 @@ namespace Horde.Build.Tasks
 		/// <param name="agent">The agent to assign a lease to</param>
 		/// <param name="cancellationToken">Cancellation token for the wait</param>
 		/// <returns>Task returning a new lease object.</returns>
-		Task<Task<AgentLease>> AssignLeaseAsync(IAgent agent, CancellationToken cancellationToken);
+		Task<Task<AgentLease?>> AssignLeaseAsync(IAgent agent, CancellationToken cancellationToken);
 
 		/// <summary>
 		/// Cancel a lease that was previously assigned to an agent, allowing it to be assigned out again
@@ -173,7 +173,7 @@ namespace Horde.Build.Tasks
 		public MessageDescriptor Descriptor => s_message.Descriptor;
 
 		/// <inheritdoc/>
-		public abstract Task<Task<AgentLease>> AssignLeaseAsync(IAgent agent, CancellationToken cancellationToken);
+		public abstract Task<Task<AgentLease?>> AssignLeaseAsync(IAgent agent, CancellationToken cancellationToken);
 
 		/// <inheritdoc/>
 		public Task CancelLeaseAsync(IAgent agent, LeaseId leaseId, Any payload) => CancelLeaseAsync(agent, leaseId, payload.Unpack<TMessage>());
@@ -230,9 +230,9 @@ namespace Horde.Build.Tasks
         /// </summary>
         /// <param name="token">The cancellation token</param>
         /// <returns>Lease task</returns>
-        protected static Task<AgentLease> Skip(CancellationToken token)
+        protected static Task<AgentLease?> Skip(CancellationToken token)
         {
-            return token.AsTask<AgentLease>();
+            return Task.FromResult<AgentLease?>(null);
         }
 
         /// <summary>
@@ -240,10 +240,10 @@ namespace Horde.Build.Tasks
         /// </summary>
         /// <param name="token">The cancellation token</param>
         /// <returns>Lease task</returns>
-        protected static async Task<Task<AgentLease>> DrainAsync(CancellationToken token)
+        protected static async Task<Task<AgentLease?>> DrainAsync(CancellationToken token)
         {
             await token.AsTask();
-            return Task.FromCanceled<AgentLease>(token);
+            return Task.FromResult<AgentLease?>(null);
         }
 
         /// <summary>
@@ -251,9 +251,9 @@ namespace Horde.Build.Tasks
         /// </summary>
         /// <param name="lease">Lease to create the task from</param>
         /// <returns></returns>
-        protected static Task<AgentLease> Lease(AgentLease lease)
+        protected static Task<AgentLease?> Lease(AgentLease lease)
         {
-            return Task.FromResult<AgentLease>(lease);
+            return Task.FromResult<AgentLease?>(lease);
         }
     }
 }
diff --git a/Engine/Source/Programs/Horde/Horde.Build/Agents/Impl/AgentService.cs b/Engine/Source/Programs/Horde/Horde.Build/Agents/Impl/AgentService.cs
--- a/Engine/Source/Programs/Horde/Horde.Build/Agents/Impl/AgentService.cs
+++ b/Engine/Source/Programs/Horde/Horde.Build/Agents/Impl/AgentService.cs
@@ -469,7 +469,7 @@ namespace Horde.Build.Services
 			CancellationToken cancellationToken = cancellationSource.Token;
 			try
 			{
-				Task<AgentLease> task = await source.AssignLeaseAsync(agent, cancellationToken);
+				Task<AgentLease?> task = await source.AssignLeaseAsync(agent, cancellationToken);
 				return task.ContinueWith(x => WrapAssignedLease(source, x, cancellationSource), TaskScheduler.Default);
 			}
 			catch (TaskCanceledException) when (cancellationToken.IsCancellationRequested)
@@ -483,16 +483,23 @@ namespace Horde.Build.Services
 			}
 		}
 
-		(ITaskSource, AgentLease)? WrapAssignedLease(ITaskSource source, Task<AgentLease> task, CancellationTokenSource cancellationSource)
+		(ITaskSource, AgentLease)? WrapAssignedLease(ITaskSource source, Task<AgentLease?> task, CancellationTokenSource cancellationSource)
 		{
 			if (task.IsCanceled)
 			{
 				return null;
 			}
-			else if (task.IsCompletedSuccessfully)
+			else if (task.TryGetResult(out AgentLease? lease))
 			{
-				cancellationSource.Cancel();
-				return (source, task.Result);
+				if (lease == null)
+				{
+					return null;
+				}
+				else
+				{
+					cancellationSource.Cancel();
+					return (source, lease);
+				}
 			}
 			else if (task.IsFaulted)
 			{
diff --git a/Engine/Source/Programs/Horde/Horde.Build/Agents/Impl/Tasks/ConformTaskSource.cs b/Engine/Source/Programs/Horde/Horde.Build/Agents/Impl/Tasks/ConformTaskSource.cs
--- a/Engine/Source/Programs/Horde/Horde.Build/Agents/Impl/Tasks/ConformTaskSource.cs
+++ b/Engine/Source/Programs/Horde/Horde.Build/Agents/Impl/Tasks/ConformTaskSource.cs
@@ -149,7 +149,7 @@ namespace Horde.Build.Tasks.Impl
 		}
 
 		/// <inheritdoc/>
-		public override async Task<Task<AgentLease>> AssignLeaseAsync(IAgent agent, CancellationToken cancellationToken)
+		public override async Task<Task<AgentLease?>> AssignLeaseAsync(IAgent agent, CancellationToken cancellationToken)
 		{
 			DateTime utcNow = DateTime.UtcNow;
 			if (!await IsConformPendingAsync(agent, utcNow))
diff --git a/Engine/Source/Programs/Horde/Horde.Build/Agents/Impl/Tasks/RestartTaskSource.cs b/Engine/Source/Programs/Horde/Horde.Build/Agents/Impl/Tasks/RestartTaskSource.cs
--- a/Engine/Source/Programs/Horde/Horde.Build/Agents/Impl/Tasks/RestartTaskSource.cs
+++ b/Engine/Source/Programs/Horde/Horde.Build/Agents/Impl/Tasks/RestartTaskSource.cs
@@ -35,7 +35,7 @@ namespace Horde.Build.Tasks.Impl
 			OnLeaseStartedProperties.Add(nameof(RestartTask.LogId), x => new LogId(x.LogId));
 		}
 
-		public override async Task<Task<AgentLease>> AssignLeaseAsync(IAgent agent, CancellationToken cancellationToken)
+		public override async Task<Task<AgentLease?>> AssignLeaseAsync(IAgent agent, CancellationToken cancellationToken)
 		{
 			if (!agent.RequestRestart)
 			{
diff --git a/Engine/Source/Programs/Horde/Horde.Build/Agents/Impl/Tasks/ShutdownTaskSource.cs b/Engine/Source/Programs/Horde/Horde.Build/Agents/Impl/Tasks/ShutdownTaskSource.cs
--- a/Engine/Source/Programs/Horde/Horde.Build/Agents/Impl/Tasks/ShutdownTaskSource.cs
+++ b/Engine/Source/Programs/Horde/Horde.Build/Agents/Impl/Tasks/ShutdownTaskSource.cs
@@ -31,7 +31,7 @@ namespace Horde.Build.Tasks.Impl
 			OnLeaseStartedProperties.Add(nameof(ShutdownTask.LogId), x => new LogId(x.LogId));
 		}
 
-		public override async Task<Task<AgentLease>> AssignLeaseAsync(IAgent agent, CancellationToken cancellationToken)
+		public override async Task<Task<AgentLease?>> AssignLeaseAsync(IAgent agent, CancellationToken cancellationToken)
 		{
 			if (!agent.RequestShutdown)
 			{
diff --git a/Engine/Source/Programs/Horde/Horde.Build/Agents/Impl/Tasks/UpgradeTaskSource.cs b/Engine/Source/Programs/Horde/Horde.Build/Agents/Impl/Tasks/UpgradeTaskSource.cs
--- a/Engine/Source/Programs/Horde/Horde.Build/Agents/Impl/Tasks/UpgradeTaskSource.cs
+++ b/Engine/Source/Programs/Horde/Horde.Build/Agents/Impl/Tasks/UpgradeTaskSource.cs
@@ -40,7 +40,7 @@ namespace Horde.Build.Tasks.Impl
 			OnLeaseStartedProperties.Add(nameof(UpgradeTask.LogId), x => new LogId(x.LogId));
 		}
 
-		public override async Task<Task<AgentLease>> AssignLeaseAsync(IAgent agent, CancellationToken cancellationToken)
+		public override async Task<Task<AgentLease?>> AssignLeaseAsync(IAgent agent, CancellationToken cancellationToken)
 		{
 			string? requiredVersion = await GetRequiredSoftwareVersion(agent);
 			if (requiredVersion == null || agent.Version == requiredVersion)
diff --git a/Engine/Source/Programs/Horde/Horde.Build/Compute/Impl/ComputeService.cs b/Engine/Source/Programs/Horde/Horde.Build/Compute/Impl/ComputeService.cs
--- a/Engine/Source/Programs/Horde/Horde.Build/Compute/Impl/ComputeService.cs
+++ b/Engine/Source/Programs/Horde/Horde.Build/Compute/Impl/ComputeService.cs
@@ -272,19 +272,23 @@ namespace Horde.Build.Compute.Impl
 		}
 
 		/// <inheritdoc/>
-		public override async Task<Task<AgentLease>> AssignLeaseAsync(IAgent agent, CancellationToken cancellationToken)
+		public override async Task<Task<AgentLease?>> AssignLeaseAsync(IAgent agent, CancellationToken cancellationToken)
 		{
-			Task<(QueueKey, ComputeTaskInfo)> task = await _taskScheduler.DequeueAsync(queueKey => CheckRequirements(agent, queueKey), cancellationToken);
+			Task<(QueueKey, ComputeTaskInfo)?> task = await _taskScheduler.DequeueAsync(queueKey => CheckRequirements(agent, queueKey), cancellationToken);
 			return WaitForLeaseAsync(agent, task, cancellationToken);
 		}
 
-		private async Task<AgentLease> WaitForLeaseAsync(IAgent agent, Task<(QueueKey, ComputeTaskInfo)> task, CancellationToken cancellationToken)
+		private async Task<AgentLease?> WaitForLeaseAsync(IAgent agent, Task<(QueueKey, ComputeTaskInfo)?> task, CancellationToken cancellationToken)
 		{
 			for (; ; )
 			{
-				(QueueKey, ComputeTaskInfo) entry = await task;
+				(QueueKey, ComputeTaskInfo)? entry = await task;
+				if (entry == null)
+				{
+					return null;
+				}
 
-				AgentLease? lease = await CreateLeaseForEntryAsync(agent, await task);
+				AgentLease? lease = await CreateLeaseForEntryAsync(agent, entry.Value);
 				if (lease != null)
 				{
 					return lease;
