Horde: Missing file.

diff --git a/Engine/Source/Programs/Horde/Horde.Build/Jobs/Impl/JobTaskSource.cs b/Engine/Source/Programs/Horde/Horde.Build/Jobs/Impl/JobTaskSource.cs
--- a/Engine/Source/Programs/Horde/Horde.Build/Jobs/Impl/JobTaskSource.cs
+++ b/Engine/Source/Programs/Horde/Horde.Build/Jobs/Impl/JobTaskSource.cs
@@ -162,12 +162,12 @@ namespace Horde.Build.Tasks.Impl
 			/// <summary>
 			/// Task to wait for a lease to be assigned
 			/// </summary>
-			public Task<AgentLease> Task => LeaseSource.Task;
+			public Task<AgentLease?> Task => LeaseSource.Task;
 
 			/// <summary>
 			/// Completion source for the waiting agent. If a new queue item becomes available, the result will be passed through 
 			/// </summary>
-			public TaskCompletionSource<AgentLease> LeaseSource { get; } = new TaskCompletionSource<AgentLease>();
+			public TaskCompletionSource<AgentLease?> LeaseSource { get; } = new TaskCompletionSource<AgentLease?>();
 
 			/// <summary>
 			/// Constructor
@@ -485,7 +485,7 @@ namespace Horde.Build.Tasks.Impl
 			}
 		}
 
-	private async Task<IJob?> SkipBatchAsync(IJob job, SubResourceId batchId, IGraph graph, JobStepBatchError reason)
+		private async Task<IJob?> SkipBatchAsync(IJob job, SubResourceId batchId, IGraph graph, JobStepBatchError reason)
 		{
 			_logger.LogInformation("Skipping batch {BatchId} for job {JobId} (reason: {Reason})", batchId, job.Id, reason);
 
@@ -661,7 +661,7 @@ namespace Horde.Build.Tasks.Impl
 		}
 
 		/// <inheritdoc/>
-		public override async Task<Task<AgentLease>> AssignLeaseAsync(IAgent agent, CancellationToken cancellationToken)
+		public override Task<Task<AgentLease?>> AssignLeaseAsync(IAgent agent, CancellationToken cancellationToken)
 		{
 			QueueWaiter waiter = new QueueWaiter(agent);
 			lock (_lockObject)
@@ -671,25 +671,31 @@ namespace Horde.Build.Tasks.Impl
 				{
 					if (result == null)
 					{
-						return Skip(cancellationToken);
+						return Task.FromResult(Skip(cancellationToken));
 					}
-					return Lease(result);
+					return Task.FromResult(Lease(result));
 				}
 				_waiters.Add(waiter);
 			}
+			return Task.FromResult(WaitForLeaseAsync(waiter, cancellationToken));
+		}
 
-			AgentLease? lease;
-			using (cancellationToken.Register(() => waiter.LeaseSource.TrySetCanceled()))
+		private async Task<AgentLease?> WaitForLeaseAsync(QueueWaiter waiter, CancellationToken cancellationToken)
+		{
+			try
 			{
-				lease = await waiter.Task;
+				using (cancellationToken.Register(() => waiter.LeaseSource.TrySetResult(null)))
+				{
+					return await waiter.Task;
+				}
 			}
-
-			lock (_lockObject)
+			finally
 			{
-				_waiters.Remove(waiter);
+				lock (_lockObject)
+				{
+					_waiters.Remove(waiter);
+				}
 			}
-
-			return Lease(lease);
 		}
 
 		/// <inheritdoc/>
