Horde Storage - Clearer error message when we failed to resolve references for a blob that is not a compact binary.

diff --git a/Engine/Source/Programs/Horde/HordeStorage/Horde.Storage/Controllers/ObjectController.cs b/Engine/Source/Programs/Horde/HordeStorage/Horde.Storage/Controllers/ObjectController.cs
--- a/Engine/Source/Programs/Horde/HordeStorage/Horde.Storage/Controllers/ObjectController.cs
+++ b/Engine/Source/Programs/Horde/HordeStorage/Horde.Storage/Controllers/ObjectController.cs
@@ -201,7 +201,16 @@ namespace Horde.Storage.Controllers
             {
                 _logger.Warning("0 byte object found for {Id} {Namespace}", id, ns);
             }
-            CbObject compactBinaryObject = new CbObject(blobContents);
+
+            CbObject compactBinaryObject;
+            try
+            {
+                compactBinaryObject = new CbObject(blobContents);
+            }
+            catch (IndexOutOfRangeException)
+            {
+                return Problem(title: $"{id} was not a proper compact binary object.", detail: "Index out of range");
+            }
 
             try
             {
