Horde Storage - Added option to disable last access tracking and disabled it for the worker deployment, it will scan thru all refs to do garbage collection on them as well as replicate refs. This does not mean those refs should live longer.

diff --git a/Engine/Source/Programs/Horde/HordeStorage/Horde.Storage/Horde.Storage.Startup.cs b/Engine/Source/Programs/Horde/HordeStorage/Horde.Storage/Horde.Storage.Startup.cs
--- a/Engine/Source/Programs/Horde/HordeStorage/Horde.Storage/Horde.Storage.Startup.cs
+++ b/Engine/Source/Programs/Horde/HordeStorage/Horde.Storage/Horde.Storage.Startup.cs
@@ -255,13 +255,13 @@ namespace Horde.Storage
 
             services.AddSingleton(serviceType: typeof(IReplicationLog), ReplicationLogWriterFactory);
             
-            LastAccessTrackerRefRecord lastAccessTracker = new LastAccessTrackerRefRecord();
-            services.AddSingleton(serviceType: typeof(ILastAccessCache<RefRecord>), lastAccessTracker);
-            services.AddSingleton(serviceType: typeof(ILastAccessTracker<RefRecord>), lastAccessTracker);
+            services.AddSingleton<LastAccessTrackerRefRecord>();
+            services.AddSingleton(serviceType: typeof(ILastAccessCache<RefRecord>), p => p.GetService<LastAccessTrackerRefRecord>()!);
+            services.AddSingleton(serviceType: typeof(ILastAccessTracker<RefRecord>), p => p.GetService<LastAccessTrackerRefRecord>()!);
 
-            LastAccessTrackerReference lastAccessTrackerReferences = new LastAccessTrackerReference();
-            services.AddSingleton(serviceType: typeof(ILastAccessCache<LastAccessRecord>), lastAccessTrackerReferences);
-            services.AddSingleton(serviceType: typeof(ILastAccessTracker<LastAccessRecord>), lastAccessTrackerReferences);
+            services.AddSingleton<LastAccessTrackerReference>();
+            services.AddSingleton(serviceType: typeof(ILastAccessCache<LastAccessRecord>), p => p.GetService<LastAccessTrackerReference>()!);
+            services.AddSingleton(serviceType: typeof(ILastAccessTracker<LastAccessRecord>), p => p.GetService<LastAccessTrackerReference>()!);
 
             services.AddSingleton(serviceType: typeof(ILeaderElection), CreateLeaderElection);
 
diff --git a/Engine/Source/Programs/Horde/HordeStorage/Horde.Storage/Implementation/LastAccess/LastAccessTracker.cs b/Engine/Source/Programs/Horde/HordeStorage/Horde.Storage/Implementation/LastAccess/LastAccessTracker.cs
--- a/Engine/Source/Programs/Horde/HordeStorage/Horde.Storage/Implementation/LastAccess/LastAccessTracker.cs
+++ b/Engine/Source/Programs/Horde/HordeStorage/Horde.Storage/Implementation/LastAccess/LastAccessTracker.cs
@@ -9,12 +9,17 @@ using System.Threading.Tasks;
 using Datadog.Trace;
 using EpicGames.Horde.Storage;
 using Jupiter.Implementation;
+using Microsoft.Extensions.Options;
 using Serilog;
 
 namespace Horde.Storage.Implementation
 {
     public class LastAccessTrackerRefRecord : LastAccessTracker<RefRecord>
     {
+        public LastAccessTrackerRefRecord(IOptionsMonitor<HordeStorageSettings> settings) : base(settings)
+        {
+        }
+
         protected override string BuildCacheKey(RefRecord record)
         {
             return $"{record.Namespace}.{record.Bucket}.{record.RefName}";
@@ -37,6 +42,10 @@ namespace Horde.Storage.Implementation
 
     public class LastAccessTrackerReference : LastAccessTracker<LastAccessRecord>
     {
+        public LastAccessTrackerReference(IOptionsMonitor<HordeStorageSettings> settings) : base(settings)
+        {
+        }
+
         protected override string BuildCacheKey(LastAccessRecord record)
         {
             return $"{record.Namespace}.{record.Bucket}.{record.Key}";
@@ -45,6 +54,7 @@ namespace Horde.Storage.Implementation
 
     public abstract class LastAccessTracker<T> : ILastAccessTracker<T>, ILastAccessCache<T>
     {
+        private readonly IOptionsMonitor<HordeStorageSettings> _settings;
         private readonly ILogger _logger = Log.ForContext<LastAccessTracker<T>>();
 
         private ConcurrentDictionary<string, LastAccessRecord> _cache = new ConcurrentDictionary<string, LastAccessRecord>();
@@ -52,11 +62,18 @@ namespace Horde.Storage.Implementation
         // we will exchange the refs dictionary when fetching the records and use a rw lock to make sure no-one is trying to add things at the same time
         private readonly ReaderWriterLock _rwLock = new ReaderWriterLock();
 
+        public LastAccessTracker(IOptionsMonitor<HordeStorageSettings> settings)
+        {
+            _settings = settings;
+        }
 
         protected abstract string BuildCacheKey(T record);
 
         public Task TrackUsed(T record)
         {
+            if (!_settings.CurrentValue.EnableLastAccessTracking)
+                return Task.CompletedTask;
+
             return Task.Run(() =>
             {
                 using IScope _ = Tracer.Instance.StartActive("lastAccessTracker.track");
