Horde Storage - Setting to toggle last access tracking

diff --git a/Engine/Source/Programs/Horde/HordeStorage/Horde.Storage/Settings.cs b/Engine/Source/Programs/Horde/HordeStorage/Horde.Storage/Settings.cs
--- a/Engine/Source/Programs/Horde/HordeStorage/Horde.Storage/Settings.cs
+++ b/Engine/Source/Programs/Horde/HordeStorage/Horde.Storage/Settings.cs
@@ -135,6 +135,7 @@ namespace Horde.Storage
         public int? MaxSingleBlobSize { get; set; } = null; // disable blob partitioning
 
         public int LastAccessRollupFrequencySeconds { get; set; } = 900; // 15 minutes
+        public bool EnableLastAccessTracking { get; set; } = true;
         public bool EnableOnDemandReplication { get; set; } = false;
     }
 
