Horde Storage - Blob Store consistency check now also runs on the filesystem store (and is less S3 specific)

diff --git a/Engine/Source/Programs/Horde/HordeStorage/Horde.Storage/Implementation/Blob/BlobStoreConsistencyCheckService.cs b/Engine/Source/Programs/Horde/HordeStorage/Horde.Storage/Implementation/Blob/BlobStoreConsistencyCheckService.cs
--- a/Engine/Source/Programs/Horde/HordeStorage/Horde.Storage/Implementation/Blob/BlobStoreConsistencyCheckService.cs
+++ b/Engine/Source/Programs/Horde/HordeStorage/Horde.Storage/Implementation/Blob/BlobStoreConsistencyCheckService.cs
@@ -3,6 +3,7 @@
 using System;
 using System.Collections.Generic;
 using System.IO;
+using System.Linq;
 using System.Threading;
 using System.Threading.Tasks;
 using Dasync.Collections;
@@ -21,6 +22,7 @@ namespace Horde.Storage.Implementation
     public class BlobStoreConsistencyCheckService : PollingService<BlobStoreConsistencyCheckService.ConsistencyState>
     {
         private readonly IOptionsMonitor<ConsistencyCheckSettings> _settings;
+        private readonly IOptionsMonitor<HordeStorageSettings> _hordeStorageSettings;
         private readonly IServiceProvider _provider;
         private readonly ILeaderElection _leaderElection;
         private readonly IRefsStore _refsStore;
@@ -37,10 +39,11 @@ namespace Horde.Storage.Implementation
             return _settings.CurrentValue.EnableBlobStoreChecks;
         }
 
-        public BlobStoreConsistencyCheckService(IOptionsMonitor<ConsistencyCheckSettings> settings, IServiceProvider provider, ILeaderElection leaderElection, IRefsStore refsStore, IReferencesStore referencesStore, IBlobIndex blobIndex) :
+        public BlobStoreConsistencyCheckService(IOptionsMonitor<ConsistencyCheckSettings> settings, IOptionsMonitor<HordeStorageSettings> hordeStorageSettings, IServiceProvider provider, ILeaderElection leaderElection, IRefsStore refsStore, IReferencesStore referencesStore, IBlobIndex blobIndex) :
             base(serviceName: nameof(BlobStoreConsistencyCheckService), TimeSpan.FromSeconds(settings.CurrentValue.ConsistencyCheckPollFrequencySeconds), new ConsistencyState())
         {
             _settings = settings;
+            _hordeStorageSettings = hordeStorageSettings;
             _provider = provider;
             _leaderElection = leaderElection;
             _refsStore = refsStore;
@@ -62,8 +65,6 @@ namespace Horde.Storage.Implementation
                 return false;
             }
 
-            using IScope scope = Tracer.Instance.StartActive("blob_store.consistency_check.poll");
-
             await RunConsistencyCheck();
 
             return true;
@@ -71,57 +72,71 @@ namespace Horde.Storage.Implementation
 
         private async Task RunConsistencyCheck()
         {
-            //TODO: This isn't really specific to the S3 store, but it should likely only run on the backing store and not the cache layers which we do not have a way of fetching right now.
-            AmazonS3Store? s3Store;
-            try
-            {
-                s3Store = _provider.GetService<AmazonS3Store>();
-            }
-            catch (Exception)
-            {
-                s3Store = null;
-            }
-
-            if (s3Store == null)
+            foreach (IBlobStore blobStore in BlobService.GetBlobStores(_provider, _hordeStorageSettings).Where(RunConsistencyCheckOnBlobStore))
             {
-                _logger.Warning("No S3 Store found so will not run any consistency check");
-                return;
-            }
+                string blobStoreName = blobStore.GetType().Name;
 
-            List<NamespaceId> namespaces = await _refsStore.GetNamespaces().ToListAsync();
-            namespaces.AddRange(await _referencesStore.GetNamespaces().ToListAsync());
+                using IScope blobStoreScope = Tracer.Instance.StartActive("blob_store.consistency_check.poll");
+                blobStoreScope.Span.ResourceName = blobStoreName;
 
-            // technically this does not need to be run per namespace but per s3 bucket
-            await foreach (NamespaceId ns in namespaces)
-            {
-                ulong countOfBlobsChecked = 0;
-                ulong countOfIncorrectBlobsFound = 0;
+                bool isRootStore = blobStore is AmazonS3Store or AzureBlobStore;
 
-                using IScope scope = Tracer.Instance.StartActive("consistency_check.run");
-                scope.Span.ResourceName = ns.ToString();
+                List<NamespaceId> namespaces = await _refsStore.GetNamespaces().ToListAsync();
+                namespaces.AddRange(await _referencesStore.GetNamespaces().ToListAsync());
 
-                await foreach ((BlobIdentifier blob, DateTime lastModified) in s3Store.ListObjects(ns))
+                // technically this does not need to be run per namespace but per storage pool
+                await foreach (NamespaceId ns in namespaces)
                 {
-                    if (countOfBlobsChecked % 100 == 0)
-                        _logger.Information("Consistency check running on S3 blob store, count of blobs processed so far: {CountOfBlobs}", countOfBlobsChecked);
-                    Interlocked.Increment(ref countOfBlobsChecked);
-                    
-                    BlobContents contents = await s3Store.GetObject(ns, blob);
-                    await using Stream s = contents.Stream;
-
-                    BlobIdentifier newHash = await BlobIdentifier.FromStream(s);
-                    if (!blob.Equals(newHash))
-                    {
-                        _logger.Error("Mismatching hash for S3 object {Blob} in {Namespace}, new hash has {NewHash}. Deleting incorrect blob.", blob, ns, newHash);
+                    ulong countOfBlobsChecked = 0;
+                    ulong countOfIncorrectBlobsFound = 0;
+
+                    using IScope scope = Tracer.Instance.StartActive("consistency_check.run");
+                    scope.Span.ResourceName = ns.ToString();
 
-                        Interlocked.Increment(ref countOfIncorrectBlobsFound);
-                        await s3Store.DeleteObject(ns, blob);
-                        // update blob index tracking to indicate that we no longer have this blob in this region
-                        await _blobIndex.RemoveBlobFromRegion(ns, blob);
+                    await foreach ((BlobIdentifier blob, DateTime lastModified) in blobStore.ListObjects(ns))
+                    {
+                        if (countOfBlobsChecked % 100 == 0)
+                            _logger.Information("Consistency check running on Blob Store {BlobStore}, count of blobs processed so far: {CountOfBlobs}", blobStoreName, countOfBlobsChecked);
+                        Interlocked.Increment(ref countOfBlobsChecked);
+                        
+                        BlobContents contents = await blobStore.GetObject(ns, blob);
+                        await using Stream s = contents.Stream;
+
+                        BlobIdentifier newHash = await BlobIdentifier.FromStream(s);
+                        if (!blob.Equals(newHash))
+                        {
+                            _logger.Error("Mismatching hash for {Blob} in {Namespace} stored in {BlobStore}, new hash has {NewHash}. Deleting incorrect blob.", blob, ns, blobStoreName,newHash);
+
+                            Interlocked.Increment(ref countOfIncorrectBlobsFound);
+                            await blobStore.DeleteObject(ns, blob);
+
+                            if (isRootStore)
+                            {
+                                // update blob index tracking to indicate that we no longer have this blob in this region
+                                await _blobIndex.RemoveBlobFromRegion(ns, blob);
+                            }
+                        }
                     }
+
+                    _logger.Information("Blob Store {BlobStore}: Consistency check finished for {Namespace}, found {CountOfIncorrectBlobs} incorrect blobs. Processed {CountOfBlobs} blobs.", blobStoreName, ns, countOfIncorrectBlobsFound, countOfBlobsChecked);
                 }
+            }
+        }
 
-                _logger.Information("Blob Store Consistency check finished for {Namespace}, found {CountOfIncorrectBlobs} incorrect blobs. Processed {CountOfBlobs} blobs.", ns, countOfIncorrectBlobsFound, countOfBlobsChecked);
+        private bool RunConsistencyCheckOnBlobStore(IBlobStore blobStore)
+        {
+            switch (blobStore)
+            {
+                case FileSystemStore:
+                case AzureBlobStore:
+                case AmazonS3Store:
+                    return true;
+                case MemoryBlobStore:
+                case MemoryCacheBlobStore:
+                case RelayBlobStore:
+                    return false;
+                default:
+                    throw new NotImplementedException();
             }
         }
 
diff --git a/Engine/Source/Programs/Horde/HordeStorage/Horde.Storage/Implementation/Blob/IBlobService.cs b/Engine/Source/Programs/Horde/HordeStorage/Horde.Storage/Implementation/Blob/IBlobService.cs
--- a/Engine/Source/Programs/Horde/HordeStorage/Horde.Storage/Implementation/Blob/IBlobService.cs
+++ b/Engine/Source/Programs/Horde/HordeStorage/Horde.Storage/Implementation/Blob/IBlobService.cs
@@ -86,12 +86,12 @@ public class BlobService : IBlobService
         _namespacePolicyResolver = namespacePolicyResolver;
     }
 
-    private IEnumerable<IBlobStore> GetBlobStores(IServiceProvider provider, IOptionsMonitor<HordeStorageSettings> settings)
+    public static IEnumerable<IBlobStore> GetBlobStores(IServiceProvider provider, IOptionsMonitor<HordeStorageSettings> settings)
     {
         return settings.CurrentValue.GetStorageImplementations().Select(impl => ToStorageImplementation(provider, impl));
     }
 
-    private IBlobStore ToStorageImplementation(IServiceProvider provider, HordeStorageSettings.StorageBackendImplementations impl)
+    private static IBlobStore ToStorageImplementation(IServiceProvider provider, HordeStorageSettings.StorageBackendImplementations impl)
     {
         IBlobStore? store = impl switch
         {
