Horde Storage - Catch and ignore sharing violations when attempting to write blobs, there is only one possible value that a blob could be written as so its safe us to assume the other write will end up writing the correct content.

diff --git a/Engine/Source/Programs/Horde/HordeStorage/Horde.Storage/Implementation/Blob/FileSystemStore.cs b/Engine/Source/Programs/Horde/HordeStorage/Horde.Storage/Implementation/Blob/FileSystemStore.cs
--- a/Engine/Source/Programs/Horde/HordeStorage/Horde.Storage/Implementation/Blob/FileSystemStore.cs
+++ b/Engine/Source/Programs/Horde/HordeStorage/Horde.Storage/Implementation/Blob/FileSystemStore.cs
@@ -21,6 +21,8 @@ namespace Horde.Storage.Implementation
 
         internal const int DefaultBufferSize = 4096;
 
+        private const int SharingViolationErrorCode = 32;
+
         public FileSystemStore(IOptionsMonitor<FilesystemSettings> settings)
         {
             _settings = settings;
@@ -57,12 +59,23 @@ namespace Horde.Storage.Implementation
             FileInfo filePath = GetFilesystemPath(ns, blobIdentifier);
             filePath.Directory?.Create();
 
+            try
             {
                 await using FileStream fs = new FileStream(filePath.FullName, FileMode.Create, FileAccess.Write, FileShare.Read, DefaultBufferSize, FileOptions.Asynchronous | FileOptions.SequentialScan);
                 CancellationToken cancellationToken = default(CancellationToken);
                 await fs.WriteAsync(content, cancellationToken);
                 await fs.FlushAsync(cancellationToken);
             }
+            catch (IOException e)
+            {
+                // ignore errors about accessing the file, it could be used by another call to write it at the same time, which is completely safe as the identifier is the content of the file
+                if (e.HResult == SharingViolationErrorCode)
+                {
+                    _logger.Warning("Sharing violation attempting to write {Blob} in {Namespace}", blobIdentifier, ns);
+                }
+                else
+                    throw;
+            }
 
             if (content.Length == 0)
                 _logger.Warning("0 byte file written as {Blob} {Namespace} {Method}", blobIdentifier, ns, "ReadOnlyMemory");
@@ -77,12 +90,23 @@ namespace Horde.Storage.Implementation
             FileInfo filePath = GetFilesystemPath(ns, blobIdentifier);
             filePath.Directory?.Create();
 
+            try
             {
                 await using FileStream fs = new FileStream(filePath.FullName, FileMode.Create, FileAccess.Write, FileShare.Read, DefaultBufferSize, FileOptions.Asynchronous | FileOptions.SequentialScan);
                 CancellationToken cancellationToken = default(CancellationToken);
                 await content.CopyToAsync(fs, cancellationToken);
                 await fs.FlushAsync(cancellationToken);
             }
+            catch (IOException e)
+            {
+                // ignore errors about accessing the file, it could be used by another call to write it at the same time, which is completely safe as the identifier is the content of the file
+                if (e.HResult == SharingViolationErrorCode)
+                {
+                    _logger.Warning("Sharing violation attempting to write {Blob} in {Namespace}", blobIdentifier, ns);
+                }
+                else
+                    throw;
+            }
             filePath.Refresh();
             if (filePath.Length == 0)
                 _logger.Warning("0 byte file written as {Blob} {Namespace} {Method}", blobIdentifier, ns, "Stream");
@@ -96,12 +120,23 @@ namespace Horde.Storage.Implementation
             FileInfo filePath = GetFilesystemPath(ns, blobIdentifier);
             filePath.Directory?.Create();
 
+            try
             {
                 await using FileStream fs = new FileStream(filePath.FullName, FileMode.Create, FileAccess.Write, FileShare.Read, DefaultBufferSize, FileOptions.Asynchronous | FileOptions.SequentialScan);
                 CancellationToken cancellationToken = default(CancellationToken);
                 await fs.WriteAsync(content, cancellationToken);
                 await fs.FlushAsync(cancellationToken);
             }
+            catch (IOException e)
+            {
+                // ignore errors about accessing the file, it could be used by another call to write it at the same time, which is completely safe as the identifier is the content of the file
+                if (e.HResult == SharingViolationErrorCode)
+                {
+                    _logger.Warning("Sharing violation attempting to write {Blob} in {Namespace}", blobIdentifier, ns);
+                }
+                else
+                    throw;
+            }
 
             if (content.Length == 0)
                 _logger.Warning("0 byte file written as {Blob} {Namespace} {Method}", blobIdentifier, ns, "Array");
