Horde Storage - Exists checks for the compressed blobs endpoints now always returns a needs list if anything is missing, this avoids exposing internal implementation details in how the blobs are stored inside Horde Storage.

diff --git a/Engine/Source/Programs/Horde/HordeStorage/Horde.Storage/Controllers/CompressedBlobController.cs b/Engine/Source/Programs/Horde/HordeStorage/Horde.Storage/Controllers/CompressedBlobController.cs
--- a/Engine/Source/Programs/Horde/HordeStorage/Horde.Storage/Controllers/CompressedBlobController.cs
+++ b/Engine/Source/Programs/Horde/HordeStorage/Horde.Storage/Controllers/CompressedBlobController.cs
@@ -162,10 +162,10 @@ namespace Horde.Storage.Controllers
             });
             await Task.WhenAll(tasks);
 
-            if (invalidContentIds.Count != 0)
-                return NotFound(new ValidationProblemDetails { Title = $"Missing content ids {string.Join(" ,", invalidContentIds)}"});
-
-            return Ok(new ExistCheckMultipleContentIdResponse { Needs = partialContentIds.ToArray()});
+            List<ContentId> needs = new List<ContentId>(invalidContentIds);
+            needs.AddRange(partialContentIds);
+             
+            return Ok(new ExistCheckMultipleContentIdResponse { Needs = needs.ToArray()});
         }
 
         [HttpPost("{ns}/exist")]
@@ -206,10 +206,10 @@ namespace Horde.Storage.Controllers
             });
             await Task.WhenAll(tasks);
 
-            if (invalidContentIds.Count != 0)
-                return NotFound(new ValidationProblemDetails { Title = $"Missing content ids {string.Join(" ,", invalidContentIds)}" });
-
-            return Ok(new ExistCheckMultipleContentIdResponse { Needs = partialContentIds.ToArray() });
+            List<ContentId> needs = new List<ContentId>(invalidContentIds);
+            needs.AddRange(partialContentIds);
+             
+            return Ok(new ExistCheckMultipleContentIdResponse { Needs = needs.ToArray()});
         }
 
         private async Task<(BlobContents, string)> GetImpl(NamespaceId ns, ContentId contentId)
