Horde Storage - Do not track all refs that got deleted as we do not really use that anyway, avoids unnesscary memory build up in prod.

diff --git a/Engine/Source/Programs/Horde/HordeStorage/Horde.Storage/Controllers/AdminController.cs b/Engine/Source/Programs/Horde/HordeStorage/Horde.Storage/Controllers/AdminController.cs
--- a/Engine/Source/Programs/Horde/HordeStorage/Horde.Storage/Controllers/AdminController.cs
+++ b/Engine/Source/Programs/Horde/HordeStorage/Horde.Storage/Controllers/AdminController.cs
@@ -70,10 +70,8 @@ namespace Horde.Storage.Controllers
         [HttpPost("refCleanup/{ns}")]
         public async Task<IActionResult> RefCleanup([FromRoute] [Required] NamespaceId ns)
         {
-            List<OldRecord> records = await _refCleanup.Cleanup(ns, CancellationToken.None);
-            return Ok(new RemovedRefRecordsResponse(
-                records.Select(r => new RemovedRefRecordsResponse.RemovedRecord(r.RefName, r.Bucket))
-            ));
+            int countOfDeletedRecords = await _refCleanup.Cleanup(ns, CancellationToken.None);
+            return Ok(new RemovedRefRecordsResponse(countOfDeletedRecords));
         }
         
         /// <summary>
@@ -135,24 +133,12 @@ namespace Horde.Storage.Controllers
 
     public class RemovedRefRecordsResponse
     {
-        public class RemovedRecord
+        public RemovedRefRecordsResponse(int countOfRemovedRecords)
         {
-            public RemovedRecord(KeyId name, BucketId bucket)
-            {
-                Name = name;
-                Bucket = bucket;
-            }
-
-            public KeyId Name { get; }
-            public BucketId Bucket { get; }
-        }
-
-        public RemovedRefRecordsResponse(IEnumerable<RemovedRecord> removedRecords)
-        {
-            RemovedRecords = removedRecords.ToArray();
+            CountOfRemovedRecords = countOfRemovedRecords;
         }
 
-        public RemovedRecord[] RemovedRecords { get; }
+        public int CountOfRemovedRecords { get; }
     }
 
     public class UpdatedRecordsResponse
diff --git a/Engine/Source/Programs/Horde/HordeStorage/Horde.Storage/Implementation/GC/CacheCleanup.cs b/Engine/Source/Programs/Horde/HordeStorage/Horde.Storage/Implementation/GC/CacheCleanup.cs
--- a/Engine/Source/Programs/Horde/HordeStorage/Horde.Storage/Implementation/GC/CacheCleanup.cs
+++ b/Engine/Source/Programs/Horde/HordeStorage/Horde.Storage/Implementation/GC/CacheCleanup.cs
@@ -16,7 +16,7 @@ namespace Horde.Storage.Implementation
 {
     public interface IRefCleanup
     {
-        Task<List<OldRecord>> Cleanup(NamespaceId ns, CancellationToken cancellationToken);
+        Task<int> Cleanup(NamespaceId ns, CancellationToken cancellationToken);
     }
 
     public class RefCleanup : IRefCleanup
@@ -39,7 +39,7 @@ namespace Horde.Storage.Implementation
             _namespacePolicyResolver = namespacePolicyResolver;
         }
 
-        public Task<List<OldRecord>> Cleanup(NamespaceId ns, CancellationToken cancellationToken)
+        public Task<int> Cleanup(NamespaceId ns, CancellationToken cancellationToken)
         {
             NamespaceSettings.PerNamespaceSettings policies;
             try
@@ -49,7 +49,7 @@ namespace Horde.Storage.Implementation
             catch (UnknownNamespaceException)
             {
                 _logger.Warning("Namespace {Namespace} does not configure any policy, not running ref cleanup on it.", ns);
-                return Task.FromResult(new List<OldRecord>());
+                return Task.FromResult(0);
             }
 
             using IScope scope = Tracer.Instance.StartActive("gc.refs.namespace");
@@ -57,7 +57,7 @@ namespace Horde.Storage.Implementation
             if (ns == INamespacePolicyResolver.JupiterInternalNamespace)
             {
                 // do not apply our cleanup policies to the internal namespace
-                return Task.FromResult(new List<OldRecord>());
+                return Task.FromResult(0);
             }
             else if (!policies.IsLegacyNamespace.HasValue)
             {
@@ -78,9 +78,9 @@ namespace Horde.Storage.Implementation
             }
         }
 
-        private async Task<List<OldRecord>> CleanNamespace(NamespaceId ns, CancellationToken cancellationToken)
+        private async Task<int> CleanNamespace(NamespaceId ns, CancellationToken cancellationToken)
         {
-            List<OldRecord> deletedRecords = new List<OldRecord>();
+            int countOfDeletedRecords = 0;
             DateTime cutoffTime = DateTime.Now.AddSeconds(-1 * _settings.CurrentValue.LastAccessCutoff.TotalSeconds);
             int consideredCount = 0;
             await foreach ((BucketId bucket, IoHashKey name, DateTime lastAccessTime) in _referencesStore.GetRecords(ns).WithCancellation(cancellationToken))
@@ -110,8 +110,7 @@ namespace Horde.Storage.Implementation
 
                 if (await storeDelete)
                 {
-                    // we convert the ObjectRecords key (which is a iohash) to a keyid (which is a generic string) not the prettiest but this result is only used for debugging purposes anyway
-                    deletedRecords.Add(new OldRecord(ns, bucket, new KeyId(name.ToString())));
+                    Interlocked.Increment(ref countOfDeletedRecords);
                 }
                 else
                 {
@@ -119,14 +118,14 @@ namespace Horde.Storage.Implementation
                 }
             }
 
-            _logger.Information("Finished cleaning {Namespace}. Refs considered: {ConsideredCount} Refs Deleted: {DeletedCount}", ns, consideredCount, deletedRecords.Count);
+            _logger.Information("Finished cleaning {Namespace}. Refs considered: {ConsideredCount} Refs Deleted: {DeletedCount}", ns, consideredCount, countOfDeletedRecords);
 
-            return deletedRecords;
+            return countOfDeletedRecords;
         }
 
-        private async Task<List<OldRecord>> CleanNamespaceLegacy(NamespaceId ns, CancellationToken cancellationToken)
+        private async Task<int> CleanNamespaceLegacy(NamespaceId ns, CancellationToken cancellationToken)
         {
-            List<OldRecord> deletedRecords = new List<OldRecord>();
+            int countOfDeletedRecords = 0;
             await foreach (OldRecord record in _refs.GetOldRecords(ns, _settings.CurrentValue.LastAccessCutoff).WithCancellation(cancellationToken))
             {
                 // delete the old record from the ref refs
@@ -142,10 +141,11 @@ namespace Horde.Storage.Implementation
                 {
                     _logger.Warning(e, "Exception when attempting to delete record {Record} in {Namespace}", record, ns);
                 }
-                deletedRecords.Add(record);
+
+                Interlocked.Increment(ref countOfDeletedRecords);
             }
 
-            return deletedRecords;
+            return countOfDeletedRecords;
         }
     }
 }
diff --git a/Engine/Source/Programs/Horde/HordeStorage/Horde.Storage/Implementation/GC/RefCleanupService.cs b/Engine/Source/Programs/Horde/HordeStorage/Horde.Storage/Implementation/GC/RefCleanupService.cs
--- a/Engine/Source/Programs/Horde/HordeStorage/Horde.Storage/Implementation/GC/RefCleanupService.cs
+++ b/Engine/Source/Programs/Horde/HordeStorage/Horde.Storage/Implementation/GC/RefCleanupService.cs
@@ -8,7 +8,6 @@ using Dasync.Collections;
 using Datadog.Trace;
 using EpicGames.Horde.Storage;
 using Jupiter;
-using Jupiter.Implementation;
 using Microsoft.Extensions.Options;
 using Serilog;
 
@@ -69,8 +68,8 @@ namespace Horde.Storage.Implementation
                     _logger.Information("Attempting to run Refs Cleanup of {Namespace}. ", ns);
                     try
                     {
-                        List<OldRecord> oldRecords = await state.RefCleanup.Cleanup(ns, cancellationToken);
-                        _logger.Information("Ran Refs Cleanup of {Namespace}. Deleted {CountRefRecords}", ns, oldRecords.Count);
+                        int countOfRemovedRecords = await state.RefCleanup.Cleanup(ns, cancellationToken);
+                        _logger.Information("Ran Refs Cleanup of {Namespace}. Deleted {CountRefRecords}", ns, countOfRemovedRecords);
                     }
                     catch (Exception e)
                     {
@@ -87,8 +86,8 @@ namespace Horde.Storage.Implementation
                     _logger.Information("Attempting to run Refs Cleanup of {Namespace}. ", ns);
                     try
                     {
-                        List<OldRecord> oldRecords = await state.RefCleanup.Cleanup(ns, cancellationToken);
-                        _logger.Information("Ran Refs Cleanup of {Namespace}. Deleted {CountRefRecords}", ns, oldRecords.Count);
+                        int countOfRemovedRecords = await state.RefCleanup.Cleanup(ns, cancellationToken);
+                        _logger.Information("Ran Refs Cleanup of {Namespace}. Deleted {CountRefRecords}", ns, countOfRemovedRecords);
                     }
                     catch (Exception e)
                     {
diff --git a/Engine/Source/Programs/Horde/HordeStorage/Horde.Storage/Tests/Functional/GC/GCRefTests.cs b/Engine/Source/Programs/Horde/HordeStorage/Horde.Storage/Tests/Functional/GC/GCRefTests.cs
--- a/Engine/Source/Programs/Horde/HordeStorage/Horde.Storage/Tests/Functional/GC/GCRefTests.cs
+++ b/Engine/Source/Programs/Horde/HordeStorage/Horde.Storage/Tests/Functional/GC/GCRefTests.cs
@@ -102,12 +102,7 @@ namespace Horde.Storage.FunctionalTests.GC
                 new KeyId("object5"),
                 new KeyId("object6")
             };
-            Assert.AreEqual(3, removedRefRecords.RemovedRecords.Length);
-            foreach (RemovedRefRecordsResponse.RemovedRecord removedRecord in removedRefRecords.RemovedRecords)
-            {
-                Assert.AreEqual("default", removedRecord.Bucket.ToString());
-                Assert.IsTrue(expectedRemovedRefRecords.Contains(removedRecord.Name));
-            }
+            Assert.AreEqual(3, removedRefRecords.CountOfRemovedRecords);
 
             _refMock!.Verify();
             _refMock.VerifyNoOtherCalls();
@@ -253,12 +248,7 @@ namespace Horde.Storage.FunctionalTests.GC
                 object3Name.ToString(),
                 object6Name.ToString()
             };
-            Assert.AreEqual(4, removedRefRecords.RemovedRecords.Length);
-            foreach (RemovedRefRecordsResponse.RemovedRecord removedRecord in removedRefRecords.RemovedRecords)
-            {
-                Assert.AreEqual(DefaultBucket, removedRecord.Bucket);
-                Assert.IsTrue(expectedRemovedRefRecords.Contains(removedRecord.Name.ToString()));
-            }
+            Assert.AreEqual(4, removedRefRecords.CountOfRemovedRecords);
         }
 
         private (BlobIdentifier, CbObject) GetCBWithAttachment(BlobIdentifier blobIdentifier)
