Horde Storage - Removed top level datadog scopes for long running tasks, this causes a fairly signficant build up of scopes that are not flushed until this long running task is completed. That both uses a lot of memory and makes the monitoring a bit useless as you dont really get the results until its finished. It is useful to be able to group all operations together on when they were triggered though, but not useful enough for the cost.
I did keep the top level scope for replication as we do use that quite frequently to understand what the replicators are doing.

diff --git a/Engine/Source/Programs/Horde/HordeStorage/Horde.Storage/Implementation/Blob/BlobIndexConsistencyCheckService.cs b/Engine/Source/Programs/Horde/HordeStorage/Horde.Storage/Implementation/Blob/BlobIndexConsistencyCheckService.cs
--- a/Engine/Source/Programs/Horde/HordeStorage/Horde.Storage/Implementation/Blob/BlobIndexConsistencyCheckService.cs
+++ b/Engine/Source/Programs/Horde/HordeStorage/Horde.Storage/Implementation/Blob/BlobIndexConsistencyCheckService.cs
@@ -60,8 +60,6 @@ namespace Horde.Storage.Implementation
                 return false;
             }
 
-            using IScope scope = Tracer.Instance.StartActive("blob_index.consistency.poll");
-
             await RunConsistencyCheck();
 
             return true;
@@ -100,6 +98,9 @@ namespace Horde.Storage.Implementation
                 if (NamespaceShouldBeCheckedForConsistency(blobInfo.Namespace))
                     continue;
                 
+                using IScope scope = Tracer.Instance.StartActive("consistency_check.blob_index");
+                scope.Span.ResourceName = $"{blobInfo.Namespace}.{blobInfo.BlobIdentifier}";
+
                 try
                 {
                     if (!blobInfo.Regions.Any())
diff --git a/Engine/Source/Programs/Horde/HordeStorage/Horde.Storage/Implementation/Blob/BlobStoreConsistencyCheckService.cs b/Engine/Source/Programs/Horde/HordeStorage/Horde.Storage/Implementation/Blob/BlobStoreConsistencyCheckService.cs
--- a/Engine/Source/Programs/Horde/HordeStorage/Horde.Storage/Implementation/Blob/BlobStoreConsistencyCheckService.cs
+++ b/Engine/Source/Programs/Horde/HordeStorage/Horde.Storage/Implementation/Blob/BlobStoreConsistencyCheckService.cs
@@ -76,9 +76,6 @@ namespace Horde.Storage.Implementation
             {
                 string blobStoreName = blobStore.GetType().Name;
 
-                using IScope blobStoreScope = Tracer.Instance.StartActive("blob_store.consistency_check.poll");
-                blobStoreScope.Span.ResourceName = blobStoreName;
-
                 bool isRootStore = blobStore is AmazonS3Store or AzureBlobStore;
 
                 List<NamespaceId> namespaces = await _refsStore.GetNamespaces().ToListAsync();
@@ -90,11 +87,12 @@ namespace Horde.Storage.Implementation
                     ulong countOfBlobsChecked = 0;
                     ulong countOfIncorrectBlobsFound = 0;
 
-                    using IScope scope = Tracer.Instance.StartActive("consistency_check.run");
-                    scope.Span.ResourceName = ns.ToString();
-
                     await foreach ((BlobIdentifier blob, DateTime lastModified) in blobStore.ListObjects(ns))
                     {
+                        using IScope scope = Tracer.Instance.StartActive("consistency_check.blob");
+                        scope.Span.ResourceName = $"{ns}.{blob}";
+                        scope.Span.SetTag("BlobStore", blobStoreName);
+
                         if (countOfBlobsChecked % 100 == 0)
                             _logger.Information("Consistency check running on Blob Store {BlobStore}, count of blobs processed so far: {CountOfBlobs}", blobStoreName, countOfBlobsChecked);
                         Interlocked.Increment(ref countOfBlobsChecked);
diff --git a/Engine/Source/Programs/Horde/HordeStorage/Horde.Storage/Implementation/Blob/FileSystemStore.cs b/Engine/Source/Programs/Horde/HordeStorage/Horde.Storage/Implementation/Blob/FileSystemStore.cs
--- a/Engine/Source/Programs/Horde/HordeStorage/Horde.Storage/Implementation/Blob/FileSystemStore.cs
+++ b/Engine/Source/Programs/Horde/HordeStorage/Horde.Storage/Implementation/Blob/FileSystemStore.cs
@@ -7,6 +7,7 @@ using System.Linq;
 using System.Threading;
 using System.Threading.Tasks;
 using Dasync.Collections;
+using Datadog.Trace;
 using EpicGames.Horde.Storage;
 using Jupiter.Implementation;
 using Microsoft.Extensions.Options;
@@ -250,6 +251,8 @@ namespace Horde.Storage.Implementation
         /// <returns></returns>
         public async Task<ulong> CleanupInternal(CancellationToken cancellationToken, int batchSize = 100000)
         {
+            using IScope scope = Tracer.Instance.StartActive("gc.filesystem");
+
             long size = await CalculateDiskSpaceUsed();
             ulong maxSizeBytes = _settings.CurrentValue.MaxSizeBytes;
             long triggerSize = (long) (maxSizeBytes * _settings.CurrentValue.TriggerThresholdPercentage);
diff --git a/Engine/Source/Programs/Horde/HordeStorage/Horde.Storage/Implementation/GC/BlobCleanupService.cs b/Engine/Source/Programs/Horde/HordeStorage/Horde.Storage/Implementation/GC/BlobCleanupService.cs
--- a/Engine/Source/Programs/Horde/HordeStorage/Horde.Storage/Implementation/GC/BlobCleanupService.cs
+++ b/Engine/Source/Programs/Horde/HordeStorage/Horde.Storage/Implementation/GC/BlobCleanupService.cs
@@ -86,8 +86,6 @@ namespace Horde.Storage.Implementation
 
                 string type = blobCleanup.GetType().ToString();
                 _logger.Information("Blob cleanup running for {BlobCleanup}", type);
-                using IScope scope = Tracer.Instance.StartActive("gc.blob");
-                scope.Span.SetTag("type", type);
 
                 _logger.Information("Attempting to run Blob Cleanup {BlobCleanup}. ", type);
                 try
diff --git a/Engine/Source/Programs/Horde/HordeStorage/Horde.Storage/Implementation/GC/CacheCleanup.cs b/Engine/Source/Programs/Horde/HordeStorage/Horde.Storage/Implementation/GC/CacheCleanup.cs
--- a/Engine/Source/Programs/Horde/HordeStorage/Horde.Storage/Implementation/GC/CacheCleanup.cs
+++ b/Engine/Source/Programs/Horde/HordeStorage/Horde.Storage/Implementation/GC/CacheCleanup.cs
@@ -52,8 +52,6 @@ namespace Horde.Storage.Implementation
                 return Task.FromResult(0);
             }
 
-            using IScope scope = Tracer.Instance.StartActive("gc.refs.namespace");
-            scope.Span.ResourceName = ns.ToString();
             if (ns == INamespacePolicyResolver.JupiterInternalNamespace)
             {
                 // do not apply our cleanup policies to the internal namespace
@@ -92,7 +90,7 @@ namespace Horde.Storage.Implementation
                     continue;
 
                 _logger.Information("Attempting to delete object {Namespace} {Bucket} {Name} as it was last updated {LastAccessTime} which is older then {CutoffTime}", ns, bucket, name, lastAccessTime, cutoffTime);
-                using IScope scope = Tracer.Instance.StartActive("refCleanup.delete_record");
+                using IScope scope = Tracer.Instance.StartActive("gc.ref");
                 scope.Span.ResourceName = $"{ns}:{bucket}.{name}";
                 // delete the old record from the ref refs
                 Task<bool> storeDelete = _referencesStore.Delete(ns, bucket, name);
@@ -128,6 +126,9 @@ namespace Horde.Storage.Implementation
             int countOfDeletedRecords = 0;
             await foreach (OldRecord record in _refs.GetOldRecords(ns, _settings.CurrentValue.LastAccessCutoff).WithCancellation(cancellationToken))
             {
+                using IScope scope = Tracer.Instance.StartActive("gc.ref.legacy");
+                scope.Span.ResourceName = $"{ns}:{record.Bucket}.{record.RefName}";
+
                 // delete the old record from the ref refs
                 Task storeDelete = _refs.Delete(record.Namespace, record.Bucket, record.RefName);
                 // insert a delete event into the transaction log
diff --git a/Engine/Source/Programs/Horde/HordeStorage/Horde.Storage/Implementation/GC/OrphanBlobCleanup.cs b/Engine/Source/Programs/Horde/HordeStorage/Horde.Storage/Implementation/GC/OrphanBlobCleanup.cs
--- a/Engine/Source/Programs/Horde/HordeStorage/Horde.Storage/Implementation/GC/OrphanBlobCleanup.cs
+++ b/Engine/Source/Programs/Horde/HordeStorage/Horde.Storage/Implementation/GC/OrphanBlobCleanup.cs
@@ -8,6 +8,7 @@ using System.Threading;
 using System.Threading.Tasks;
 using async_enumerable_dotnet;
 using Dasync.Collections;
+using Datadog.Trace;
 using EpicGames.Horde.Storage;
 using Jupiter;
 using Jupiter.Common;
@@ -90,6 +91,7 @@ namespace Horde.Storage.Implementation
                 return 0;
             }
 
+
             List<NamespaceId> namespaces = await ListNamespaces().Where(NamespaceShouldBeCleaned).ToListAsync();
             ConcurrentDictionary<NamespaceId, GCRootState> perNamespaceRoots = new();
             await namespaces.ParallelForEachAsync(async ns =>
@@ -147,9 +149,13 @@ namespace Horde.Storage.Implementation
                 if (cancellationToken.IsCancellationRequested)
                     break;
                 _logger.Information("Attempting to GC Orphan blob {Blob}", blob);
+
                 bool deleted = true;
                 foreach (NamespaceId ns in namespaces)
                 {
+                    using IScope removeBlobScope = Tracer.Instance.StartActive("gc.blob.legacy");
+                    removeBlobScope.Span.ResourceName = $"{ns}.{blob}";
+
                     try
                     {
                         await _blobService.DeleteObject(ns, blob);
diff --git a/Engine/Source/Programs/Horde/HordeStorage/Horde.Storage/Implementation/GC/OrphanBlobCleanupRefs.cs b/Engine/Source/Programs/Horde/HordeStorage/Horde.Storage/Implementation/GC/OrphanBlobCleanupRefs.cs
--- a/Engine/Source/Programs/Horde/HordeStorage/Horde.Storage/Implementation/GC/OrphanBlobCleanupRefs.cs
+++ b/Engine/Source/Programs/Horde/HordeStorage/Horde.Storage/Implementation/GC/OrphanBlobCleanupRefs.cs
@@ -59,8 +59,7 @@ namespace Horde.Storage.Implementation
             }
 
             List<NamespaceId> namespaces = await ListNamespaces().Where(NamespaceShouldBeCleaned).ToListAsync();
-          
-
+            
             // enumerate all namespaces, and check if the old blob is valid in any of them to allow for a blob store to just store them in a single pile if it wants to
             ulong countOfBlobsRemoved = 0;
             foreach (NamespaceId @namespace in namespaces)
@@ -68,9 +67,6 @@ namespace Horde.Storage.Implementation
                 if (cancellationToken.IsCancellationRequested)
                     break;
 
-                using IScope scope = Tracer.Instance.StartActive("gc.blob.namespace");
-                scope.Span.ResourceName = @namespace.ToString();
-
                 // only consider blobs that have been around for 60 minutes
                 // this due to cases were blobs are uploaded first
                 DateTime cutoff = DateTime.Now.AddMinutes(-60);
@@ -82,6 +78,9 @@ namespace Horde.Storage.Implementation
                     if (cancellationToken.IsCancellationRequested)
                         break;
 
+                    using IScope removeBlobScope = Tracer.Instance.StartActive("gc.blob");
+                    removeBlobScope.Span.ResourceName = $"{@namespace}.{blob}";
+
                     bool found = false;
                     NamespaceSettings.PerNamespaceSettings policy = _namespacePolicyResolver.GetPoliciesForNs(@namespace);
                     
@@ -120,10 +119,13 @@ namespace Horde.Storage.Implementation
                         }
                     }
 
+                    if (cancellationToken.IsCancellationRequested)
+                        break;
+
+                    removeBlobScope.Span.SetTag("removed", (!found).ToString());
+
                     if (!found)
                     {
-                        using IScope removeBlobScope = Tracer.Instance.StartActive("gc.blob.remove-blob");
-                        removeBlobScope.Span.ResourceName = blob.ToString();
                         await RemoveBlob(@namespace, blob);
                         ++countOfBlobsRemoved;
                     }
diff --git a/Engine/Source/Programs/Horde/HordeStorage/Horde.Storage/Implementation/GC/RefCleanupService.cs b/Engine/Source/Programs/Horde/HordeStorage/Horde.Storage/Implementation/GC/RefCleanupService.cs
--- a/Engine/Source/Programs/Horde/HordeStorage/Horde.Storage/Implementation/GC/RefCleanupService.cs
+++ b/Engine/Source/Programs/Horde/HordeStorage/Horde.Storage/Implementation/GC/RefCleanupService.cs
@@ -62,9 +62,6 @@ namespace Horde.Storage.Implementation
                 }
                 await foreach (NamespaceId ns in state.Refs.GetNamespaces().WithCancellation(cancellationToken))
                 {
-                    using IScope scope = Tracer.Instance.StartActive("gc.refs");
-                    scope.Span.ResourceName = ns.ToString();
-
                     _logger.Information("Attempting to run Refs Cleanup of {Namespace}. ", ns);
                     try
                     {
@@ -80,9 +77,6 @@ namespace Horde.Storage.Implementation
                 List<NamespaceId>? namespaces = await _referencesStore.GetNamespaces().ToListAsync(cancellationToken);
                 await foreach (NamespaceId ns in namespaces)
                 {
-                    using IScope scope = Tracer.Instance.StartActive("gc.refs");
-                    scope.Span.ResourceName = ns.ToString();
-
                     _logger.Information("Attempting to run Refs Cleanup of {Namespace}. ", ns);
                     try
                     {
