Horde Storage - Fixed issue were consistency checks were only being run on legacy namespaces (that doesnt need consistency checks).
Also added some more tags to the DD scopes to better filter out operations that are of intrest

diff --git a/Engine/Source/Programs/Horde/HordeStorage/Horde.Storage/Implementation/Blob/BlobIndexConsistencyCheckService.cs b/Engine/Source/Programs/Horde/HordeStorage/Horde.Storage/Implementation/Blob/BlobIndexConsistencyCheckService.cs
--- a/Engine/Source/Programs/Horde/HordeStorage/Horde.Storage/Implementation/Blob/BlobIndexConsistencyCheckService.cs
+++ b/Engine/Source/Programs/Horde/HordeStorage/Horde.Storage/Implementation/Blob/BlobIndexConsistencyCheckService.cs
@@ -95,18 +95,22 @@ namespace Horde.Storage.Implementation
 
 
                 // skip namespace that we shouldn't consistency check, e.g. the legacy namespaces
-                if (NamespaceShouldBeCheckedForConsistency(blobInfo.Namespace))
+                if (!NamespaceShouldBeCheckedForConsistency(blobInfo.Namespace))
                     continue;
                 
                 using IScope scope = Tracer.Instance.StartActive("consistency_check.blob_index");
                 scope.Span.ResourceName = $"{blobInfo.Namespace}.{blobInfo.BlobIdentifier}";
 
+                bool issueFound = false;
+                bool deleted = false;
+
                 try
                 {
                     if (!blobInfo.Regions.Any())
                     {
                         Interlocked.Increment(ref countOfIncorrectBlobsFound);
                         _logger.Warning("Blob {Blob} in namespace {Namespace} is not tracked to exist in any region", blobInfo.BlobIdentifier, blobInfo.Namespace);
+                        issueFound = true;
 
                         if (await _blobService.ExistsInRootStore(blobInfo.Namespace, blobInfo.BlobIdentifier))
                         {
@@ -124,6 +128,7 @@ namespace Horde.Storage.Implementation
                                 _logger.Warning("Blob {Blob} in namespace {Namespace} was removed from the blob index as it didnt exist anywhere", blobInfo.BlobIdentifier, blobInfo.Namespace);
 
                                 await _blobIndex.RemoveBlobFromIndex(blobInfo.Namespace, blobInfo.BlobIdentifier);
+                                deleted = true;
                             }
                         }
                     }
@@ -133,6 +138,7 @@ namespace Horde.Storage.Implementation
                         if (!await _blobService.ExistsInRootStore(blobInfo.Namespace, blobInfo.BlobIdentifier))
                         {
                             Interlocked.Increment(ref countOfIncorrectBlobsFound);
+                            issueFound = true;
                             
                             if (blobInfo.Regions.Count > 1)
                             {
@@ -151,6 +157,7 @@ namespace Horde.Storage.Implementation
                                     {
                                         _logger.Warning("Updating blob index to remove Blob {Blob} in namespace {Namespace} as we failed to repair it.", blobInfo.BlobIdentifier, blobInfo.Namespace);
                                         await _blobIndex.RemoveBlobFromRegion(blobInfo.Namespace, blobInfo.BlobIdentifier);
+                                        deleted = true;
                                     }
                                 }
                             }
@@ -165,6 +172,7 @@ namespace Horde.Storage.Implementation
                                     _logger.Warning("Blob {Blob} in namespace {Namespace} can not be repaired so removing existence from current region.", blobInfo.BlobIdentifier, blobInfo.Namespace);
 
                                     await _blobIndex.RemoveBlobFromRegion(blobInfo.Namespace, blobInfo.BlobIdentifier);
+                                    deleted = true;
                                 }
                             }
                         }
@@ -178,12 +186,16 @@ namespace Horde.Storage.Implementation
 
                         // for entries that are of a unknown namespace we simply remove them
                         await _blobIndex.RemoveBlobFromIndex(blobInfo.Namespace, blobInfo.BlobIdentifier);
+                        deleted = true;
                     }
                 }
                 catch (Exception e)
                 {
                     _logger.Error(e, "Exception when doing blob index consistency check for {Blob} in namespace {Namespace}", blobInfo.BlobIdentifier, blobInfo.Namespace);
                 }
+
+                scope.Span.SetTag("issueFound", issueFound.ToString());
+                scope.Span.SetTag("deleted", deleted.ToString());
             }
 
             _logger.Information("Blob Index Consistency check finished, found {CountOfIncorrectBlobs} incorrect blobs. Processed {CountOfBlobs} blobs.", countOfIncorrectBlobsFound, countOfBlobsChecked);
diff --git a/Engine/Source/Programs/Horde/HordeStorage/Horde.Storage/Implementation/Blob/BlobStoreConsistencyCheckService.cs b/Engine/Source/Programs/Horde/HordeStorage/Horde.Storage/Implementation/Blob/BlobStoreConsistencyCheckService.cs
--- a/Engine/Source/Programs/Horde/HordeStorage/Horde.Storage/Implementation/Blob/BlobStoreConsistencyCheckService.cs
+++ b/Engine/Source/Programs/Horde/HordeStorage/Horde.Storage/Implementation/Blob/BlobStoreConsistencyCheckService.cs
@@ -100,6 +100,7 @@ namespace Horde.Storage.Implementation
                         BlobContents contents = await blobStore.GetObject(ns, blob);
                         await using Stream s = contents.Stream;
 
+                        bool inconsistencyFound = false;
                         BlobIdentifier newHash = await BlobIdentifier.FromStream(s);
                         if (!blob.Equals(newHash))
                         {
@@ -114,6 +115,8 @@ namespace Horde.Storage.Implementation
                                 await _blobIndex.RemoveBlobFromRegion(ns, blob);
                             }
                         }
+
+                        scope.Span.SetTag("deleted", inconsistencyFound.ToString());
                     }
 
                     _logger.Information("Blob Store {BlobStore}: Consistency check finished for {Namespace}, found {CountOfIncorrectBlobs} incorrect blobs. Processed {CountOfBlobs} blobs.", blobStoreName, ns, countOfIncorrectBlobsFound, countOfBlobsChecked);
