Horde: Fix incorrect field being indexed for audit log queries.

diff --git a/Engine/Source/Programs/Horde/Horde.Build/Auditing/Impl/AuditLog.cs b/Engine/Source/Programs/Horde/Horde.Build/Auditing/Impl/AuditLog.cs
--- a/Engine/Source/Programs/Horde/Horde.Build/Auditing/Impl/AuditLog.cs
+++ b/Engine/Source/Programs/Horde/Horde.Build/Auditing/Impl/AuditLog.cs
@@ -107,7 +107,7 @@ namespace Horde.Build.Collections.Impl
 
 			if (!databaseService.ReadOnlyMode)
 			{
-				_messages.Indexes.CreateOne(new CreateIndexModel<AuditLogMessage>(Builders<AuditLogMessage>.IndexKeys.Ascending(x => x.Id).Descending(x => x.TimeUtc)));
+				_messages.Indexes.CreateOne(new CreateIndexModel<AuditLogMessage>(Builders<AuditLogMessage>.IndexKeys.Ascending(x => x.Subject).Descending(x => x.TimeUtc)));
 			}
 
 			_backgroundTask = Task.Run(() => WriteMessagesAsync());
