Horde: Rename DatabaseService to MongoService for clarity, and symmetry with RedisService.

diff --git a/Engine/Source/Programs/Horde/Horde.Build.Tests/ControllerIntegrationTest.cs b/Engine/Source/Programs/Horde/Horde.Build.Tests/ControllerIntegrationTest.cs
--- a/Engine/Source/Programs/Horde/Horde.Build.Tests/ControllerIntegrationTest.cs
+++ b/Engine/Source/Programs/Horde/Horde.Build.Tests/ControllerIntegrationTest.cs
@@ -7,6 +7,7 @@ using System.Threading.Tasks;
 using Horde.Build.Collections;
 using Horde.Build.Collections.Impl;
 using Horde.Build.Jobs;
+using Horde.Build.Server;
 using Horde.Build.Services;
 using Horde.Build.Tests.Stubs.Services;
 using Microsoft.AspNetCore.Hosting;
@@ -78,13 +79,13 @@ public class ControllerIntegrationTest : IDisposable
 	private async Task<Fixture> CreateFixture()
 	{
 		IServiceProvider services = Factory.Services;
-		DatabaseService databaseService = services.GetRequiredService<DatabaseService>();
+		MongoService mongoService = services.GetRequiredService<MongoService>();
 		ITemplateCollection templateService = services.GetRequiredService<ITemplateCollection>();
 		JobService jobService = services.GetRequiredService<JobService>();
 		IArtifactCollection artifactCollection = services.GetRequiredService<IArtifactCollection>();
 		StreamService streamService = services.GetRequiredService<StreamService>();
 		AgentService agentService = services.GetRequiredService<AgentService>();
-		GraphCollection graphCollection = new(databaseService);
+		GraphCollection graphCollection = new(mongoService);
 
 		return await Fixture.Create(graphCollection, templateService, jobService, artifactCollection, streamService,
 			agentService);
diff --git a/Engine/Source/Programs/Horde/Horde.Build.Tests/DatabaseIntegrationTest.cs b/Engine/Source/Programs/Horde/Horde.Build.Tests/DatabaseIntegrationTest.cs
--- a/Engine/Source/Programs/Horde/Horde.Build.Tests/DatabaseIntegrationTest.cs
+++ b/Engine/Source/Programs/Horde/Horde.Build.Tests/DatabaseIntegrationTest.cs
@@ -3,7 +3,7 @@
 using System;
 using System.Collections.Generic;
 using System.Net;
-using Horde.Build.Services;
+using Horde.Build.Server;
 using Horde.Build.Utilities;
 using Microsoft.Extensions.DependencyInjection;
 using Microsoft.Extensions.Logging;
@@ -135,7 +135,7 @@ namespace Horde.Build.Tests
 		private static readonly object s_lockObject = new object();
 
 		private MongoDbInstance? _mongoDbInstance;
-		private DatabaseService? _databaseService;
+		private MongoService? _mongoService;
 		private readonly LoggerFactory _loggerFactory = new LoggerFactory();
 
 		private static RedisRunner? s_redisRunner;
@@ -150,7 +150,7 @@ namespace Horde.Build.Tests
 
 		protected override void ConfigureServices(IServiceCollection services)
 		{
-			services.AddSingleton(GetDatabaseServiceSingleton());
+			services.AddSingleton(GetMongoServiceSingleton());
 			services.AddSingleton(GetRedisDatabase());
 		}
 
@@ -159,17 +159,17 @@ namespace Horde.Build.Tests
 			base.Dispose(disposing);
 
 			_mongoDbInstance?.Dispose();
-			_databaseService?.Dispose();
+			_mongoService?.Dispose();
 			_loggerFactory.Dispose();
 
 			_conMux?.Dispose();
 		}
 
-		public DatabaseService GetDatabaseServiceSingleton()
+		public MongoService GetMongoServiceSingleton()
         {
 			lock(s_lockObject)
 			{
-				if (_databaseService == null)
+				if (_mongoService == null)
 				{
 					_mongoDbInstance = new MongoDbInstance();
 
@@ -177,10 +177,10 @@ namespace Horde.Build.Tests
 					ss.DatabaseName = _mongoDbInstance.DatabaseName;
 					ss.DatabaseConnectionString = _mongoDbInstance.ConnectionString;
 
-					_databaseService = new DatabaseService(Options.Create(ss), _loggerFactory);
+					_mongoService = new MongoService(Options.Create(ss), _loggerFactory);
 				}
 			}
-			return _databaseService;
+			return _mongoService;
         }
 
 		private static RedisRunner GetRedisRunner()
diff --git a/Engine/Source/Programs/Horde/Horde.Build.Tests/Fixture.cs b/Engine/Source/Programs/Horde/Horde.Build.Tests/Fixture.cs
--- a/Engine/Source/Programs/Horde/Horde.Build.Tests/Fixture.cs
+++ b/Engine/Source/Programs/Horde/Horde.Build.Tests/Fixture.cs
@@ -9,6 +9,7 @@ using Horde.Build.Api;
 using Horde.Build.Collections;
 using Horde.Build.Jobs;
 using Horde.Build.Models;
+using Horde.Build.Server;
 using Horde.Build.Services;
 using Horde.Build.Utilities;
 using HordeCommon;
diff --git a/Engine/Source/Programs/Horde/Horde.Build.Tests/Fleet/JobQueueStrategyTest.cs b/Engine/Source/Programs/Horde/Horde.Build.Tests/Fleet/JobQueueStrategyTest.cs
--- a/Engine/Source/Programs/Horde/Horde.Build.Tests/Fleet/JobQueueStrategyTest.cs
+++ b/Engine/Source/Programs/Horde/Horde.Build.Tests/Fleet/JobQueueStrategyTest.cs
@@ -7,6 +7,7 @@ using EpicGames.Core;
 using Horde.Build.Api;
 using Horde.Build.Fleet.Autoscale;
 using Horde.Build.Models;
+using Horde.Build.Server;
 using Horde.Build.Utilities;
 using Microsoft.VisualStudio.TestTools.UnitTesting;
 
diff --git a/Engine/Source/Programs/Horde/Horde.Build.Tests/IssueServiceTests.cs b/Engine/Source/Programs/Horde/Horde.Build.Tests/IssueServiceTests.cs
--- a/Engine/Source/Programs/Horde/Horde.Build.Tests/IssueServiceTests.cs
+++ b/Engine/Source/Programs/Horde/Horde.Build.Tests/IssueServiceTests.cs
@@ -11,6 +11,7 @@ using Horde.Agent.Parser;
 using Horde.Agent.Utility;
 using Horde.Build.IssueHandlers.Impl;
 using Horde.Build.Models;
+using Horde.Build.Server;
 using Horde.Build.Services;
 using Horde.Build.Tests.Stubs.Services;
 using Horde.Build.Utilities;
diff --git a/Engine/Source/Programs/Horde/Horde.Build.Tests/JobServiceTests.cs b/Engine/Source/Programs/Horde/Horde.Build.Tests/JobServiceTests.cs
--- a/Engine/Source/Programs/Horde/Horde.Build.Tests/JobServiceTests.cs
+++ b/Engine/Source/Programs/Horde/Horde.Build.Tests/JobServiceTests.cs
@@ -7,6 +7,7 @@ using System.Threading.Tasks;
 using EpicGames.Core;
 using Horde.Build.Api;
 using Horde.Build.Models;
+using Horde.Build.Server;
 using Horde.Build.Utilities;
 using HordeCommon;
 using Microsoft.VisualStudio.TestTools.UnitTesting;
diff --git a/Engine/Source/Programs/Horde/Horde.Build.Tests/LogFileServiceTests.cs b/Engine/Source/Programs/Horde/Horde.Build.Tests/LogFileServiceTests.cs
--- a/Engine/Source/Programs/Horde/Horde.Build.Tests/LogFileServiceTests.cs
+++ b/Engine/Source/Programs/Horde/Horde.Build.Tests/LogFileServiceTests.cs
@@ -38,7 +38,7 @@ namespace Horde.Build.Tests
 
 		public LogFileServiceTest()
         {
-            LogFileCollection logFileCollection = new LogFileCollection(GetDatabaseServiceSingleton());
+            LogFileCollection logFileCollection = new LogFileCollection(GetMongoServiceSingleton());
 
 			_loggerFactory = new LoggerFactory();
             ILogger<LogFileService> logger = _loggerFactory.CreateLogger<LogFileService>();
@@ -239,7 +239,7 @@ namespace Horde.Build.Tests
         [TestMethod]
         public async Task GetLogFileTest()
         {
-            await GetDatabaseServiceSingleton().Database.DropCollectionAsync("LogFiles");
+            await GetMongoServiceSingleton().Database.DropCollectionAsync("LogFiles");
             Assert.AreEqual(0, (await _logFileService.GetLogFilesAsync()).Count);
 
 			// Will implicitly test GetLogFileAsync(), AddCachedLogFile()
diff --git a/Engine/Source/Programs/Horde/Horde.Build.Tests/LogIndexingTests.cs b/Engine/Source/Programs/Horde/Horde.Build.Tests/LogIndexingTests.cs
--- a/Engine/Source/Programs/Horde/Horde.Build.Tests/LogIndexingTests.cs
+++ b/Engine/Source/Programs/Horde/Horde.Build.Tests/LogIndexingTests.cs
@@ -36,7 +36,7 @@ namespace Horde.Build.Tests
 
 		public LogIndexingTests()
 		{
-			LogFileCollection logFileCollection = new LogFileCollection(GetDatabaseServiceSingleton());
+			LogFileCollection logFileCollection = new LogFileCollection(GetMongoServiceSingleton());
 
 			ServiceProvider serviceProvider = new ServiceCollection()
 				.AddLogging(builder => builder.AddConsole().SetMinimumLevel(LogLevel.Debug))
diff --git a/Engine/Source/Programs/Horde/Horde.Build.Tests/PoolServiceTests.cs b/Engine/Source/Programs/Horde/Horde.Build.Tests/PoolServiceTests.cs
--- a/Engine/Source/Programs/Horde/Horde.Build.Tests/PoolServiceTests.cs
+++ b/Engine/Source/Programs/Horde/Horde.Build.Tests/PoolServiceTests.cs
@@ -46,7 +46,7 @@ namespace Horde.Build.Tests
         [TestMethod]
         public async Task GetPoolsTest()
         {
-            await GetDatabaseServiceSingleton().Database.DropCollectionAsync("Pools");
+            await GetMongoServiceSingleton().Database.DropCollectionAsync("Pools");
             
             List<IPool> pools = await PoolService.GetPoolsAsync();
             Assert.AreEqual(pools.Count, 0);
diff --git a/Engine/Source/Programs/Horde/Horde.Build.Tests/SchedulerTests.cs b/Engine/Source/Programs/Horde/Horde.Build.Tests/SchedulerTests.cs
--- a/Engine/Source/Programs/Horde/Horde.Build.Tests/SchedulerTests.cs
+++ b/Engine/Source/Programs/Horde/Horde.Build.Tests/SchedulerTests.cs
@@ -7,6 +7,7 @@ using System.Threading.Tasks;
 using Horde.Build.Api;
 using Horde.Build.Collections;
 using Horde.Build.Models;
+using Horde.Build.Server;
 using Horde.Build.Services;
 using Horde.Build.Utilities;
 using HordeCommon;
diff --git a/Engine/Source/Programs/Horde/Horde.Build.Tests/ServiceAccountCollectionTests.cs b/Engine/Source/Programs/Horde/Horde.Build.Tests/ServiceAccountCollectionTests.cs
--- a/Engine/Source/Programs/Horde/Horde.Build.Tests/ServiceAccountCollectionTests.cs
+++ b/Engine/Source/Programs/Horde/Horde.Build.Tests/ServiceAccountCollectionTests.cs
@@ -5,7 +5,7 @@ using System.Threading.Tasks;
 using Horde.Build.Collections;
 using Horde.Build.Collections.Impl;
 using Horde.Build.Models;
-using Horde.Build.Services;
+using Horde.Build.Server;
 using Microsoft.VisualStudio.TestTools.UnitTesting;
 using MongoDB.Bson;
 
@@ -19,8 +19,8 @@ namespace Horde.Build.Tests
 		
 		public ServiceAccountCollectionTests()
 		{
-			DatabaseService databaseService = GetDatabaseServiceSingleton();
-			_serviceAccounts = new ServiceAccountCollection(databaseService);
+			MongoService mongoService = GetMongoServiceSingleton();
+			_serviceAccounts = new ServiceAccountCollection(mongoService);
 			_serviceAccount = _serviceAccounts.AddAsync(ObjectId.GenerateNewId().ToString(), new List<string> {"myclaim###myvalue"}, "mydesc").Result;
 		}
 
diff --git a/Engine/Source/Programs/Horde/Horde.Build.Tests/TestSetup.cs b/Engine/Source/Programs/Horde/Horde.Build.Tests/TestSetup.cs
--- a/Engine/Source/Programs/Horde/Horde.Build.Tests/TestSetup.cs
+++ b/Engine/Source/Programs/Horde/Horde.Build.Tests/TestSetup.cs
@@ -20,6 +20,7 @@ using Horde.Build.Logs.Storage.Impl;
 using Horde.Build.Models;
 using Horde.Build.Notifications;
 using Horde.Build.Notifications.Impl;
+using Horde.Build.Server;
 using Horde.Build.Services;
 using Horde.Build.Services.Impl;
 using Horde.Build.Storage;
@@ -69,7 +70,7 @@ namespace Horde.Build.Tests
 		public AclService AclService => ServiceProvider.GetRequiredService<AclService>();
 		public AgentSoftwareService AgentSoftwareService => ServiceProvider.GetRequiredService<AgentSoftwareService>();
 		public AgentService AgentService => ServiceProvider.GetRequiredService<AgentService>();
-		public DatabaseService DatabaseService => ServiceProvider.GetRequiredService<DatabaseService>();
+		public MongoService MongoService => ServiceProvider.GetRequiredService<MongoService>();
 		public ITemplateCollection TemplateCollection => ServiceProvider.GetRequiredService<ITemplateCollection>();
 		internal PerforceServiceStub PerforceService => (PerforceServiceStub)ServiceProvider.GetRequiredService<IPerforceService>();
 		public ILogFileService LogFileService => ServiceProvider.GetRequiredService<ILogFileService>();
diff --git a/Engine/Source/Programs/Horde/Horde.Build.Tests/VersionedCollectionTests.cs b/Engine/Source/Programs/Horde/Horde.Build.Tests/VersionedCollectionTests.cs
--- a/Engine/Source/Programs/Horde/Horde.Build.Tests/VersionedCollectionTests.cs
+++ b/Engine/Source/Programs/Horde/Horde.Build.Tests/VersionedCollectionTests.cs
@@ -1,7 +1,7 @@
 // Copyright Epic Games, Inc. All Rights Reserved.
 
 using System.Threading.Tasks;
-using Horde.Build.Services;
+using Horde.Build.Server;
 using Horde.Build.Tools.Impl;
 using Microsoft.VisualStudio.TestTools.UnitTesting;
 using MongoDB.Driver;
@@ -46,7 +46,7 @@ namespace Horde.Build.Tests
 			public override DocumentV2 UpgradeToLatest() => this;
 		}
 
-		readonly DatabaseService _databaseService;
+		readonly MongoService _mongoService;
 		readonly IDatabase _redis;
 		readonly RedisKey _baseKey = new RedisKey("versioned/");
 		readonly IMongoCollection<VersionedDocument<string, DocumentV2>> _baseCollection;
@@ -54,9 +54,9 @@ namespace Horde.Build.Tests
 
 		public VersionedCollectionTests()
 		{
-			_databaseService = GetDatabaseServiceSingleton();
+			_mongoService = GetMongoServiceSingleton();
 			_redis = GetRedisDatabase();
-			_baseCollection = _databaseService.GetCollection<VersionedDocument<string, DocumentV2>>("versioned");
+			_baseCollection = _mongoService.GetCollection<VersionedDocument<string, DocumentV2>>("versioned");
 			_collection = new VersionedCollection<string, DocumentV2>(_baseCollection, _redis, _baseKey);
 		}
 
diff --git a/Engine/Source/Programs/Horde/Horde.Build/Acls/AclService.cs b/Engine/Source/Programs/Horde/Horde.Build/Acls/AclService.cs
--- a/Engine/Source/Programs/Horde/Horde.Build/Acls/AclService.cs
+++ b/Engine/Source/Programs/Horde/Horde.Build/Acls/AclService.cs
@@ -1,4 +1,4 @@
-﻿// Copyright Epic Games, Inc. All Rights Reserved.
+// Copyright Epic Games, Inc. All Rights Reserved.
 
 using System;
 using System.Collections.Generic;
@@ -7,7 +7,7 @@ using System.Linq;
 using System.Security.Claims;
 using System.Threading.Tasks;
 using Horde.Build.Models;
-using Horde.Build.Services;
+using Horde.Build.Server;
 using Horde.Build.Utilities;
 using Microsoft.Extensions.Options;
 using Microsoft.IdentityModel.Tokens;
@@ -65,16 +65,16 @@ namespace Horde.Build.Acls
 		public static AclClaim AgentRoleClaim { get; } = new AclClaim(HordeClaimTypes.Role, "agent");
 
 		private readonly Acl _defaultAcl = new();
-		private readonly DatabaseService _databaseService;
+		private readonly MongoService _mongoService;
 
 		/// <summary>
 		/// Constructor
 		/// </summary>
-		/// <param name="databaseService">The database service instance</param>
+		/// <param name="mongoService">The database service instance</param>
 		/// <param name="settings">The settings object</param>
-		public AclService(DatabaseService databaseService, IOptionsMonitor<ServerSettings> settings)
+		public AclService(MongoService mongoService, IOptionsMonitor<ServerSettings> settings)
 		{
-			_databaseService = databaseService;
+			_mongoService = mongoService;
 
             List<AclAction> adminActions = new();
 			foreach (AclAction? action in Enum.GetValues(typeof(AclAction)))
@@ -104,7 +104,7 @@ namespace Horde.Build.Acls
 		/// <returns>Scopes instance</returns>
 		public async Task<Acl> GetRootAcl()
 		{
-			Globals globals = await _databaseService.GetGlobalsAsync();
+			Globals globals = await _mongoService.GetGlobalsAsync();
 			return globals.RootAcl ?? new Acl();
 		}
 
@@ -152,9 +152,9 @@ namespace Horde.Build.Acls
 		/// <returns>JWT security token with a claim for creating new agents</returns>
 		public string IssueBearerToken(IEnumerable<Claim> claims, TimeSpan? expiry)
 		{
-			SigningCredentials signingCredentials = new(_databaseService.JwtSigningKey, SecurityAlgorithms.HmacSha256);
+			SigningCredentials signingCredentials = new(_mongoService.JwtSigningKey, SecurityAlgorithms.HmacSha256);
 
-			JwtSecurityToken token = new(_databaseService.JwtIssuer, null, claims, null, DateTime.UtcNow + expiry, signingCredentials);
+			JwtSecurityToken token = new(_mongoService.JwtIssuer, null, claims, null, DateTime.UtcNow + expiry, signingCredentials);
 			return new JwtSecurityTokenHandler().WriteToken(token);
 		}
 
diff --git a/Engine/Source/Programs/Horde/Horde.Build/Agents/Impl/AgentCollection.cs b/Engine/Source/Programs/Horde/Horde.Build/Agents/Impl/AgentCollection.cs
--- a/Engine/Source/Programs/Horde/Horde.Build/Agents/Impl/AgentCollection.cs
+++ b/Engine/Source/Programs/Horde/Horde.Build/Agents/Impl/AgentCollection.cs
@@ -9,6 +9,7 @@ using EpicGames.Redis;
 using Google.Protobuf.WellKnownTypes;
 using Horde.Build.Acls;
 using Horde.Build.Models;
+using Horde.Build.Server;
 using Horde.Build.Services;
 using Horde.Build.Utilities;
 using HordeCommon;
@@ -150,14 +151,14 @@ namespace Horde.Build.Collections.Impl
 		/// <summary>
 		/// Constructor
 		/// </summary>
-		public AgentCollection(DatabaseService databaseService, RedisService redisService, IAuditLog<AgentId> auditLog)
+		public AgentCollection(MongoService mongoService, RedisService redisService, IAuditLog<AgentId> auditLog)
 		{
-			_agents = databaseService.GetCollection<AgentDocument>("Agents");
+			_agents = mongoService.GetCollection<AgentDocument>("Agents");
 			_redisService = redisService;
 			_updateEventChannel = new RedisChannel<AgentId>("agents/notify");
 			_auditLog = auditLog;
 
-			if (!databaseService.ReadOnlyMode)
+			if (!mongoService.ReadOnlyMode)
 			{
 				_agents.Indexes.CreateOne(new CreateIndexModel<AgentDocument>(Builders<AgentDocument>.IndexKeys.Ascending(x => x.Deleted).Ascending(x => x.Id).Ascending(x => x.Pools)));
 			}
diff --git a/Engine/Source/Programs/Horde/Horde.Build/Agents/Impl/AgentService.cs b/Engine/Source/Programs/Horde/Horde.Build/Agents/Impl/AgentService.cs
--- a/Engine/Source/Programs/Horde/Horde.Build/Agents/Impl/AgentService.cs
+++ b/Engine/Source/Programs/Horde/Horde.Build/Agents/Impl/AgentService.cs
@@ -14,6 +14,7 @@ using Google.Protobuf.WellKnownTypes;
 using Horde.Build.Acls;
 using Horde.Build.Collections;
 using Horde.Build.Models;
+using Horde.Build.Server;
 using Horde.Build.Tasks;
 using Horde.Build.Utilities;
 using HordeCommon;
diff --git a/Engine/Source/Programs/Horde/Horde.Build/Agents/Impl/AgentSoftwareCollection.cs b/Engine/Source/Programs/Horde/Horde.Build/Agents/Impl/AgentSoftwareCollection.cs
--- a/Engine/Source/Programs/Horde/Horde.Build/Agents/Impl/AgentSoftwareCollection.cs
+++ b/Engine/Source/Programs/Horde/Horde.Build/Agents/Impl/AgentSoftwareCollection.cs
@@ -4,7 +4,7 @@ using System.Collections.Generic;
 using System.IO;
 using System.Linq;
 using System.Threading.Tasks;
-using Horde.Build.Services;
+using Horde.Build.Server;
 using Horde.Build.Utilities;
 using MongoDB.Bson;
 using MongoDB.Bson.Serialization.Attributes;
@@ -88,11 +88,11 @@ namespace Horde.Build.Collections.Impl
 		/// <summary>
 		/// Constructor
 		/// </summary>
-		/// <param name="databaseService">The database service singleton</param>
-		public AgentSoftwareCollection(DatabaseService databaseService)
+		/// <param name="mongoService">The database service singleton</param>
+		public AgentSoftwareCollection(MongoService mongoService)
 		{
-			_collection = databaseService.GetCollection<AgentSoftwareDocument>("AgentSoftware");
-			_agentChunks = databaseService.GetCollection<ChunkDocument>("AgentSoftwareChunks");
+			_collection = mongoService.GetCollection<AgentSoftwareDocument>("AgentSoftware");
+			_agentChunks = mongoService.GetCollection<ChunkDocument>("AgentSoftwareChunks");
 		}
 
 		/// <inheritdoc/>
diff --git a/Engine/Source/Programs/Horde/Horde.Build/Agents/Impl/LeaseCollection.cs b/Engine/Source/Programs/Horde/Horde.Build/Agents/Impl/LeaseCollection.cs
--- a/Engine/Source/Programs/Horde/Horde.Build/Agents/Impl/LeaseCollection.cs
+++ b/Engine/Source/Programs/Horde/Horde.Build/Agents/Impl/LeaseCollection.cs
@@ -5,7 +5,7 @@ using System.Collections.Generic;
 using System.Linq;
 using System.Threading.Tasks;
 using Horde.Build.Models;
-using Horde.Build.Services;
+using Horde.Build.Server;
 using Horde.Build.Utilities;
 using HordeCommon;
 using MongoDB.Bson;
@@ -100,11 +100,11 @@ namespace Horde.Build.Collections.Impl
 		/// <summary>
 		/// Constructor
 		/// </summary>
-		/// <param name="databaseService">The database service instance</param>
-		public LeaseCollection(DatabaseService databaseService)
+		/// <param name="mongoService">The database service instance</param>
+		public LeaseCollection(MongoService mongoService)
 		{
-			_leases = databaseService.GetCollection<LeaseDocument>("Leases");
-			_indexes = new DatabaseIndexes(_leases, databaseService.ReadOnlyMode);
+			_leases = mongoService.GetCollection<LeaseDocument>("Leases");
+			_indexes = new DatabaseIndexes(_leases, mongoService.ReadOnlyMode);
 		}
 
 		/// <inheritdoc/>
diff --git a/Engine/Source/Programs/Horde/Horde.Build/Agents/Impl/SessionCollection.cs b/Engine/Source/Programs/Horde/Horde.Build/Agents/Impl/SessionCollection.cs
--- a/Engine/Source/Programs/Horde/Horde.Build/Agents/Impl/SessionCollection.cs
+++ b/Engine/Source/Programs/Horde/Horde.Build/Agents/Impl/SessionCollection.cs
@@ -5,7 +5,7 @@ using System.Collections.Generic;
 using System.Linq;
 using System.Threading.Tasks;
 using Horde.Build.Models;
-using Horde.Build.Services;
+using Horde.Build.Server;
 using Horde.Build.Utilities;
 using MongoDB.Bson.Serialization.Attributes;
 using MongoDB.Driver;
@@ -68,12 +68,12 @@ namespace Horde.Build.Collections.Impl
 		/// <summary>
 		/// Constructor
 		/// </summary>
-		/// <param name="databaseService">The database service</param>
-		public SessionCollection(DatabaseService databaseService)
+		/// <param name="mongoService">The database service</param>
+		public SessionCollection(MongoService mongoService)
 		{
-			_sessions = databaseService.GetCollection<SessionDocument>("Sessions");
+			_sessions = mongoService.GetCollection<SessionDocument>("Sessions");
 
-			if (!databaseService.ReadOnlyMode)
+			if (!mongoService.ReadOnlyMode)
 			{
 				_sessions.Indexes.CreateOne(new CreateIndexModel<SessionDocument>(Builders<SessionDocument>.IndexKeys.Ascending(x => x.AgentId)));
 				_sessions.Indexes.CreateOne(new CreateIndexModel<SessionDocument>(Builders<SessionDocument>.IndexKeys.Ascending(x => x.StartTime)));
diff --git a/Engine/Source/Programs/Horde/Horde.Build/Agents/Impl/Tasks/ConformTaskSource.cs b/Engine/Source/Programs/Horde/Horde.Build/Agents/Impl/Tasks/ConformTaskSource.cs
--- a/Engine/Source/Programs/Horde/Horde.Build/Agents/Impl/Tasks/ConformTaskSource.cs
+++ b/Engine/Source/Programs/Horde/Horde.Build/Agents/Impl/Tasks/ConformTaskSource.cs
@@ -10,6 +10,7 @@ using Google.Protobuf.WellKnownTypes;
 using Horde.Build.Api;
 using Horde.Build.Collections;
 using Horde.Build.Models;
+using Horde.Build.Server;
 using Horde.Build.Services;
 using Horde.Build.Utilities;
 using HordeCommon;
@@ -34,7 +35,7 @@ namespace Horde.Build.Tasks.Impl
 		/// <inheritdoc/>
 		public override TaskSourceFlags Flags => TaskSourceFlags.AllowWhenDisabled;
 
-		readonly DatabaseService _databaseService;
+		readonly MongoService _mongoService;
 		readonly IAgentCollection _agentCollection;
 		readonly PoolService _poolService;
 		readonly SingletonDocument<ConformList> _conformList;
@@ -46,12 +47,12 @@ namespace Horde.Build.Tasks.Impl
 		/// <summary>
 		/// Constructor
 		/// </summary>
-		public ConformTaskSource(DatabaseService databaseService, IAgentCollection agentCollection, PoolService poolService, ILogFileService logService, PerforceLoadBalancer perforceLoadBalancer, IClock clock, ILogger<ConformTaskSource> logger)
+		public ConformTaskSource(MongoService mongoService, IAgentCollection agentCollection, PoolService poolService, ILogFileService logService, PerforceLoadBalancer perforceLoadBalancer, IClock clock, ILogger<ConformTaskSource> logger)
 		{
-			_databaseService = databaseService;
+			_mongoService = mongoService;
 			_agentCollection = agentCollection;
 			_poolService = poolService;
-			_conformList = new SingletonDocument<ConformList>(databaseService);
+			_conformList = new SingletonDocument<ConformList>(mongoService);
 			_perforceLoadBalancer = perforceLoadBalancer;
 			_logService = logService;
 			_logger = logger;
@@ -195,7 +196,7 @@ namespace Horde.Build.Tasks.Impl
 		/// <inheritdoc/>
 		public async Task<bool> GetWorkspacesAsync(IAgent agent, IList<HordeCommon.Rpc.Messages.AgentWorkspace> workspaces)
 		{
-			Globals globals = await _databaseService.GetGlobalsAsync();
+			Globals globals = await _mongoService.GetGlobalsAsync();
 
 			HashSet<AgentWorkspace> conformWorkspaces = await _poolService.GetWorkspacesAsync(agent, DateTime.UtcNow);
 			foreach (AgentWorkspace conformWorkspace in conformWorkspaces)
@@ -225,7 +226,7 @@ namespace Horde.Build.Tasks.Impl
 		/// <returns>True if the resource was allocated, false otherwise</returns>
 		private async Task<bool> AllocateConformLeaseAsync(AgentId agentId, IEnumerable<HordeCommon.Rpc.Messages.AgentWorkspace> workspaces, LeaseId leaseId)
 		{
-			Globals globals = await _databaseService.GetGlobalsAsync();
+			Globals globals = await _mongoService.GetGlobalsAsync();
 			for (; ; )
 			{
 				ConformList currentValue = await _conformList.GetAsync();
diff --git a/Engine/Source/Programs/Horde/Horde.Build/Agents/Impl/UpgradeService.cs b/Engine/Source/Programs/Horde/Horde.Build/Agents/Impl/UpgradeService.cs
--- a/Engine/Source/Programs/Horde/Horde.Build/Agents/Impl/UpgradeService.cs
+++ b/Engine/Source/Programs/Horde/Horde.Build/Agents/Impl/UpgradeService.cs
@@ -5,6 +5,7 @@ using System.Threading.Tasks;
 using Horde.Build.Collections;
 using Horde.Build.Collections.Impl;
 using Horde.Build.Models;
+using Horde.Build.Server;
 using Microsoft.Extensions.DependencyInjection;
 using Microsoft.Extensions.Logging;
 
@@ -23,7 +24,7 @@ namespace Horde.Build.Services
 		/// <summary>
 		/// The database service instance
 		/// </summary>
-		DatabaseService DatabaseService { get; set; }
+		readonly MongoService _mongoService;
 
 		/// <summary>
 		/// The DI service provider
@@ -38,12 +39,12 @@ namespace Horde.Build.Services
 		/// <summary>
 		/// Constructor
 		/// </summary>
-		/// <param name="databaseService">The database service singleton</param>
+		/// <param name="mongoService">The database service singleton</param>
 		/// <param name="serviceProvider">The DI service provider</param>
 		/// <param name="logger">Logger instance</param>
-		public UpgradeService(DatabaseService databaseService, IServiceProvider serviceProvider, ILogger<UpgradeService> logger)
+		public UpgradeService(MongoService mongoService, IServiceProvider serviceProvider, ILogger<UpgradeService> logger)
 		{
-			DatabaseService = databaseService;
+			_mongoService = mongoService;
 			_serviceProvider = serviceProvider;
 			_logger = logger;
 		}
@@ -55,7 +56,7 @@ namespace Horde.Build.Services
 		/// <returns>Async task</returns>
 		public async Task UpgradeSchemaAsync(int? fromVersion)
 		{
-			Globals globals = await DatabaseService.GetGlobalsAsync();
+			Globals globals = await _mongoService.GetGlobalsAsync();
 			int schemaVersion = fromVersion ?? globals.SchemaVersion ?? 0;
 
 			while (schemaVersion < LatestSchemaVersion)
@@ -82,8 +83,8 @@ namespace Horde.Build.Services
 				}
 				if (schemaVersion == 4)
 				{
-					UserCollectionV1 userCollectionV1 = new UserCollectionV1(_serviceProvider.GetRequiredService<DatabaseService>());
-					using UserCollectionV2 userCollectionV2 = new UserCollectionV2(_serviceProvider.GetRequiredService<DatabaseService>(), _serviceProvider.GetRequiredService<ILogger<UserCollectionV2>>());
+					UserCollectionV1 userCollectionV1 = new UserCollectionV1(_serviceProvider.GetRequiredService<MongoService>());
+					using UserCollectionV2 userCollectionV2 = new UserCollectionV2(_serviceProvider.GetRequiredService<MongoService>(), _serviceProvider.GetRequiredService<ILogger<UserCollectionV2>>());
 					await userCollectionV2.ResaveDocumentsAsync(userCollectionV1);
 				}
 
@@ -94,11 +95,11 @@ namespace Horde.Build.Services
 				while (globals.SchemaVersion == null || globals.SchemaVersion < schemaVersion)
 				{
 					globals.SchemaVersion = schemaVersion;
-					if (await DatabaseService.TryUpdateSingletonAsync(globals))
+					if (await _mongoService.TryUpdateSingletonAsync(globals))
 					{
 						break;
 					}
-					globals = await DatabaseService.GetGlobalsAsync();
+					globals = await _mongoService.GetGlobalsAsync();
 				}
 			}
 		}
diff --git a/Engine/Source/Programs/Horde/Horde.Build/Auditing/Impl/AuditLog.cs b/Engine/Source/Programs/Horde/Horde.Build/Auditing/Impl/AuditLog.cs
--- a/Engine/Source/Programs/Horde/Horde.Build/Auditing/Impl/AuditLog.cs
+++ b/Engine/Source/Programs/Horde/Horde.Build/Auditing/Impl/AuditLog.cs
@@ -5,7 +5,7 @@ using System.Collections.Generic;
 using System.Threading.Channels;
 using System.Threading.Tasks;
 using EpicGames.Core;
-using Horde.Build.Services;
+using Horde.Build.Server;
 using Horde.Build.Utilities;
 using Microsoft.Extensions.Logging;
 using MongoDB.Bson;
@@ -98,14 +98,14 @@ namespace Horde.Build.Collections.Impl
 
 		public IAuditLogChannel<TSubject> this[TSubject subject] => new AuditLogChannel(this, subject);
 
-		public AuditLog(DatabaseService databaseService, string collectionName, string subjectProperty, ILogger logger)
+		public AuditLog(MongoService mongoService, string collectionName, string subjectProperty, ILogger logger)
 		{
-			_messages = databaseService.GetCollection<AuditLogMessage>(collectionName);
+			_messages = mongoService.GetCollection<AuditLogMessage>(collectionName);
 			_messageChannel = Channel.CreateUnbounded<AuditLogMessage>();
 			_subjectProperty = subjectProperty;
 			_logger = logger;
 
-			if (!databaseService.ReadOnlyMode)
+			if (!mongoService.ReadOnlyMode)
 			{
 				_messages.Indexes.CreateOne(new CreateIndexModel<AuditLogMessage>(Builders<AuditLogMessage>.IndexKeys.Ascending(x => x.Subject).Descending(x => x.TimeUtc)));
 			}
@@ -183,18 +183,18 @@ namespace Horde.Build.Collections.Impl
 
 	class AuditLogFactory<TSubject> : IAuditLogFactory<TSubject>
 	{
-		readonly DatabaseService _databaseService;
+		readonly MongoService _mongoService;
 		readonly ILogger<AuditLog<TSubject>> _logger;
 
-		public AuditLogFactory(DatabaseService databaseService, ILogger<AuditLog<TSubject>> logger)
+		public AuditLogFactory(MongoService mongoService, ILogger<AuditLog<TSubject>> logger)
 		{
-			_databaseService = databaseService;
+			_mongoService = mongoService;
 			_logger = logger;
 		}
 
 		public IAuditLog<TSubject> Create(string collectionName, string subjectProperty)
 		{
-			return new AuditLog<TSubject>(_databaseService, collectionName, subjectProperty, _logger);
+			return new AuditLog<TSubject>(_mongoService, collectionName, subjectProperty, _logger);
 		}
 	}
 }
diff --git a/Engine/Source/Programs/Horde/Horde.Build/Authentication/HordeJwtBearerHandler.cs b/Engine/Source/Programs/Horde/Horde.Build/Authentication/HordeJwtBearerHandler.cs
--- a/Engine/Source/Programs/Horde/Horde.Build/Authentication/HordeJwtBearerHandler.cs
+++ b/Engine/Source/Programs/Horde/Horde.Build/Authentication/HordeJwtBearerHandler.cs
@@ -4,7 +4,7 @@ using System;
 using System.IdentityModel.Tokens.Jwt;
 using System.Text.Encodings.Web;
 using System.Threading.Tasks;
-using Horde.Build.Services;
+using Horde.Build.Server;
 using Horde.Build.Utilities;
 using Microsoft.AspNetCore.Authentication;
 using Microsoft.AspNetCore.Authentication.JwtBearer;
@@ -22,33 +22,33 @@ namespace Horde.Build.Authentication
 		/// Default name of the authentication scheme
 		/// </summary>
 		public const string AuthenticationScheme = "ServerJwt";
-		readonly DatabaseService _databaseService;
+		readonly MongoService _mongoService;
 
-		public HordeJwtBearerHandler(ILoggerFactory logger, UrlEncoder encoder, ISystemClock clock, DatabaseService databaseService, IOptionsMonitorCache<JwtBearerOptions> optionsCache)
-			: base(GetOptionsMonitor(databaseService, optionsCache), logger, encoder, clock)
+		public HordeJwtBearerHandler(ILoggerFactory logger, UrlEncoder encoder, ISystemClock clock, MongoService mongoService, IOptionsMonitorCache<JwtBearerOptions> optionsCache)
+			: base(GetOptionsMonitor(mongoService, optionsCache), logger, encoder, clock)
 		{
-			_databaseService = databaseService;
+			_mongoService = mongoService;
 		}
 
-		private static IOptionsMonitor<JwtBearerOptions> GetOptionsMonitor(DatabaseService databaseService, IOptionsMonitorCache<JwtBearerOptions> optionsCache)
+		private static IOptionsMonitor<JwtBearerOptions> GetOptionsMonitor(MongoService mongoService, IOptionsMonitorCache<JwtBearerOptions> optionsCache)
 		{
-			ConfigureNamedOptions<JwtBearerOptions> namedOptions = new ConfigureNamedOptions<JwtBearerOptions>(AuthenticationScheme, options => Configure(options, databaseService));
+			ConfigureNamedOptions<JwtBearerOptions> namedOptions = new ConfigureNamedOptions<JwtBearerOptions>(AuthenticationScheme, options => Configure(options, mongoService));
 			OptionsFactory<JwtBearerOptions> optionsFactory = new OptionsFactory<JwtBearerOptions>(new[] { namedOptions }, Array.Empty<IPostConfigureOptions<JwtBearerOptions>>());
 			return new OptionsMonitor<JwtBearerOptions>(optionsFactory, Array.Empty<IOptionsChangeTokenSource<JwtBearerOptions>>(), optionsCache);
 		}
 
-		private static void Configure(JwtBearerOptions options, DatabaseService databaseService)
+		private static void Configure(JwtBearerOptions options, MongoService mongoService)
 		{
 			options.TokenValidationParameters.ValidateAudience = false;
 
 			options.TokenValidationParameters.RequireExpirationTime = false;
 			options.TokenValidationParameters.ValidateLifetime = true;
 
-			options.TokenValidationParameters.ValidIssuer = databaseService.JwtIssuer;
+			options.TokenValidationParameters.ValidIssuer = mongoService.JwtIssuer;
 			options.TokenValidationParameters.ValidateIssuer = true;
 
 			options.TokenValidationParameters.ValidateIssuerSigningKey = true;
-			options.TokenValidationParameters.IssuerSigningKey = databaseService.JwtSigningKey;
+			options.TokenValidationParameters.IssuerSigningKey = mongoService.JwtSigningKey;
 		}
 
 		protected override Task<AuthenticateResult> HandleAuthenticateAsync()
@@ -62,7 +62,7 @@ namespace Horde.Build.Authentication
 
 			// Validate that it's from the correct issuer
 			JwtSecurityToken? jwtToken;
-			if (!JwtUtils.TryParseJwt(token, out jwtToken) || !String.Equals(jwtToken.Issuer, _databaseService.JwtIssuer, StringComparison.Ordinal))
+			if (!JwtUtils.TryParseJwt(token, out jwtToken) || !String.Equals(jwtToken.Issuer, _mongoService.JwtIssuer, StringComparison.Ordinal))
 			{
 				return Task.FromResult(AuthenticateResult.NoResult());
 			}
diff --git a/Engine/Source/Programs/Horde/Horde.Build/Compute/Impl/ComputeService.cs b/Engine/Source/Programs/Horde/Horde.Build/Compute/Impl/ComputeService.cs
--- a/Engine/Source/Programs/Horde/Horde.Build/Compute/Impl/ComputeService.cs
+++ b/Engine/Source/Programs/Horde/Horde.Build/Compute/Impl/ComputeService.cs
@@ -13,7 +13,7 @@ using EpicGames.Serialization;
 using Google.Protobuf;
 using Google.Protobuf.WellKnownTypes;
 using Horde.Build.Models;
-using Horde.Build.Services;
+using Horde.Build.Server;
 using Horde.Build.Tasks;
 using Horde.Build.Utilities;
 using HordeCommon;
@@ -183,14 +183,14 @@ namespace Horde.Build.Compute.Impl
 		/// <summary>
 		/// Constructor
 		/// </summary>
-		public ComputeService(DatabaseService databaseService, IDatabase redis, IStorageClient storageClient, IClock clock, ILogger<ComputeService> logger)
+		public ComputeService(MongoService mongoService, RedisService redisService, IStorageClient storageClient, IClock clock, ILogger<ComputeService> logger)
 		{
 			_storageClient = storageClient;
-			_taskScheduler = new RedisTaskScheduler<QueueKey, ComputeTaskInfo>(redis, "compute/tasks/", logger);
-			_messageQueue = new RedisMessageQueue<ComputeTaskStatus>(redis, "compute/messages/");
+			_taskScheduler = new RedisTaskScheduler<QueueKey, ComputeTaskInfo>(redisService.Database, "compute/tasks/", logger);
+			_messageQueue = new RedisMessageQueue<ComputeTaskStatus>(redisService.Database, "compute/messages/");
 			_expireTasksTicker = clock.AddTicker<ComputeService>(TimeSpan.FromMinutes(2.0), ExpireTasksAsync, logger);
 			_requirementsCache = new MemoryCache(new MemoryCacheOptions());
-			_globals = new LazyCachedValue<Task<Globals>>(() => databaseService.GetGlobalsAsync(), TimeSpan.FromSeconds(120.0));
+			_globals = new LazyCachedValue<Task<Globals>>(() => mongoService.GetGlobalsAsync(), TimeSpan.FromSeconds(120.0));
 			_logger = logger;
 
 			OnLeaseStartedProperties.Add(x => x.TaskRefId);
diff --git a/Engine/Source/Programs/Horde/Horde.Build/Controllers/ArtifactsController.cs b/Engine/Source/Programs/Horde/Horde.Build/Controllers/ArtifactsController.cs
--- a/Engine/Source/Programs/Horde/Horde.Build/Controllers/ArtifactsController.cs
+++ b/Engine/Source/Programs/Horde/Horde.Build/Controllers/ArtifactsController.cs
@@ -13,6 +13,7 @@ using Horde.Build.Acls;
 using Horde.Build.Api;
 using Horde.Build.Jobs;
 using Horde.Build.Models;
+using Horde.Build.Server;
 using Horde.Build.Services;
 using Horde.Build.Utilities;
 using Microsoft.AspNetCore.Authorization;
@@ -35,7 +36,7 @@ namespace Horde.Build.Controllers
 		/// <summary>
 		/// Instance of the database service
 		/// </summary>
-		private readonly DatabaseService _databaseService;
+		private readonly MongoService _mongoService;
 
 		/// <summary>
 		/// Instance of the artifact collection
@@ -55,9 +56,9 @@ namespace Horde.Build.Controllers
 		/// <summary>
 		/// Constructor
 		/// </summary>
-		public ArtifactsController(DatabaseService databaseSerivce, IArtifactCollection artifactCollection, AclService aclService, JobService jobService)
+		public ArtifactsController(MongoService mongoService, IArtifactCollection artifactCollection, AclService aclService, JobService jobService)
 		{
-			_databaseService = databaseSerivce;
+			_mongoService = mongoService;
 			_artifactCollection = artifactCollection;
 			_aclService = aclService;
 			_jobService = jobService;
@@ -306,10 +307,10 @@ namespace Horde.Build.Controllers
 			parameters.ValidateAudience = false;
 			parameters.RequireExpirationTime = true;
 			parameters.ValidateLifetime = true;
-			parameters.ValidIssuer = _databaseService.JwtIssuer;
+			parameters.ValidIssuer = _mongoService.JwtIssuer;
 			parameters.ValidateIssuer = true;
 			parameters.ValidateIssuerSigningKey = true;
-			parameters.IssuerSigningKey = _databaseService.JwtSigningKey;
+			parameters.IssuerSigningKey = _mongoService.JwtSigningKey;
 
 			JwtSecurityTokenHandler handler = new JwtSecurityTokenHandler();
 			ClaimsPrincipal principal = handler.ValidateToken(code, parameters, out _);
diff --git a/Engine/Source/Programs/Horde/Horde.Build/Controllers/ConfigController.cs b/Engine/Source/Programs/Horde/Horde.Build/Controllers/ConfigController.cs
--- a/Engine/Source/Programs/Horde/Horde.Build/Controllers/ConfigController.cs
+++ b/Engine/Source/Programs/Horde/Horde.Build/Controllers/ConfigController.cs
@@ -6,8 +6,7 @@ using System.Threading.Tasks;
 using EpicGames.Core;
 using Horde.Build.Acls;
 using Horde.Build.Api;
-using Horde.Build.Models;
-using Horde.Build.Services;
+using Horde.Build.Server;
 using Horde.Build.Utilities;
 using Microsoft.AspNetCore.Authorization;
 using Microsoft.AspNetCore.Mvc;
diff --git a/Engine/Source/Programs/Horde/Horde.Build/Controllers/LifetimeController.cs b/Engine/Source/Programs/Horde/Horde.Build/Controllers/LifetimeController.cs
--- a/Engine/Source/Programs/Horde/Horde.Build/Controllers/LifetimeController.cs
+++ b/Engine/Source/Programs/Horde/Horde.Build/Controllers/LifetimeController.cs
@@ -1,7 +1,7 @@
 // Copyright Epic Games, Inc. All Rights Reserved.
 
 using System.Threading.Tasks;
-using Horde.Build.Services;
+using Horde.Build.Server;
 using Microsoft.AspNetCore.Mvc;
 
 namespace Horde.Build.Controllers
diff --git a/Engine/Source/Programs/Horde/Horde.Build/Controllers/PerforceController.cs b/Engine/Source/Programs/Horde/Horde.Build/Controllers/PerforceController.cs
--- a/Engine/Source/Programs/Horde/Horde.Build/Controllers/PerforceController.cs
+++ b/Engine/Source/Programs/Horde/Horde.Build/Controllers/PerforceController.cs
@@ -1,10 +1,11 @@
-﻿// Copyright Epic Games, Inc. All Rights Reserved.
+// Copyright Epic Games, Inc. All Rights Reserved.
 
 using System.Collections.Generic;
 using System.Net;
 using System.Threading.Tasks;
 using Horde.Build.Acls;
 using Horde.Build.Models;
+using Horde.Build.Server;
 using Horde.Build.Services;
 using Microsoft.AspNetCore.Authorization;
 using Microsoft.AspNetCore.Mvc;
@@ -23,7 +24,7 @@ namespace Horde.Build.Controllers
 		/// <summary>
 		/// The database service instance
 		/// </summary>
-		private readonly DatabaseService _databaseService;
+		private readonly MongoService _mongoService;
 
 		/// <summary>
 		/// The ACL service instance
@@ -38,9 +39,9 @@ namespace Horde.Build.Controllers
 		/// <summary>
 		/// Constructor
 		/// </summary>
-		public PerforceController(DatabaseService databaseService, AclService aclService, PerforceLoadBalancer perforceLoadBalancer)
+		public PerforceController(MongoService mongoService, AclService aclService, PerforceLoadBalancer perforceLoadBalancer)
 		{
-			_databaseService = databaseService;
+			_mongoService = mongoService;
 			_aclService = aclService;
 			_perforceLoadBalancer = perforceLoadBalancer;
 		}
@@ -76,7 +77,7 @@ namespace Horde.Build.Controllers
 				return Forbid();
 			}
 
-			Globals globals = await _databaseService.GetGlobalsAsync();
+			Globals globals = await _mongoService.GetGlobalsAsync();
 			return globals.PerforceClusters;
 		}
 	}
diff --git a/Engine/Source/Programs/Horde/Horde.Build/Devices/Impl/DeviceCollection.cs b/Engine/Source/Programs/Horde/Horde.Build/Devices/Impl/DeviceCollection.cs
--- a/Engine/Source/Programs/Horde/Horde.Build/Devices/Impl/DeviceCollection.cs
+++ b/Engine/Source/Programs/Horde/Horde.Build/Devices/Impl/DeviceCollection.cs
@@ -6,7 +6,7 @@ using System.Linq;
 using System.Threading.Tasks;
 using Horde.Build.Acls;
 using Horde.Build.Models;
-using Horde.Build.Services;
+using Horde.Build.Server;
 using Horde.Build.Utilities;
 using MongoDB.Bson;
 using MongoDB.Bson.Serialization.Attributes;
@@ -248,14 +248,14 @@ namespace Horde.Build.Collections.Impl
 		/// <summary>
 		/// Constructor
 		/// </summary>
-		public DeviceCollection(DatabaseService databaseService)
+		public DeviceCollection(MongoService mongoService)
 		{
-			_devices = databaseService.GetCollection<DeviceDocument>("Devices");
-			_platforms = databaseService.GetCollection<DevicePlatformDocument>("Devices.Platforms");
-			_pools = databaseService.GetCollection<DevicePoolDocument>("Devices.Pools");
-			_reservations = databaseService.GetCollection<DeviceReservationDocument>("Devices.Reservations");
+			_devices = mongoService.GetCollection<DeviceDocument>("Devices");
+			_platforms = mongoService.GetCollection<DevicePlatformDocument>("Devices.Platforms");
+			_pools = mongoService.GetCollection<DevicePoolDocument>("Devices.Pools");
+			_reservations = mongoService.GetCollection<DeviceReservationDocument>("Devices.Reservations");
 
-			if (!databaseService.ReadOnlyMode)
+			if (!mongoService.ReadOnlyMode)
 			{
 				_platforms.Indexes.CreateOne(new CreateIndexModel<DevicePlatformDocument>(Builders<DevicePlatformDocument>.IndexKeys.Ascending(x => x.Name), new CreateIndexOptions { Unique = true }));
 				_pools.Indexes.CreateOne(new CreateIndexModel<DevicePoolDocument>(Builders<DevicePoolDocument>.IndexKeys.Ascending(x => x.Name), new CreateIndexOptions { Unique = true }));
diff --git a/Engine/Source/Programs/Horde/Horde.Build/Fleet/Autoscale/AutoscaleService.cs b/Engine/Source/Programs/Horde/Horde.Build/Fleet/Autoscale/AutoscaleService.cs
--- a/Engine/Source/Programs/Horde/Horde.Build/Fleet/Autoscale/AutoscaleService.cs
+++ b/Engine/Source/Programs/Horde/Horde.Build/Fleet/Autoscale/AutoscaleService.cs
@@ -9,6 +9,7 @@ using System.Threading.Tasks;
 using Google.Protobuf.WellKnownTypes;
 using Horde.Build.Collections;
 using Horde.Build.Models;
+using Horde.Build.Server;
 using Horde.Build.Utilities;
 using HordeCommon;
 using HordeCommon.Rpc.Tasks;
@@ -96,7 +97,7 @@ namespace Horde.Build.Services
 		/// <summary>
 		/// Constructor
 		/// </summary>
-		public AutoscaleService(IAgentCollection agentCollection, IPoolCollection poolCollection, ILeaseCollection leaseCollection, IFleetManager fleetManager, IDogStatsd dogStatsd, DatabaseService databaseService, IClock clock, ILogger<AutoscaleService> logger)
+		public AutoscaleService(IAgentCollection agentCollection, IPoolCollection poolCollection, ILeaseCollection leaseCollection, IFleetManager fleetManager, IDogStatsd dogStatsd, MongoService mongoService, IClock clock, ILogger<AutoscaleService> logger)
 		{
 			_agentCollection = agentCollection;
 			_poolCollection = poolCollection;
@@ -104,7 +105,7 @@ namespace Horde.Build.Services
 			_fleetManager = fleetManager;
 			_dogStatsd = dogStatsd;
 			_clock = clock;
-			if (databaseService.ReadOnlyMode)
+			if (mongoService.ReadOnlyMode)
 			{
 				_ticker = new NullTicker();
 			}
diff --git a/Engine/Source/Programs/Horde/Horde.Build/Fleet/Impl/PoolCollection.cs b/Engine/Source/Programs/Horde/Horde.Build/Fleet/Impl/PoolCollection.cs
--- a/Engine/Source/Programs/Horde/Horde.Build/Fleet/Impl/PoolCollection.cs
+++ b/Engine/Source/Programs/Horde/Horde.Build/Fleet/Impl/PoolCollection.cs
@@ -7,7 +7,7 @@ using System.Threading.Tasks;
 using EpicGames.Horde.Common;
 using Horde.Build.Fleet.Autoscale;
 using Horde.Build.Models;
-using Horde.Build.Services;
+using Horde.Build.Server;
 using Horde.Build.Utilities;
 using MongoDB.Bson;
 using MongoDB.Bson.Serialization.Attributes;
@@ -84,10 +84,10 @@ namespace Horde.Build.Collections.Impl
 		/// <summary>
 		/// Constructor
 		/// </summary>
-		/// <param name="databaseService">The database service instance</param>
-		public PoolCollection(DatabaseService databaseService)
+		/// <param name="mongoService">The database service instance</param>
+		public PoolCollection(MongoService mongoService)
 		{
-			_pools = databaseService.GetCollection<PoolDocument>("Pools");
+			_pools = mongoService.GetCollection<PoolDocument>("Pools");
 		}
 
 		/// <inheritdoc/>
diff --git a/Engine/Source/Programs/Horde/Horde.Build/Fleet/Impl/TelemetryCollection.cs b/Engine/Source/Programs/Horde/Horde.Build/Fleet/Impl/TelemetryCollection.cs
--- a/Engine/Source/Programs/Horde/Horde.Build/Fleet/Impl/TelemetryCollection.cs
+++ b/Engine/Source/Programs/Horde/Horde.Build/Fleet/Impl/TelemetryCollection.cs
@@ -6,7 +6,7 @@ using System.Linq;
 using System.Threading.Tasks;
 using EpicGames.Core;
 using Horde.Build.Models;
-using Horde.Build.Services;
+using Horde.Build.Server;
 using Horde.Build.Utilities;
 using MongoDB.Bson.Serialization.Attributes;
 using MongoDB.Driver;
@@ -103,7 +103,7 @@ namespace Horde.Build.Collections.Impl
 		/// Constructor
 		/// </summary>
 		/// <param name="database"></param>
-		public TelemetryCollection(DatabaseService database)
+		public TelemetryCollection(MongoService database)
 		{
 			_utilization = database.GetCollection<UtilizationDocument>("Utilization");
 
diff --git a/Engine/Source/Programs/Horde/Horde.Build/Fleet/PoolService.cs b/Engine/Source/Programs/Horde/Horde.Build/Fleet/PoolService.cs
--- a/Engine/Source/Programs/Horde/Horde.Build/Fleet/PoolService.cs
+++ b/Engine/Source/Programs/Horde/Horde.Build/Fleet/PoolService.cs
@@ -9,6 +9,7 @@ using EpicGames.Horde.Common;
 using Horde.Build.Collections;
 using Horde.Build.Fleet.Autoscale;
 using Horde.Build.Models;
+using Horde.Build.Server;
 using Horde.Build.Utilities;
 using HordeCommon;
 
@@ -24,7 +25,7 @@ namespace Horde.Build.Services
 		/// <summary>
 		/// The database service instance
 		/// </summary>
-		readonly DatabaseService _databaseService;
+		readonly MongoService _mongoService;
 
 		/// <summary>
 		/// Collection of pool documents
@@ -44,12 +45,12 @@ namespace Horde.Build.Services
 		/// <summary>
 		/// Constructor
 		/// </summary>
-		/// <param name="databaseService"></param>
+		/// <param name="mongoService"></param>
 		/// <param name="pools">Collection of pool documents</param>
 		/// <param name="clock"></param>
-		public PoolService(DatabaseService databaseService, IPoolCollection pools, IClock clock)
+		public PoolService(MongoService mongoService, IPoolCollection pools, IClock clock)
 		{
-			_databaseService = databaseService;
+			_mongoService = mongoService;
 			_pools = pools;
 			_clock = clock;
 		}
@@ -198,7 +199,7 @@ namespace Horde.Build.Services
 
 			if (bAddAutoSdkWorkspace)
 			{
-				Globals globals = await _databaseService.GetGlobalsAsync();
+				Globals globals = await _mongoService.GetGlobalsAsync();
 				workspaces.UnionWith(agent.GetAutoSdkWorkspaces(globals, workspaces.ToList()));
 			}
 
diff --git a/Engine/Source/Programs/Horde/Horde.Build/Issues/Impl/IssueCollection.cs b/Engine/Source/Programs/Horde/Horde.Build/Issues/Impl/IssueCollection.cs
--- a/Engine/Source/Programs/Horde/Horde.Build/Issues/Impl/IssueCollection.cs
+++ b/Engine/Source/Programs/Horde/Horde.Build/Issues/Impl/IssueCollection.cs
@@ -7,7 +7,7 @@ using System.Linq;
 using System.Threading.Tasks;
 using EpicGames.Redis.Utility;
 using Horde.Build.Models;
-using Horde.Build.Services;
+using Horde.Build.Server;
 using Horde.Build.Utilities;
 using Microsoft.Extensions.Caching.Memory;
 using Microsoft.Extensions.Logging;
@@ -395,20 +395,20 @@ namespace Horde.Build.Collections.Impl
 		readonly IAuditLog<int> _auditLog;
 		readonly ILogger _logger;
 
-		public IssueCollection(DatabaseService databaseService, RedisService redisService, IAuditLogFactory<int> auditLogFactory, ILogger<IssueCollection> logger)
+		public IssueCollection(MongoService mongoService, RedisService redisService, IAuditLogFactory<int> auditLogFactory, ILogger<IssueCollection> logger)
 		{
 			_redisService = redisService;
 			_logger = logger;
 
-			_ledgerSingleton = new SingletonDocument<IssueLedger>(databaseService);
+			_ledgerSingleton = new SingletonDocument<IssueLedger>(mongoService);
 
-			_issues = databaseService.GetCollection<Issue>("IssuesV2");
-			_issueSpans = databaseService.GetCollection<IssueSpan>("IssuesV2.Spans");
-			_issueSteps = databaseService.GetCollection<IssueStep>("IssuesV2.Steps");
-			_issueSuspects = databaseService.GetCollection<IssueSuspect>("IssuesV2.Suspects");
+			_issues = mongoService.GetCollection<Issue>("IssuesV2");
+			_issueSpans = mongoService.GetCollection<IssueSpan>("IssuesV2.Spans");
+			_issueSteps = mongoService.GetCollection<IssueStep>("IssuesV2.Steps");
+			_issueSuspects = mongoService.GetCollection<IssueSuspect>("IssuesV2.Suspects");
 			_auditLog = auditLogFactory.Create("IssuesV2.History", "IssueId");
 
-			if (!databaseService.ReadOnlyMode)
+			if (!mongoService.ReadOnlyMode)
 			{
 				_issues.Indexes.CreateOne(new CreateIndexModel<Issue>(Builders<Issue>.IndexKeys.Ascending(x => x.ResolvedAt)));
 				_issues.Indexes.CreateOne(new CreateIndexModel<Issue>(Builders<Issue>.IndexKeys.Ascending(x => x.VerifiedAt)));
diff --git a/Engine/Source/Programs/Horde/Horde.Build/Jobs/Impl/ArtifactCollection.cs b/Engine/Source/Programs/Horde/Horde.Build/Jobs/Impl/ArtifactCollection.cs
--- a/Engine/Source/Programs/Horde/Horde.Build/Jobs/Impl/ArtifactCollection.cs
+++ b/Engine/Source/Programs/Horde/Horde.Build/Jobs/Impl/ArtifactCollection.cs
@@ -6,6 +6,7 @@ using System.Linq;
 using System.Threading.Tasks;
 using Horde.Build.Jobs;
 using Horde.Build.Models;
+using Horde.Build.Server;
 using Horde.Build.Storage;
 using Horde.Build.Utilities;
 using MongoDB.Bson;
@@ -66,15 +67,15 @@ namespace Horde.Build.Services
 		/// <summary>
 		/// Constructor
 		/// </summary>
-		/// <param name="databaseService">The database service</param>
+		/// <param name="mongoService">The database service</param>
 		/// <param name="storageBackend">The storage backend</param>
-		public ArtifactCollection(DatabaseService databaseService, IStorageBackend<ArtifactCollection> storageBackend)
+		public ArtifactCollection(MongoService mongoService, IStorageBackend<ArtifactCollection> storageBackend)
 		{
 			_storageBackend = storageBackend;
 
 			// Initialize Artifacts table
-			_artifacts = databaseService.GetCollection<Artifact>("Artifacts");
-			if (!databaseService.ReadOnlyMode)
+			_artifacts = mongoService.GetCollection<Artifact>("Artifacts");
+			if (!mongoService.ReadOnlyMode)
 			{
 				_artifacts.Indexes.CreateOne(new CreateIndexModel<Artifact>(Builders<Artifact>.IndexKeys.Ascending(x => x.JobId)));
 			}
diff --git a/Engine/Source/Programs/Horde/Horde.Build/Jobs/Impl/GraphCollection.cs b/Engine/Source/Programs/Horde/Horde.Build/Jobs/Impl/GraphCollection.cs
--- a/Engine/Source/Programs/Horde/Horde.Build/Jobs/Impl/GraphCollection.cs
+++ b/Engine/Source/Programs/Horde/Horde.Build/Jobs/Impl/GraphCollection.cs
@@ -9,7 +9,7 @@ using System.Threading;
 using System.Threading.Tasks;
 using EpicGames.Core;
 using Horde.Build.Models;
-using Horde.Build.Services;
+using Horde.Build.Server;
 using HordeCommon;
 using MongoDB.Bson;
 using MongoDB.Bson.Serialization.Attributes;
@@ -352,10 +352,10 @@ namespace Horde.Build.Collections.Impl
 		/// <summary>
 		/// Constructor
 		/// </summary>
-		/// <param name="databaseService">The database service singleton</param>
-		public GraphCollection(DatabaseService databaseService)
+		/// <param name="mongoService">The database service singleton</param>
+		public GraphCollection(MongoService mongoService)
 		{
-			_graphs = databaseService.GetCollection<GraphDocument>("Graphs");
+			_graphs = mongoService.GetCollection<GraphDocument>("Graphs");
 		}
 
 		/// <summary>
diff --git a/Engine/Source/Programs/Horde/Horde.Build/Jobs/Impl/JobCollection.cs b/Engine/Source/Programs/Horde/Horde.Build/Jobs/Impl/JobCollection.cs
--- a/Engine/Source/Programs/Horde/Horde.Build/Jobs/Impl/JobCollection.cs
+++ b/Engine/Source/Programs/Horde/Horde.Build/Jobs/Impl/JobCollection.cs
@@ -10,6 +10,7 @@ using EpicGames.Core;
 using Horde.Build.Acls;
 using Horde.Build.Api;
 using Horde.Build.Models;
+using Horde.Build.Server;
 using Horde.Build.Services;
 using Horde.Build.Utilities;
 using HordeCommon;
@@ -383,16 +384,16 @@ namespace Horde.Build.Collections.Impl
 		/// <summary>
 		/// Constructor
 		/// </summary>
-		/// <param name="databaseService">The database service singleton</param>
+		/// <param name="mongoService">The database service singleton</param>
 		/// <param name="clock"></param>
 		/// <param name="logger">The logger instance</param>
-		public JobCollection(DatabaseService databaseService, IClock clock, ILogger<JobCollection> logger)
+		public JobCollection(MongoService mongoService, IClock clock, ILogger<JobCollection> logger)
 		{
 			_clock = clock;
 			_logger = logger;
 
-			_jobs = databaseService.GetCollection<JobDocument>("Jobs");
-			_indexes = new DatabaseIndexes(_jobs, databaseService.ReadOnlyMode);
+			_jobs = mongoService.GetCollection<JobDocument>("Jobs");
+			_indexes = new DatabaseIndexes(_jobs, mongoService.ReadOnlyMode);
 		}
 
 		static Task PostLoadAsync(JobDocument job)
diff --git a/Engine/Source/Programs/Horde/Horde.Build/Jobs/Impl/JobStepRefCollection.cs b/Engine/Source/Programs/Horde/Horde.Build/Jobs/Impl/JobStepRefCollection.cs
--- a/Engine/Source/Programs/Horde/Horde.Build/Jobs/Impl/JobStepRefCollection.cs
+++ b/Engine/Source/Programs/Horde/Horde.Build/Jobs/Impl/JobStepRefCollection.cs
@@ -5,7 +5,7 @@ using System.Collections.Generic;
 using System.Linq;
 using System.Threading.Tasks;
 using Horde.Build.Models;
-using Horde.Build.Services;
+using Horde.Build.Server;
 using Horde.Build.Utilities;
 using HordeCommon;
 using MongoDB.Bson.Serialization.Attributes;
@@ -93,12 +93,12 @@ namespace Horde.Build.Collections.Impl
 		/// <summary>
 		/// Constructor
 		/// </summary>
-		/// <param name="databaseService">The database service instance</param>
-		public JobStepRefCollection(DatabaseService databaseService)
+		/// <param name="mongoService">The database service instance</param>
+		public JobStepRefCollection(MongoService mongoService)
 		{
-			_jobStepRefs = databaseService.GetCollection<JobStepRef>("JobStepRefs");
+			_jobStepRefs = mongoService.GetCollection<JobStepRef>("JobStepRefs");
 
-			if (!databaseService.ReadOnlyMode)
+			if (!mongoService.ReadOnlyMode)
 			{
 				_jobStepRefs.Indexes.CreateOne(new CreateIndexModel<JobStepRef>(Builders<JobStepRef>.IndexKeys.Ascending(x => x.StreamId).Ascending(x => x.TemplateId).Ascending(x => x.Name).Descending(x => x.Change)));
 			}
diff --git a/Engine/Source/Programs/Horde/Horde.Build/Jobs/Impl/JobTaskSource.cs b/Engine/Source/Programs/Horde/Horde.Build/Jobs/Impl/JobTaskSource.cs
--- a/Engine/Source/Programs/Horde/Horde.Build/Jobs/Impl/JobTaskSource.cs
+++ b/Engine/Source/Programs/Horde/Horde.Build/Jobs/Impl/JobTaskSource.cs
@@ -15,6 +15,7 @@ using Google.Protobuf.WellKnownTypes;
 using Horde.Build.Api;
 using Horde.Build.Collections;
 using Horde.Build.Models;
+using Horde.Build.Server;
 using Horde.Build.Services;
 using Horde.Build.Utilities;
 using HordeCommon;
@@ -179,7 +180,7 @@ namespace Horde.Build.Tasks.Impl
 			}
 		}
 
-		readonly DatabaseService _databaseService;
+		readonly MongoService _mongoService;
 		readonly StreamService _streamService;
 		readonly ILogFileService _logFileService;
 		readonly IAgentCollection _agentsCollection;
@@ -230,9 +231,9 @@ namespace Horde.Build.Tasks.Impl
 		/// <summary>
 		/// Constructor
 		/// </summary>
-		public JobTaskSource(DatabaseService databaseService, IAgentCollection agents, IJobCollection jobs, IJobStepRefCollection jobStepRefs, IGraphCollection graphs, IPoolCollection pools, IUgsMetadataCollection ugsMetadataCollection, StreamService streamService, ILogFileService logFileService, PerforceLoadBalancer perforceLoadBalancer, IClock clock, IOptionsMonitor<ServerSettings> settings, ILogger<JobTaskSource> logger)
+		public JobTaskSource(MongoService mongoService, IAgentCollection agents, IJobCollection jobs, IJobStepRefCollection jobStepRefs, IGraphCollection graphs, IPoolCollection pools, IUgsMetadataCollection ugsMetadataCollection, StreamService streamService, ILogFileService logFileService, PerforceLoadBalancer perforceLoadBalancer, IClock clock, IOptionsMonitor<ServerSettings> settings, ILogger<JobTaskSource> logger)
 		{
-			_databaseService = databaseService;
+			_mongoService = mongoService;
 			_agentsCollection = agents;
 			_jobs = jobs;
 			_jobStepRefs = jobStepRefs;
@@ -743,7 +744,7 @@ namespace Horde.Build.Tasks.Impl
 				leaseName.Append(CultureInfo.InvariantCulture, $" - {job.Name}");
 
 				// Get the global settings
-				Globals globals = await _databaseService.GetGlobalsAsync();
+				Globals globals = await _mongoService.GetGlobalsAsync();
 
 				// Encode the payload
 				ExecuteJobTask? task = await CreateExecuteJobTaskAsync(item._stream, job, batch, agent, item._workspace, item._useAutoSdk, logId);
@@ -810,7 +811,7 @@ namespace Horde.Build.Tasks.Impl
 			leaseName.Append(CultureInfo.InvariantCulture, $" - {job.Name}");
 
 			// Get the global settings
-			Globals globals = await _databaseService.GetGlobalsAsync();
+			Globals globals = await _mongoService.GetGlobalsAsync();
 
 			// Encode the payload
 			ExecuteJobTask task = new ExecuteJobTask();
diff --git a/Engine/Source/Programs/Horde/Horde.Build/Jobs/Impl/JobTimingCollection.cs b/Engine/Source/Programs/Horde/Horde.Build/Jobs/Impl/JobTimingCollection.cs
--- a/Engine/Source/Programs/Horde/Horde.Build/Jobs/Impl/JobTimingCollection.cs
+++ b/Engine/Source/Programs/Horde/Horde.Build/Jobs/Impl/JobTimingCollection.cs
@@ -5,7 +5,7 @@ using System.Diagnostics.CodeAnalysis;
 using System.Linq;
 using System.Threading.Tasks;
 using Horde.Build.Models;
-using Horde.Build.Services;
+using Horde.Build.Server;
 using Horde.Build.Utilities;
 using MongoDB.Bson.Serialization.Attributes;
 using MongoDB.Driver;
@@ -82,10 +82,10 @@ namespace Horde.Build.Collections.Impl
 		/// <summary>
 		/// Constructor
 		/// </summary>
-		/// <param name="databaseService">The database service singleton</param>
-		public JobTimingCollection(DatabaseService databaseService)
+		/// <param name="mongoService">The database service singleton</param>
+		public JobTimingCollection(MongoService mongoService)
 		{
-			_collection = databaseService.GetCollection<JobTimingDocument>("JobTiming");
+			_collection = mongoService.GetCollection<JobTimingDocument>("JobTiming");
 		}
 
 		/// <inheritdoc/>
diff --git a/Engine/Source/Programs/Horde/Horde.Build/Jobs/Impl/TemplateCollection.cs b/Engine/Source/Programs/Horde/Horde.Build/Jobs/Impl/TemplateCollection.cs
--- a/Engine/Source/Programs/Horde/Horde.Build/Jobs/Impl/TemplateCollection.cs
+++ b/Engine/Source/Programs/Horde/Horde.Build/Jobs/Impl/TemplateCollection.cs
@@ -6,7 +6,7 @@ using System.Linq;
 using System.Threading.Tasks;
 using EpicGames.Core;
 using Horde.Build.Models;
-using Horde.Build.Services;
+using Horde.Build.Server;
 using Horde.Build.Utilities;
 using HordeCommon;
 using Microsoft.Extensions.Caching.Memory;
@@ -88,8 +88,8 @@ namespace Horde.Build.Collections.Impl
 		/// <summary>
 		/// Constructor
 		/// </summary>
-		/// <param name="databaseService">The database service singleton</param>
-		public TemplateCollection(DatabaseService databaseService)
+		/// <param name="mongoService">The database service singleton</param>
+		public TemplateCollection(MongoService mongoService)
 		{
 			// Ensure discriminator cannot be registered twice (throws exception). Can otherwise happen during unit tests.
 			if (BsonSerializer.LookupDiscriminatorConvention(typeof(JobsTabColumn)) == null)
@@ -97,7 +97,7 @@ namespace Horde.Build.Collections.Impl
 				BsonSerializer.RegisterDiscriminatorConvention(typeof(JobsTabColumn), new DefaultDiscriminatorConvention(typeof(JobsTabColumn), typeof(JobsTabLabelColumn)));	
 			}
 			
-			_templates = databaseService.GetCollection<TemplateDocument>("Templates");
+			_templates = mongoService.GetCollection<TemplateDocument>("Templates");
 
 			MemoryCacheOptions options = new MemoryCacheOptions();
 			_templateCache = new MemoryCache(options);
diff --git a/Engine/Source/Programs/Horde/Horde.Build/Jobs/Impl/TestDataCollection.cs b/Engine/Source/Programs/Horde/Horde.Build/Jobs/Impl/TestDataCollection.cs
--- a/Engine/Source/Programs/Horde/Horde.Build/Jobs/Impl/TestDataCollection.cs
+++ b/Engine/Source/Programs/Horde/Horde.Build/Jobs/Impl/TestDataCollection.cs
@@ -5,7 +5,7 @@ using System.Collections.Generic;
 using System.Linq;
 using System.Threading.Tasks;
 using Horde.Build.Models;
-using Horde.Build.Services;
+using Horde.Build.Server;
 using Horde.Build.Utilities;
 using MongoDB.Bson;
 using MongoDB.Driver;
@@ -62,12 +62,12 @@ namespace Horde.Build.Collections.Impl
 		/// <summary>
 		/// Constructor
 		/// </summary>
-		/// <param name="databaseService">The database service instance</param>
-		public TestDataCollection(DatabaseService databaseService)
+		/// <param name="mongoService">The database service instance</param>
+		public TestDataCollection(MongoService mongoService)
 		{
-			_testDataDocuments = databaseService.GetCollection<TestDataDocument>("TestData");
+			_testDataDocuments = mongoService.GetCollection<TestDataDocument>("TestData");
 
-			if (!databaseService.ReadOnlyMode)
+			if (!mongoService.ReadOnlyMode)
 			{
 				_testDataDocuments.Indexes.CreateOne(new CreateIndexModel<TestDataDocument>(Builders<TestDataDocument>.IndexKeys.Ascending(x => x.StreamId).Ascending(x => x.Change).Ascending(x => x.Key)));
 				_testDataDocuments.Indexes.CreateOne(new CreateIndexModel<TestDataDocument>(Builders<TestDataDocument>.IndexKeys.Ascending(x => x.JobId).Ascending(x => x.StepId).Ascending(x => x.Key), new CreateIndexOptions { Unique = true }));
diff --git a/Engine/Source/Programs/Horde/Horde.Build/Logs/Impl/LogEventCollection.cs b/Engine/Source/Programs/Horde/Horde.Build/Logs/Impl/LogEventCollection.cs
--- a/Engine/Source/Programs/Horde/Horde.Build/Logs/Impl/LogEventCollection.cs
+++ b/Engine/Source/Programs/Horde/Horde.Build/Logs/Impl/LogEventCollection.cs
@@ -5,7 +5,7 @@ using System.Collections.Generic;
 using System.Linq;
 using System.Threading.Tasks;
 using Horde.Build.Models;
-using Horde.Build.Services;
+using Horde.Build.Server;
 using Horde.Build.Utilities;
 using HordeCommon;
 using Microsoft.Extensions.Logging;
@@ -106,16 +106,16 @@ namespace Horde.Build.Collections.Impl
 		/// <summary>
 		/// Constructor
 		/// </summary>
-		/// <param name="databaseService">The database service instance</param>
+		/// <param name="mongoService">The database service instance</param>
 		/// <param name="logger">The logger instance</param>
-		public LogEventCollection(DatabaseService databaseService, ILogger<LogEventCollection> logger)
+		public LogEventCollection(MongoService mongoService, ILogger<LogEventCollection> logger)
 		{
 			_logger = logger;
 
-			_logEvents = databaseService.GetCollection<LogEventDocument>("LogEvents");
-			_legacyEvents = databaseService.GetCollection<LegacyEventDocument>("Events");
+			_logEvents = mongoService.GetCollection<LogEventDocument>("LogEvents");
+			_legacyEvents = mongoService.GetCollection<LegacyEventDocument>("Events");
 
-			if (!databaseService.ReadOnlyMode)
+			if (!mongoService.ReadOnlyMode)
 			{
 				_logEvents.Indexes.CreateOne(new CreateIndexModel<LogEventDocument>(Builders<LogEventDocument>.IndexKeys.Ascending(x => x.Id.LogId)));
 				_logEvents.Indexes.CreateOne(new CreateIndexModel<LogEventDocument>(Builders<LogEventDocument>.IndexKeys.Ascending(x => x.SpanId).Ascending(x => x.Id)));
diff --git a/Engine/Source/Programs/Horde/Horde.Build/Logs/Impl/LogFileCollection.cs b/Engine/Source/Programs/Horde/Horde.Build/Logs/Impl/LogFileCollection.cs
--- a/Engine/Source/Programs/Horde/Horde.Build/Logs/Impl/LogFileCollection.cs
+++ b/Engine/Source/Programs/Horde/Horde.Build/Logs/Impl/LogFileCollection.cs
@@ -6,7 +6,7 @@ using System.Net;
 using System.Threading.Tasks;
 using Horde.Build.Api;
 using Horde.Build.Models;
-using Horde.Build.Services;
+using Horde.Build.Server;
 using Horde.Build.Utilities;
 using MongoDB.Bson.Serialization.Attributes;
 using MongoDB.Driver;
@@ -109,10 +109,10 @@ namespace Horde.Build.Collections.Impl
 		/// <summary>
 		/// Constructor
 		/// </summary>
-		/// <param name="databaseService">The database service singleton</param>
-		public LogFileCollection(DatabaseService databaseService)
+		/// <param name="mongoService">The database service singleton</param>
+		public LogFileCollection(MongoService mongoService)
 		{
-			_logFiles = databaseService.GetCollection<LogFileDocument>("LogFiles");
+			_logFiles = mongoService.GetCollection<LogFileDocument>("LogFiles");
 			_hostName = Dns.GetHostName();
 		}
 
diff --git a/Engine/Source/Programs/Horde/Horde.Build/Models/Globals.cs b/Engine/Source/Programs/Horde/Horde.Build/Models/Globals.cs
--- a/Engine/Source/Programs/Horde/Horde.Build/Models/Globals.cs
+++ b/Engine/Source/Programs/Horde/Horde.Build/Models/Globals.cs
@@ -6,6 +6,7 @@ using System.ComponentModel.DataAnnotations;
 using System.Linq;
 using System.Security.Cryptography;
 using Horde.Build.Acls;
+using Horde.Build.Server;
 using Horde.Build.Utilities;
 using MongoDB.Bson;
 using MongoDB.Bson.Serialization.Attributes;
diff --git a/Engine/Source/Programs/Horde/Horde.Build/Notices/Impl/NoticeCollection.cs b/Engine/Source/Programs/Horde/Horde.Build/Notices/Impl/NoticeCollection.cs
--- a/Engine/Source/Programs/Horde/Horde.Build/Notices/Impl/NoticeCollection.cs
+++ b/Engine/Source/Programs/Horde/Horde.Build/Notices/Impl/NoticeCollection.cs
@@ -5,7 +5,7 @@ using System.Collections.Generic;
 using System.Linq;
 using System.Threading.Tasks;
 using Horde.Build.Models;
-using Horde.Build.Services;
+using Horde.Build.Server;
 using Horde.Build.Utilities;
 using MongoDB.Bson;
 using MongoDB.Bson.Serialization.Attributes;
@@ -60,9 +60,9 @@ namespace Horde.Build.Collections.Impl
 		/// <summary>
 		/// Constructor
 		/// </summary>
-		public NoticeCollection(DatabaseService databaseService)
+		public NoticeCollection(MongoService mongoService)
 		{
-			_notices = databaseService.GetCollection<NoticeDocument>("Notices");
+			_notices = mongoService.GetCollection<NoticeDocument>("Notices");
 		}
 
 		/// <inheritdoc/>
diff --git a/Engine/Source/Programs/Horde/Horde.Build/Notifications/Impl/NotificationService.cs b/Engine/Source/Programs/Horde/Horde.Build/Notifications/Impl/NotificationService.cs
--- a/Engine/Source/Programs/Horde/Horde.Build/Notifications/Impl/NotificationService.cs
+++ b/Engine/Source/Programs/Horde/Horde.Build/Notifications/Impl/NotificationService.cs
@@ -11,6 +11,7 @@ using EpicGames.Core;
 using Horde.Build.Api;
 using Horde.Build.Collections;
 using Horde.Build.Models;
+using Horde.Build.Server;
 using Horde.Build.Services;
 using Horde.Build.Utilities;
 using HordeCommon;
@@ -122,7 +123,7 @@ namespace Horde.Build.Notifications.Impl
 		/// <summary>
 		/// Constructor
 		/// </summary>
-		public NotificationService(IEnumerable<INotificationSink> sinks, IOptionsMonitor<ServerSettings> settings, ILogger<NotificationService> logger, IGraphCollection graphCollection, ISubscriptionCollection subscriptionCollection, INotificationTriggerCollection triggerCollection, IUserCollection userCollection, JobService jobService, StreamService streamService, IIssueService issueService, ILogFileService logFileService, IDogStatsd dogStatsd, IDatabase redisDatabase, IClock clock)
+		public NotificationService(IEnumerable<INotificationSink> sinks, IOptionsMonitor<ServerSettings> settings, ILogger<NotificationService> logger, IGraphCollection graphCollection, ISubscriptionCollection subscriptionCollection, INotificationTriggerCollection triggerCollection, IUserCollection userCollection, JobService jobService, StreamService streamService, IIssueService issueService, ILogFileService logFileService, IDogStatsd dogStatsd, RedisService redisService, IClock clock)
 		{
 			_sinks = sinks.ToList();
 			_settings = settings;
@@ -136,7 +137,7 @@ namespace Horde.Build.Notifications.Impl
 			_issueService = issueService;
 			_logFileService = logFileService;
 			_dogStatsd = dogStatsd;
-			_redisDatabase = redisDatabase;
+			_redisDatabase = redisService.Database;
 
 			issueService.OnIssueUpdated += NotifyIssueUpdated;
 			jobService.OnJobStepComplete += NotifyJobStepComplete;
diff --git a/Engine/Source/Programs/Horde/Horde.Build/Notifications/Impl/NotificationTriggerCollection.cs b/Engine/Source/Programs/Horde/Horde.Build/Notifications/Impl/NotificationTriggerCollection.cs
--- a/Engine/Source/Programs/Horde/Horde.Build/Notifications/Impl/NotificationTriggerCollection.cs
+++ b/Engine/Source/Programs/Horde/Horde.Build/Notifications/Impl/NotificationTriggerCollection.cs
@@ -4,7 +4,7 @@ using System.Collections.Generic;
 using System.Linq;
 using System.Threading.Tasks;
 using Horde.Build.Models;
-using Horde.Build.Services;
+using Horde.Build.Server;
 using Horde.Build.Utilities;
 using MongoDB.Bson;
 using MongoDB.Bson.Serialization.Attributes;
@@ -65,11 +65,11 @@ namespace Horde.Build.Collections.Impl
 		/// <summary>
 		/// Constructor
 		/// </summary>
-		/// <param name="databaseService">The database singleton</param>
+		/// <param name="mongoService">The database singleton</param>
 		/// <param name="userCollection"></param>
-		public NotificationTriggerCollection(DatabaseService databaseService, IUserCollection userCollection)
+		public NotificationTriggerCollection(MongoService mongoService, IUserCollection userCollection)
 		{
-			_triggers = databaseService.GetCollection<TriggerDocument>("NotificationTriggers");
+			_triggers = mongoService.GetCollection<TriggerDocument>("NotificationTriggers");
 			_userCollection = userCollection;
 		}
 
diff --git a/Engine/Source/Programs/Horde/Horde.Build/Notifications/Impl/SlackNotificationSink.cs b/Engine/Source/Programs/Horde/Horde.Build/Notifications/Impl/SlackNotificationSink.cs
--- a/Engine/Source/Programs/Horde/Horde.Build/Notifications/Impl/SlackNotificationSink.cs
+++ b/Engine/Source/Programs/Horde/Horde.Build/Notifications/Impl/SlackNotificationSink.cs
@@ -18,6 +18,7 @@ using EpicGames.Core;
 using Horde.Build.Api;
 using Horde.Build.Collections;
 using Horde.Build.Models;
+using Horde.Build.Server;
 using Horde.Build.Services;
 using Horde.Build.Utilities;
 using Horde.Build.Utilities.Slack.BlockKit;
@@ -210,22 +211,22 @@ namespace Horde.Build.Notifications.Impl
 		/// <summary>
 		/// Constructor
 		/// </summary>
-		/// <param name="databaseService"></param>
+		/// <param name="mongoService"></param>
 		/// <param name="issueService"></param>
 		/// <param name="userCollection">The user collection</param>
 		/// <param name="streamService"></param>
 		/// <param name="environment"></param>
 		/// <param name="settings">The current configuration settings</param>
 		/// <param name="logger">Logging device</param>
-		public SlackNotificationSink(DatabaseService databaseService, IIssueService issueService, IUserCollection userCollection, StreamService streamService, IWebHostEnvironment environment, IOptions<ServerSettings> settings, ILogger<SlackNotificationSink> logger)
+		public SlackNotificationSink(MongoService mongoService, IIssueService issueService, IUserCollection userCollection, StreamService streamService, IWebHostEnvironment environment, IOptions<ServerSettings> settings, ILogger<SlackNotificationSink> logger)
 		{
 			_issueService = issueService;
 			_userCollection = userCollection;
 			_streamService = streamService;
 			_environment = environment;
 			_settings = settings.Value;
-			_messageStates = databaseService.Database.GetCollection<MessageStateDocument>("Slack");
-			_slackUsers = databaseService.Database.GetCollection<SlackUser>("Slack.UsersV2");
+			_messageStates = mongoService.Database.GetCollection<MessageStateDocument>("Slack");
+			_slackUsers = mongoService.Database.GetCollection<SlackUser>("Slack.UsersV2");
 			_logger = logger;
 
 			if (!String.IsNullOrEmpty(settings.Value.SlackUsers))
diff --git a/Engine/Source/Programs/Horde/Horde.Build/Notifications/Impl/SubscriptionCollection.cs b/Engine/Source/Programs/Horde/Horde.Build/Notifications/Impl/SubscriptionCollection.cs
--- a/Engine/Source/Programs/Horde/Horde.Build/Notifications/Impl/SubscriptionCollection.cs
+++ b/Engine/Source/Programs/Horde/Horde.Build/Notifications/Impl/SubscriptionCollection.cs
@@ -9,7 +9,7 @@ using System.Threading.Tasks;
 using EpicGames.Core;
 using Horde.Build.Api;
 using Horde.Build.Models;
-using Horde.Build.Services;
+using Horde.Build.Server;
 using Horde.Build.Utilities;
 using HordeCommon;
 using MongoDB.Bson.Serialization.Attributes;
@@ -184,11 +184,11 @@ namespace Horde.Build.Collections.Impl
 		/// <summary>
 		/// Constructor
 		/// </summary>
-		/// <param name="databaseService">The database service</param>
-		public SubscriptionCollection(DatabaseService databaseService)
+		/// <param name="mongoService">The database service</param>
+		public SubscriptionCollection(MongoService mongoService)
 		{
-			_collection = databaseService.GetCollection<Subscription>("SubscriptionsV2");
-			if (!databaseService.ReadOnlyMode)
+			_collection = mongoService.GetCollection<Subscription>("SubscriptionsV2");
+			if (!mongoService.ReadOnlyMode)
 			{
 				_collection.Indexes.CreateOne(new CreateIndexModel<Subscription>(Builders<Subscription>.IndexKeys.Ascending(x => x.Event)));
 				_collection.Indexes.CreateOne(new CreateIndexModel<Subscription>(Builders<Subscription>.IndexKeys.Ascending(x => x.UserId)));
diff --git a/Engine/Source/Programs/Horde/Horde.Build/Perforce/Impl/CommitCollection.cs b/Engine/Source/Programs/Horde/Horde.Build/Perforce/Impl/CommitCollection.cs
--- a/Engine/Source/Programs/Horde/Horde.Build/Perforce/Impl/CommitCollection.cs
+++ b/Engine/Source/Programs/Horde/Horde.Build/Perforce/Impl/CommitCollection.cs
@@ -5,7 +5,7 @@ using System.Collections.Generic;
 using System.Linq;
 using System.Threading.Tasks;
 using Horde.Build.Models;
-using Horde.Build.Services;
+using Horde.Build.Server;
 using Horde.Build.Utilities;
 using MongoDB.Bson;
 using MongoDB.Bson.Serialization.Attributes;
@@ -82,11 +82,11 @@ namespace Horde.Build.Commits.Impl
 		/// <summary>
 		/// Constructor
 		/// </summary>
-		public CommitCollection(DatabaseService databaseService)
+		public CommitCollection(MongoService mongoService)
 		{
-			_commits = databaseService.Database.GetCollection<Commit>("Commits");
+			_commits = mongoService.Database.GetCollection<Commit>("Commits");
 
-			if (!databaseService.ReadOnlyMode)
+			if (!mongoService.ReadOnlyMode)
 			{
 				_commits.Indexes.CreateOne(new CreateIndexModel<Commit>(Builders<Commit>.IndexKeys.Ascending(x => x.StreamId).Descending(x => x.Change), new CreateIndexOptions { Unique = true }));
 			}
diff --git a/Engine/Source/Programs/Horde/Horde.Build/Perforce/Impl/CommitService.cs b/Engine/Source/Programs/Horde/Horde.Build/Perforce/Impl/CommitService.cs
--- a/Engine/Source/Programs/Horde/Horde.Build/Perforce/Impl/CommitService.cs
+++ b/Engine/Source/Programs/Horde/Horde.Build/Perforce/Impl/CommitService.cs
@@ -19,6 +19,7 @@ using EpicGames.Redis;
 using EpicGames.Serialization;
 using Horde.Build.Collections;
 using Horde.Build.Models;
+using Horde.Build.Server;
 using Horde.Build.Services;
 using Horde.Build.Utilities;
 using HordeCommon;
@@ -188,13 +189,13 @@ namespace Horde.Build.Commits.Impl
 		/// <summary>
 		/// Constructor
 		/// </summary>
-		public CommitService(IDatabase redis, ICommitCollection commitCollection, IStorageClient storageClient, IStreamCollection streamCollection, IPerforceService perforceService, IUserCollection userCollection, IClock clock, IMemoryCache memoryCache, IOptions<CommitServiceOptions> options, ILogger<CommitService> logger)
+		public CommitService(RedisService redisService, ICommitCollection commitCollection, IStorageClient storageClient, IStreamCollection streamCollection, IPerforceService perforceService, IUserCollection userCollection, IClock clock, IMemoryCache memoryCache, IOptions<CommitServiceOptions> options, ILogger<CommitService> logger)
 		{
 			Options = options.Value;
 
-			_redis = redis;
-			_redisDirtyStreams = new RedisSet<StreamId>(redis, RedisBaseKey.Append("streams"));
-			_redisReservations = new RedisSortedSet<StreamId>(redis, RedisBaseKey.Append("reservations"));
+			_redis = redisService.Database;
+			_redisDirtyStreams = new RedisSet<StreamId>(_redis, RedisBaseKey.Append("streams"));
+			_redisReservations = new RedisSortedSet<StreamId>(_redis, RedisBaseKey.Append("reservations"));
 
 			_commitCollection = commitCollection;
 			_storageClient = storageClient;
diff --git a/Engine/Source/Programs/Horde/Horde.Build/Perforce/Impl/PerforceService.cs b/Engine/Source/Programs/Horde/Horde.Build/Perforce/Impl/PerforceService.cs
--- a/Engine/Source/Programs/Horde/Horde.Build/Perforce/Impl/PerforceService.cs
+++ b/Engine/Source/Programs/Horde/Horde.Build/Perforce/Impl/PerforceService.cs
@@ -16,6 +16,7 @@ using EpicGames.Core;
 using EpicGames.Perforce;
 using Horde.Build.Collections;
 using Horde.Build.Models;
+using Horde.Build.Server;
 using Horde.Build.Utilities;
 using Microsoft.Extensions.Caching.Memory;
 using Microsoft.Extensions.Logging;
@@ -89,10 +90,10 @@ namespace Horde.Build.Services
 		/// <summary>
 		/// Constructor
 		/// </summary>
-		public PerforceService(PerforceLoadBalancer loadBalancer, DatabaseService databaseService, IUserCollection userCollection, IOptions<ServerSettings> settings, ILogger<PerforceService> logger)
+		public PerforceService(PerforceLoadBalancer loadBalancer, MongoService mongoService, IUserCollection userCollection, IOptions<ServerSettings> settings, ILogger<PerforceService> logger)
 		{
 			_loadBalancer = loadBalancer;
-			_cachedGlobals = new LazyCachedValue<Task<Globals>>(() => databaseService.GetGlobalsAsync(), TimeSpan.FromSeconds(30.0));
+			_cachedGlobals = new LazyCachedValue<Task<Globals>>(() => mongoService.GetGlobalsAsync(), TimeSpan.FromSeconds(30.0));
 			_userCollection = userCollection;
 			_settings = settings.Value;
 			_logger = logger;
diff --git a/Engine/Source/Programs/Horde/Horde.Build/Perforce/PerforceLoadBalancer.cs b/Engine/Source/Programs/Horde/Horde.Build/Perforce/PerforceLoadBalancer.cs
--- a/Engine/Source/Programs/Horde/Horde.Build/Perforce/PerforceLoadBalancer.cs
+++ b/Engine/Source/Programs/Horde/Horde.Build/Perforce/PerforceLoadBalancer.cs
@@ -12,6 +12,7 @@ using System.Threading.Tasks;
 using Google.Protobuf.WellKnownTypes;
 using Horde.Build.Collections;
 using Horde.Build.Models;
+using Horde.Build.Server;
 using Horde.Build.Utilities;
 using HordeCommon;
 using HordeCommon.Rpc.Tasks;
@@ -134,7 +135,7 @@ namespace Horde.Build.Services
 			public List<PerforceServerEntry> Servers { get; set; } = new List<PerforceServerEntry>();
 		}
 
-		readonly DatabaseService _databaseService;
+		readonly MongoService _mongoService;
 		readonly ILeaseCollection _leaseCollection;
 		readonly SingletonDocument<PerforceServerList> _serverListSingleton;
 		readonly Random _random = new Random();
@@ -144,13 +145,13 @@ namespace Horde.Build.Services
 		/// <summary>
 		/// Constructor
 		/// </summary>
-		public PerforceLoadBalancer(DatabaseService databaseService, ILeaseCollection leaseCollection, IClock clock, ILogger<PerforceLoadBalancer> logger)
+		public PerforceLoadBalancer(MongoService mongoService, ILeaseCollection leaseCollection, IClock clock, ILogger<PerforceLoadBalancer> logger)
 		{
-			_databaseService = databaseService;
+			_mongoService = mongoService;
 			_leaseCollection = leaseCollection;
-			_serverListSingleton = new SingletonDocument<PerforceServerList>(databaseService);
+			_serverListSingleton = new SingletonDocument<PerforceServerList>(mongoService);
 			_logger = logger;
-			if (databaseService.ReadOnlyMode)
+			if (mongoService.ReadOnlyMode)
 			{
 				_ticker = new NullTicker();
 			}
@@ -320,7 +321,7 @@ namespace Horde.Build.Services
 		/// <inheritdoc/>
 		async ValueTask TickInternalAsync(CancellationToken cancellationToken)
 		{
-			Globals globals = await _databaseService.GetGlobalsAsync();
+			Globals globals = await _mongoService.GetGlobalsAsync();
 
 			// Set of new server entries
 			List<PerforceServerEntry> newServers = new List<PerforceServerEntry>();
diff --git a/Engine/Source/Programs/Horde/Horde.Build/Projects/IProjectCollection.cs b/Engine/Source/Programs/Horde/Horde.Build/Projects/IProjectCollection.cs
--- a/Engine/Source/Programs/Horde/Horde.Build/Projects/IProjectCollection.cs
+++ b/Engine/Source/Programs/Horde/Horde.Build/Projects/IProjectCollection.cs
@@ -3,6 +3,7 @@
 using System.Collections.Generic;
 using System.Threading.Tasks;
 using Horde.Build.Models;
+using Horde.Build.Server;
 using Horde.Build.Utilities;
 
 namespace Horde.Build.Collections
diff --git a/Engine/Source/Programs/Horde/Horde.Build/Projects/IStreamCollection.cs b/Engine/Source/Programs/Horde/Horde.Build/Projects/IStreamCollection.cs
--- a/Engine/Source/Programs/Horde/Horde.Build/Projects/IStreamCollection.cs
+++ b/Engine/Source/Programs/Horde/Horde.Build/Projects/IStreamCollection.cs
@@ -4,6 +4,7 @@ using System;
 using System.Collections.Generic;
 using System.Threading.Tasks;
 using Horde.Build.Models;
+using Horde.Build.Server;
 using Horde.Build.Utilities;
 
 namespace Horde.Build.Collections
diff --git a/Engine/Source/Programs/Horde/Horde.Build/Projects/Impl/ProjectCollection.cs b/Engine/Source/Programs/Horde/Horde.Build/Projects/Impl/ProjectCollection.cs
--- a/Engine/Source/Programs/Horde/Horde.Build/Projects/Impl/ProjectCollection.cs
+++ b/Engine/Source/Programs/Horde/Horde.Build/Projects/Impl/ProjectCollection.cs
@@ -1,4 +1,4 @@
-﻿// Copyright Epic Games, Inc. All Rights Reserved.
+// Copyright Epic Games, Inc. All Rights Reserved.
 
 using System;
 using System.Collections.Generic;
@@ -7,7 +7,7 @@ using System.Linq;
 using System.Threading.Tasks;
 using Horde.Build.Acls;
 using Horde.Build.Models;
-using Horde.Build.Services;
+using Horde.Build.Server;
 using Horde.Build.Utilities;
 using MongoDB.Bson.Serialization.Attributes;
 using MongoDB.Driver;
@@ -103,11 +103,11 @@ namespace Horde.Build.Collections.Impl
 		/// <summary>
 		/// Constructor
 		/// </summary>
-		/// <param name="databaseService">The database service instance</param>
-		public ProjectCollection(DatabaseService databaseService)
+		/// <param name="mongoService">The database service instance</param>
+		public ProjectCollection(MongoService mongoService)
 		{
-			_projects = databaseService.GetCollection<ProjectDocument>("Projects");
-			_projectLogos = databaseService.GetCollection<ProjectLogoDocument>("ProjectLogos");
+			_projects = mongoService.GetCollection<ProjectDocument>("Projects");
+			_projectLogos = mongoService.GetCollection<ProjectLogoDocument>("ProjectLogos");
 		}
 
 		/// <inheritdoc/>
diff --git a/Engine/Source/Programs/Horde/Horde.Build/Projects/Impl/StreamCollection.cs b/Engine/Source/Programs/Horde/Horde.Build/Projects/Impl/StreamCollection.cs
--- a/Engine/Source/Programs/Horde/Horde.Build/Projects/Impl/StreamCollection.cs
+++ b/Engine/Source/Programs/Horde/Horde.Build/Projects/Impl/StreamCollection.cs
@@ -9,7 +9,7 @@ using System.Threading.Tasks;
 using Horde.Build.Acls;
 using Horde.Build.Api;
 using Horde.Build.Models;
-using Horde.Build.Services;
+using Horde.Build.Server;
 using Horde.Build.Utilities;
 using HordeCommon;
 using MongoDB.Bson.Serialization.Attributes;
@@ -120,12 +120,12 @@ namespace Horde.Build.Collections.Impl
 		/// <summary>
 		/// Constructor
 		/// </summary>
-		/// <param name="databaseService">The database service instance</param>
+		/// <param name="mongoService">The database service instance</param>
 		/// <param name="clock"></param>
 		/// <param name="templateCollection"></param>
-		public StreamCollection(DatabaseService databaseService, IClock clock, ITemplateCollection templateCollection)
+		public StreamCollection(MongoService mongoService, IClock clock, ITemplateCollection templateCollection)
 		{
-			_streams = databaseService.GetCollection<StreamDocument>("Streams");
+			_streams = mongoService.GetCollection<StreamDocument>("Streams");
 			_clock = clock;
 			_templateCollection = templateCollection;
 		}
diff --git a/Engine/Source/Programs/Horde/Horde.Build/Projects/ScheduleService.cs b/Engine/Source/Programs/Horde/Horde.Build/Projects/ScheduleService.cs
--- a/Engine/Source/Programs/Horde/Horde.Build/Projects/ScheduleService.cs
+++ b/Engine/Source/Programs/Horde/Horde.Build/Projects/ScheduleService.cs
@@ -14,6 +14,7 @@ using EpicGames.Serialization;
 using Horde.Build.Api;
 using Horde.Build.Collections;
 using Horde.Build.Models;
+using Horde.Build.Server;
 using Horde.Build.Utilities;
 using HordeCommon;
 using Microsoft.Extensions.Hosting;
@@ -152,7 +153,7 @@ namespace Horde.Build.Services
 		/// <summary>
 		/// Constructor
 		/// </summary>
-		public ScheduleService(RedisService redis, IGraphCollection graphs, IPerforceService perforce, IJobCollection jobCollection, JobService jobService, StreamService streamService, ITemplateCollection templateCollection, DatabaseService databaseService, IClock clock, IOptionsMonitor<ServerSettings> settings, ILogger<ScheduleService> logger)
+		public ScheduleService(RedisService redis, IGraphCollection graphs, IPerforceService perforce, IJobCollection jobCollection, JobService jobService, StreamService streamService, ITemplateCollection templateCollection, MongoService mongoService, IClock clock, IOptionsMonitor<ServerSettings> settings, ILogger<ScheduleService> logger)
 		{
 			_graphs = graphs;
 			_perforce = perforce;
@@ -171,7 +172,7 @@ namespace Horde.Build.Services
 			_baseLockKey = "scheduler/locks";
 			_tickLockKey = _baseLockKey.Append("/tick");
 			_queue = new RedisSortedSet<QueueItem>(redis.Database, "scheduler/queue");
-			if (databaseService.ReadOnlyMode)
+			if (mongoService.ReadOnlyMode)
 			{
 				_ticker = new NullTicker();
 			}
diff --git a/Engine/Source/Programs/Horde/Horde.Build/Secrets/CredentialService.cs b/Engine/Source/Programs/Horde/Horde.Build/Secrets/CredentialService.cs
--- a/Engine/Source/Programs/Horde/Horde.Build/Secrets/CredentialService.cs
+++ b/Engine/Source/Programs/Horde/Horde.Build/Secrets/CredentialService.cs
@@ -1,4 +1,4 @@
-﻿// Copyright Epic Games, Inc. All Rights Reserved.
+// Copyright Epic Games, Inc. All Rights Reserved.
 
 using System.Collections.Generic;
 using System.Linq;
@@ -6,6 +6,7 @@ using System.Security.Claims;
 using System.Threading.Tasks;
 using Horde.Build.Acls;
 using Horde.Build.Models;
+using Horde.Build.Server;
 using MongoDB.Bson;
 using MongoDB.Driver;
 
@@ -30,14 +31,14 @@ namespace Horde.Build.Services
 		/// Constructor
 		/// </summary>
 		/// <param name="aclService">The ACL service instance</param>
-		/// <param name="databaseService">The database service instance</param>
-		public CredentialService(AclService aclService, DatabaseService databaseService)
+		/// <param name="mongoService">The database service instance</param>
+		public CredentialService(AclService aclService, MongoService mongoService)
 		{
 			_aclService = aclService;
 
-			_credentials = databaseService.Credentials;
+			_credentials = mongoService.Credentials;
 
-			if (!databaseService.ReadOnlyMode)
+			if (!mongoService.ReadOnlyMode)
 			{
 				_credentials.Indexes.CreateOne(new CreateIndexModel<Credential>(Builders<Credential>.IndexKeys.Ascending(x => x.NormalizedName), new CreateIndexOptions { Unique = true }));
 			}
diff --git a/Engine/Source/Programs/Horde/Horde.Build/Server/Config.cs b/Engine/Source/Programs/Horde/Horde.Build/Server/Config.cs
--- a/Engine/Source/Programs/Horde/Horde.Build/Server/Config.cs
+++ b/Engine/Source/Programs/Horde/Horde.Build/Server/Config.cs
@@ -7,9 +7,10 @@ using EpicGames.Core;
 using EpicGames.Horde.Common;
 using EpicGames.Serialization;
 using Horde.Build.Api;
+using Horde.Build.Models;
 using Horde.Build.Utilities;
 
-namespace Horde.Build.Models
+namespace Horde.Build.Server
 {
 	using ProjectId = StringId<IProject>;
 	using StreamId = StringId<IStream>;
diff --git a/Engine/Source/Programs/Horde/Horde.Build/Server/ConfigService.cs b/Engine/Source/Programs/Horde/Horde.Build/Server/ConfigService.cs
--- a/Engine/Source/Programs/Horde/Horde.Build/Server/ConfigService.cs
+++ b/Engine/Source/Programs/Horde/Horde.Build/Server/ConfigService.cs
@@ -16,6 +16,7 @@ using Horde.Build.Api;
 using Horde.Build.Collections;
 using Horde.Build.Models;
 using Horde.Build.Notifications;
+using Horde.Build.Services;
 using Horde.Build.Utilities;
 using HordeCommon;
 using Microsoft.AspNetCore.StaticFiles;
@@ -24,7 +25,7 @@ using Microsoft.Extensions.Hosting;
 using Microsoft.Extensions.Logging;
 using Microsoft.Extensions.Options;
 
-namespace Horde.Build.Services
+namespace Horde.Build.Server
 {
 	using PoolId = StringId<IPool>;
 	using ProjectId = StringId<IProject>;
@@ -42,7 +43,7 @@ namespace Horde.Build.Services
 		/// Config file version number
 		/// </summary>
 		const int Version = 10;
-		readonly DatabaseService _databaseService;
+		readonly MongoService _mongoService;
 		readonly ProjectService _projectService;
 		readonly StreamService _streamService;
 		readonly IPerforceService _perforceService;
@@ -56,9 +57,9 @@ namespace Horde.Build.Services
 		/// <summary>
 		/// Constructor
 		/// </summary>
-		public ConfigService(DatabaseService databaseService, IPerforceService perforceService, ProjectService projectService, StreamService streamService, INotificationService notificationService, PoolService poolService, AgentService agentService, IClock clock, IOptionsMonitor<ServerSettings> settings, ILogger<ConfigService> logger)
+		public ConfigService(MongoService mongoService, IPerforceService perforceService, ProjectService projectService, StreamService streamService, INotificationService notificationService, PoolService poolService, AgentService agentService, IClock clock, IOptionsMonitor<ServerSettings> settings, ILogger<ConfigService> logger)
 		{
-			_databaseService = databaseService;
+			_mongoService = mongoService;
 			_perforceService = perforceService;
 			_projectService = projectService;
 			_streamService = streamService;
@@ -66,7 +67,7 @@ namespace Horde.Build.Services
 			_poolService = poolService;
 			_agentService = agentService;
 			_settings = settings;
-			if (databaseService.ReadOnlyMode)
+			if (mongoService.ReadOnlyMode)
 			{
 				_ticker = new NullTicker();
 			}
@@ -123,7 +124,7 @@ namespace Horde.Build.Services
 				}
 				globalConfig = _cachedGlobalConfig;
 
-				Globals globals = await _databaseService.GetGlobalsAsync();
+				Globals globals = await _mongoService.GetGlobalsAsync();
 				if (globals.ConfigRevision == revision)
 				{
 					_logger.LogInformation("Updating configuration from {ConfigPath}", globals.ConfigRevision);
@@ -137,7 +138,7 @@ namespace Horde.Build.Services
 				globals.ComputeClusters = _cachedGlobalConfig.Compute;
 				globals.RootAcl = Acl.Merge(null, _cachedGlobalConfig.Acl);
 
-				if (await _databaseService.TryUpdateSingletonAsync(globals))
+				if (await _mongoService.TryUpdateSingletonAsync(globals))
 				{
 					break;
 				}
@@ -176,7 +177,7 @@ namespace Horde.Build.Services
 				}
 
 				IProject? project = projects.FirstOrDefault(x => x.Id == projectRef.Id);
-				bool update = (project == null || project.ConfigPath != projectPath.ToString() || project.ConfigRevision != revision);
+				bool update = project == null || project.ConfigPath != projectPath.ToString() || project.ConfigRevision != revision;
 
 				ProjectConfig? projectConfig;
 				if (!update && prevCachedProjectConfigs.TryGetValue(projectRef.Id, out (ProjectConfig Config, string Revision) result) && result.Revision == revision)
@@ -374,7 +375,7 @@ namespace Horde.Build.Services
 
 		Task<byte[]> ReadDataAsync(Uri configPath)
 		{
-			switch(configPath.Scheme)
+			switch (configPath.Scheme)
 			{
 				case FileScheme:
 					return File.ReadAllBytesAsync(configPath.LocalPath);
@@ -421,16 +422,16 @@ namespace Horde.Build.Services
 		{
 
 			Uri? configUri = null;
-			
+
 			if (Path.IsPathRooted(_settings.CurrentValue.ConfigPath) && !_settings.CurrentValue.ConfigPath.StartsWith("//", StringComparison.Ordinal))
 			{
 				// absolute path to config
-				configUri = new Uri(_settings.CurrentValue.ConfigPath);					
+				configUri = new Uri(_settings.CurrentValue.ConfigPath);
 			}
 			else if (_settings.CurrentValue.ConfigPath != null)
 			{
 				// relative (development) or perforce path
-				configUri = CombinePaths(new Uri(FileReference.Combine(Program.AppDir, "_").FullName), _settings.CurrentValue.ConfigPath);				
+				configUri = CombinePaths(new Uri(FileReference.Combine(Program.AppDir, "_").FullName), _settings.CurrentValue.ConfigPath);
 			}
 
 			if (configUri != null)
@@ -454,7 +455,7 @@ namespace Horde.Build.Services
 			{
 				throw new Exception($"Global config path must be rooted for service updates.  (Perforce paths are not currently supported for live updates), {_settings.CurrentValue.ConfigPath}");
 			}
-			
+
 			FileReference globalConfigFile = new FileReference(_settings.CurrentValue.ConfigPath);
 			DirectoryReference globalConfigDirectory = globalConfigFile.Directory;
 
@@ -505,8 +506,8 @@ namespace Horde.Build.Services
 			}
 
 			if (configDirty == true)
-			{				
-				await UpdateConfigAsync(new Uri("file://" + globalConfigFile.ToString()));				
+			{
+				await UpdateConfigAsync(new Uri("file://" + globalConfigFile.ToString()));
 			}
 
 			// create the default pool 
@@ -606,7 +607,7 @@ namespace Horde.Build.Services
 					{
 						userProps.Remove(pair.Key);
 						continue;
-					}					
+					}
 
 					JsonElement element = (JsonElement)pair.Value;
 
@@ -633,33 +634,33 @@ namespace Horde.Build.Services
 					{
 						response.Errors.Add($"Unable to map type for Property {property.Name}");
 						continue;
-					}					
+					}
 
 					// handle common conversions
 					if (property.GetType() != value.GetType())
 					{
 						try
 						{
-							value =	Convert.ChangeType(value, property.PropertyType, CultureInfo.CurrentCulture);						
+							value = Convert.ChangeType(value, property.PropertyType, CultureInfo.CurrentCulture);
 						}
 						catch (Exception ex)
 						{
 							response.Errors.Add($"Property {property.Name} raised exception during conversion, {ex.Message}");
-							continue;							
+							continue;
 						}
-						
+
 						if (value == null)
 						{
 							response.Errors.Add($"Property {property.Name} had null value on conversion");
-							continue;							
+							continue;
 						}
 					}
-					else 
+					else
 					{
 						response.Errors.Add($"Property {property.Name} is not assignable to {value}");
 						continue;
 					}
-					
+
 
 					//  Set the value, providing some validation
 					try
diff --git a/Engine/Source/Programs/Horde/Horde.Build/Server/DatabaseService.cs b/Engine/Source/Programs/Horde/Horde.Build/Server/DatabaseService.cs
deleted file mode 100644
diff --git a/Engine/Source/Programs/Horde/Horde.Build/Server/DowntimeService.cs b/Engine/Source/Programs/Horde/Horde.Build/Server/DowntimeService.cs
--- a/Engine/Source/Programs/Horde/Horde.Build/Server/DowntimeService.cs
+++ b/Engine/Source/Programs/Horde/Horde.Build/Server/DowntimeService.cs
@@ -10,7 +10,7 @@ using Microsoft.Extensions.Hosting;
 using Microsoft.Extensions.Logging;
 using Microsoft.Extensions.Options;
 
-namespace Horde.Build.Services
+namespace Horde.Build.Server
 {
 	/// <summary>
 	/// Interface for a service which keeps track of whether we're during downtime
@@ -40,7 +40,7 @@ namespace Horde.Build.Services
 			private set;
 		}
 
-		readonly DatabaseService _databaseService;
+		readonly MongoService _mongoService;
 		readonly ITicker _ticker;
 		readonly IOptionsMonitor<ServerSettings> _settings;
 		readonly ILogger _logger;
@@ -48,13 +48,13 @@ namespace Horde.Build.Services
 		/// <summary>
 		/// Constructor
 		/// </summary>
-		/// <param name="databaseService">The database service instance</param>
+		/// <param name="mongoService">The database service instance</param>
 		/// <param name="clock"></param>
 		/// <param name="settings">The server settings</param>
 		/// <param name="logger">Logger instance</param>
-		public DowntimeService(DatabaseService databaseService, IClock clock, IOptionsMonitor<ServerSettings> settings, ILogger<DowntimeService> logger)
+		public DowntimeService(MongoService mongoService, IClock clock, IOptionsMonitor<ServerSettings> settings, ILogger<DowntimeService> logger)
 		{
-			_databaseService = databaseService;
+			_mongoService = mongoService;
 			_settings = settings;
 			_logger = logger;
 
@@ -80,7 +80,7 @@ namespace Horde.Build.Services
 		/// <returns>Async task</returns>
 		async ValueTask TickAsync(CancellationToken stoppingToken)
 		{
-			Globals globals = await _databaseService.GetGlobalsAsync();
+			Globals globals = await _mongoService.GetGlobalsAsync();
 
 			DateTimeOffset now = TimeZoneInfo.ConvertTime(DateTimeOffset.Now, _settings.CurrentValue.TimeZoneInfo);
 			bool bIsActive = globals.ScheduledDowntime.Any(x => x.IsActive(now));
@@ -89,7 +89,7 @@ namespace Horde.Build.Services
 			foreach (ScheduledDowntime schedule in globals.ScheduledDowntime)
 			{
 				DateTimeOffset start = schedule.GetNext(now).StartTime;
-				if(next == null || start < next)
+				if (next == null || start < next)
 				{
 					next = start;
 				}
diff --git a/Engine/Source/Programs/Horde/Horde.Build/Server/LifetimeService.cs b/Engine/Source/Programs/Horde/Horde.Build/Server/LifetimeService.cs
--- a/Engine/Source/Programs/Horde/Horde.Build/Server/LifetimeService.cs
+++ b/Engine/Source/Programs/Horde/Horde.Build/Server/LifetimeService.cs
@@ -7,7 +7,7 @@ using Microsoft.Extensions.Hosting;
 using Microsoft.Extensions.Logging;
 using StackExchange.Redis;
 
-namespace Horde.Build.Services
+namespace Horde.Build.Server
 {
 	/// <summary>
 	/// Service containing an async task that allows long polling operations to complete early if the server is shutting down
@@ -33,17 +33,10 @@ namespace Horde.Build.Services
 		/// Registration token for the stopping event
 		/// </summary>
 		readonly CancellationTokenRegistration _registration;
-		
-		/// <summary>
-		/// Singleton instance of the database service
-		/// </summary>
-		readonly DatabaseService _databaseService;
-		
-		/// <summary>
-		/// The Redis database
-		/// </summary>
-		readonly IDatabase _redisDb;
-		
+
+		readonly MongoService _mongoService;
+		readonly RedisService _redisService;
+
 		/*
 		/// <summary>
 		/// Max time to wait for any outstanding requests to finish
@@ -61,13 +54,13 @@ namespace Horde.Build.Services
 		/// Constructor
 		/// </summary>
 		/// <param name="lifetime">Application lifetime interface</param>
-		/// <param name="databaseService">Database singleton service</param>
-		/// <param name="redisDb">Redis singleton service</param>
+		/// <param name="mongoService">Database singleton service</param>
+		/// <param name="redisService">Redis singleton service</param>
 		/// <param name="logger">Logging interface</param>
-		public LifetimeService(IHostApplicationLifetime lifetime, DatabaseService databaseService, IDatabase redisDb, ILogger<LifetimeService> logger)
+		public LifetimeService(IHostApplicationLifetime lifetime, MongoService mongoService, RedisService redisService, ILogger<LifetimeService> logger)
 		{
-			_databaseService = databaseService;
-			_redisDb = redisDb;
+			_mongoService = mongoService;
+			_redisService = redisService;
 			_logger = logger;
 			_stoppingTaskCompletionSource = new TaskCompletionSource<bool>();
 			_preStoppingTaskCompletionSource = new TaskCompletionSource<bool>();
@@ -96,7 +89,7 @@ namespace Horde.Build.Services
 			_logger.LogInformation("Delaying shutdown by sleeping {ShutdownDelayMs} ms...", shutdownDelayMs);
 			Thread.Sleep(shutdownDelayMs);
 			_logger.LogInformation("Server process now shutting down...");
-			
+
 			/*
 			if (PreStoppingTaskCompletionSource.TrySetResult(true))
 			{
@@ -153,17 +146,17 @@ namespace Horde.Build.Services
 			bool isHealthy = false;
 			try
 			{
-				await _databaseService.Database.ListCollectionNamesAsync(null, cancelSource.Token);
+				await _mongoService.Database.ListCollectionNamesAsync(null, cancelSource.Token);
 				isHealthy = true;
 			}
 			catch (Exception e)
 			{
 				_logger.LogError("MongoDB call failed during health check", e);
 			}
-			
+
 			return isHealthy;
 		}
-		
+
 		/// <summary>
 		/// Check if Redis can be reached
 		/// </summary>
@@ -175,15 +168,15 @@ namespace Horde.Build.Services
 			try
 			{
 				string key = "HordeLifetimeService-Health-Check";
-				await _redisDb.StringSetAsync(key, "ok");
-				await _redisDb.StringGetAsync(key);
+				await _redisService.Database.StringSetAsync(key, "ok");
+				await _redisService.Database.StringGetAsync(key);
 				isHealthy = true;
 			}
 			catch (Exception e)
 			{
 				_logger.LogError("Redis call failed during health check", e);
 			}
-			
+
 			return isHealthy;
 		}
 	}
diff --git a/Engine/Source/Programs/Horde/Horde.Build/Server/MetricService.cs b/Engine/Source/Programs/Horde/Horde.Build/Server/MetricService.cs
--- a/Engine/Source/Programs/Horde/Horde.Build/Server/MetricService.cs
+++ b/Engine/Source/Programs/Horde/Horde.Build/Server/MetricService.cs
@@ -8,7 +8,7 @@ using Microsoft.Extensions.Hosting;
 using Microsoft.Extensions.Logging;
 using StatsdClient;
 
-namespace Horde.Build.Services
+namespace Horde.Build.Server
 {
 	/// <summary>
 	/// Periodically send metrics for the CLR and other services that cannot be collected on a per-request basis
@@ -56,12 +56,12 @@ namespace Horde.Build.Services
 
 			int busyIoThreads = maxIoThreads - freeIoThreads;
 			int busyWorkerThreads = maxWorkerThreads - freeWorkerThreads;
-			
+
 			_dogStatsd.Gauge("horde.clr.threadpool.io.max", maxIoThreads);
 			_dogStatsd.Gauge("horde.clr.threadpool.io.min", minIoThreads);
 			_dogStatsd.Gauge("horde.clr.threadpool.io.free", freeIoThreads);
 			_dogStatsd.Gauge("horde.clr.threadpool.io.busy", busyIoThreads);
-			
+
 			_dogStatsd.Gauge("horde.clr.threadpool.worker.max", maxWorkerThreads);
 			_dogStatsd.Gauge("horde.clr.threadpool.worker.min", minWorkerThreads);
 			_dogStatsd.Gauge("horde.clr.threadpool.worker.free", freeWorkerThreads);
diff --git a/Engine/Source/Programs/Horde/Horde.Build/Server/MongoService.cs b/Engine/Source/Programs/Horde/Horde.Build/Server/MongoService.cs
new file mode 100644
--- /dev/null
+++ b/Engine/Source/Programs/Horde/Horde.Build/Server/MongoService.cs
@@ -0,0 +1,582 @@
+// Copyright Epic Games, Inc. All Rights Reserved.
+
+using System;
+using System.Collections.Generic;
+using System.Diagnostics;
+using System.Globalization;
+using System.IO;
+using System.Linq;
+using System.Net;
+using System.Net.NetworkInformation;
+using System.Net.Security;
+using System.Runtime.InteropServices;
+using System.Security.Cryptography.X509Certificates;
+using System.Text;
+using System.Text.RegularExpressions;
+using System.Threading.Tasks;
+using EpicGames.Core;
+using Horde.Build.Models;
+using Horde.Build.Services;
+using Horde.Build.Utiltiies;
+using Microsoft.Extensions.Logging;
+using Microsoft.Extensions.Options;
+using Microsoft.IdentityModel.Tokens;
+using MongoDB.Bson;
+using MongoDB.Bson.Serialization;
+using MongoDB.Driver;
+using MongoDB.Driver.Core.Events;
+
+namespace Horde.Build.Server
+{
+	/// <summary>
+	/// Singleton for accessing the database
+	/// </summary>
+	public sealed class MongoService : IDisposable
+	{
+		/// <summary>
+		/// The database instance
+		/// </summary>
+		public IMongoDatabase Database { get; private set; }
+
+		/// <summary>
+		/// Collection of credential documents
+		/// </summary>
+		public IMongoCollection<Credential> Credentials { get; }
+
+		/// <summary>
+		/// Collection of singleton documents
+		/// </summary>
+		public IMongoCollection<BsonDocument> Singletons { get; }
+
+		/// <summary>
+		/// Settings for the application
+		/// </summary>
+		ServerSettings Settings { get; }
+
+		/// <summary>
+		/// Name of the JWT issuer
+		/// </summary>
+		public string JwtIssuer { get; }
+
+		/// <summary>
+		/// The signing key for this instance
+		/// </summary>
+		public SymmetricSecurityKey JwtSigningKey { get; }
+
+		/// <summary>
+		/// Logger for this instance
+		/// </summary>
+		readonly ILogger<MongoService> _logger;
+
+		/// <summary>
+		/// Access the database in a read-only mode (don't create indices or modify content)
+		/// </summary>
+		public bool ReadOnlyMode { get; }
+
+		/// <summary>
+		/// The mongo process group
+		/// </summary>
+		ManagedProcessGroup? _mongoProcessGroup;
+
+		/// <summary>
+		/// The mongo process
+		/// </summary>
+		ManagedProcess? _mongoProcess;
+
+		/// <summary>
+		/// Task to read from the mongo process
+		/// </summary>
+		Task? _mongoOutputTask;
+
+		/// <summary>
+		/// Factory for creating logger instances
+		/// </summary>
+		readonly ILoggerFactory _loggerFactory;
+
+		/// <summary>
+		/// Default port for MongoDB connections
+		/// </summary>
+		const int DefaultMongoPort = 27017;
+
+		/// <summary>
+		/// Constructor
+		/// </summary>
+		/// <param name="settingsSnapshot">The settings instance</param>
+		/// <param name="loggerFactory">Instance of the logger for this service</param>
+		public MongoService(IOptions<ServerSettings> settingsSnapshot, ILoggerFactory loggerFactory)
+		{
+			Settings = settingsSnapshot.Value;
+			_logger = loggerFactory.CreateLogger<MongoService>();
+			_loggerFactory = loggerFactory;
+
+			try
+			{
+				ReadOnlyMode = Settings.DatabaseReadOnlyMode;
+				if (Settings.DatabasePublicCert != null)
+				{
+					X509Store localTrustStore = new X509Store(StoreName.Root);
+					try
+					{
+						localTrustStore.Open(OpenFlags.ReadWrite);
+
+						X509Certificate2Collection collection = ImportCertificateBundle(Settings.DatabasePublicCert);
+						foreach (X509Certificate2 certificate in collection)
+						{
+							_logger.LogInformation("Importing certificate for {Subject}", certificate.Subject);
+						}
+
+						localTrustStore.AddRange(collection);
+					}
+					finally
+					{
+						localTrustStore.Close();
+					}
+				}
+
+				string? connectionString = Settings.DatabaseConnectionString;
+				if (connectionString == null)
+				{
+					if (IsPortInUse(DefaultMongoPort))
+					{
+						connectionString = "mongodb://localhost:27017";
+					}
+					else if (TryStartMongoServer(_logger))
+					{
+						connectionString = "mongodb://localhost:27017/?readPreference=primary&appname=Horde&ssl=false";
+					}
+					else
+					{
+						throw new Exception($"Unable to connect to MongoDB server. Setup a MongoDB server and set the connection string in {Program.UserConfigFile}");
+					}
+				}
+
+				MongoClientSettings mongoSettings = MongoClientSettings.FromConnectionString(connectionString);
+				mongoSettings.ClusterConfigurator = clusterBuilder =>
+				{
+					if (_logger.IsEnabled(LogLevel.Trace))
+					{
+						clusterBuilder.Subscribe<CommandStartedEvent>(ev => TraceMongoCommand(ev.Command));
+					}
+				};
+
+				mongoSettings.SslSettings = new SslSettings();
+				mongoSettings.SslSettings.ServerCertificateValidationCallback = CertificateValidationCallBack;
+				mongoSettings.MaxConnectionPoolSize = 300; // Default is 100
+
+				//TestSslConnection(MongoSettings.Server.Host, MongoSettings.Server.Port, Logger);
+
+				MongoClient client = new MongoClient(mongoSettings);
+				Database = client.GetDatabase(Settings.DatabaseName);
+
+				Singletons = GetCollection<BsonDocument>("Singletons");
+
+				Globals globals = GetGlobalsAsync().Result;
+				while (globals.JwtSigningKey == null)
+				{
+					globals.RotateSigningKey();
+					if (!TryUpdateSingletonAsync(globals).Result)
+					{
+						globals = GetGlobalsAsync().Result;
+					}
+				}
+
+				JwtIssuer = Settings.JwtIssuer ?? Dns.GetHostName();
+				if (String.IsNullOrEmpty(Settings.JwtSecret))
+				{
+					JwtSigningKey = new SymmetricSecurityKey(globals.JwtSigningKey);
+				}
+				else
+				{
+					JwtSigningKey = new SymmetricSecurityKey(Convert.FromBase64String(Settings.JwtSecret));
+				}
+
+				Credentials = GetCollection<Credential>("Credentials");
+			}
+			catch (Exception ex)
+			{
+				_logger.LogError(ex, "Exception while initializing MongoService");
+				throw;
+			}
+		}
+
+		internal const int CtrlCEvent = 0;
+
+		[DllImport("kernel32.dll")]
+		internal static extern bool GenerateConsoleCtrlEvent(int eventId, int processGroupId);
+
+		/// <inheritdoc/>
+		public void Dispose()
+		{
+			if (_mongoProcess != null)
+			{
+				GenerateConsoleCtrlEvent(CtrlCEvent, _mongoProcess.Id);
+
+				_mongoOutputTask?.Wait();
+				_mongoOutputTask = null;
+
+				_mongoProcess.WaitForExit();
+				_mongoProcess.Dispose();
+				_mongoProcess = null;
+			}
+			if (_mongoProcessGroup != null)
+			{
+				_mongoProcessGroup.Dispose();
+				_mongoProcessGroup = null;
+			}
+		}
+
+		/// <summary>
+		/// Checks if the given port is in use
+		/// </summary>
+		/// <param name="portNumber"></param>
+		/// <returns></returns>
+		static bool IsPortInUse(int portNumber)
+		{
+			IPGlobalProperties ipGlobalProperties = IPGlobalProperties.GetIPGlobalProperties();
+
+			IPEndPoint[] listeners = ipGlobalProperties.GetActiveTcpListeners();
+			if (listeners.Any(x => x.Port == portNumber))
+			{
+				return true;
+			}
+
+			return false;
+		}
+
+		/// <summary>
+		/// Attempts to start a local instance of MongoDB
+		/// </summary>
+		/// <param name="logger"></param>
+		/// <returns></returns>
+		bool TryStartMongoServer(ILogger logger)
+		{
+			if (!RuntimeInformation.IsOSPlatform(OSPlatform.Windows))
+			{
+				return false;
+			}
+
+			FileReference mongoExe = FileReference.Combine(Program.AppDir, "ThirdParty", "Mongo", "mongod.exe");
+			if (!FileReference.Exists(mongoExe))
+			{
+				logger.LogWarning("Unable to find Mongo executable.");
+				return false;
+			}
+
+			DirectoryReference mongoDir = DirectoryReference.Combine(Program.DataDir, "Mongo");
+
+			DirectoryReference mongoDataDir = DirectoryReference.Combine(mongoDir, "Data");
+			DirectoryReference.CreateDirectory(mongoDataDir);
+
+			FileReference mongoLogFile = FileReference.Combine(mongoDir, "mongod.log");
+
+			FileReference configFile = FileReference.Combine(mongoDir, "mongod.conf");
+			if (!FileReference.Exists(configFile))
+			{
+				DirectoryReference.CreateDirectory(configFile.Directory);
+				using (StreamWriter writer = new StreamWriter(configFile.FullName))
+				{
+					writer.WriteLine("# mongod.conf");
+					writer.WriteLine();
+					writer.WriteLine("# for documentation of all options, see:");
+					writer.WriteLine("# http://docs.mongodb.org/manual/reference/configuration-options/");
+					writer.WriteLine();
+					writer.WriteLine("storage:");
+					writer.WriteLine("    dbPath: {0}", mongoDataDir.FullName);
+					writer.WriteLine();
+					writer.WriteLine("net:");
+					writer.WriteLine("    port: {0}", DefaultMongoPort);
+					writer.WriteLine("    bindIp: 127.0.0.1");
+				}
+			}
+
+			_mongoProcessGroup = new ManagedProcessGroup();
+			try
+			{
+				_mongoProcess = new ManagedProcess(_mongoProcessGroup, mongoExe.FullName, $"--config \"{configFile}\"", null, null, ProcessPriorityClass.Normal);
+				_mongoProcess.StdIn.Close();
+				_mongoOutputTask = Task.Run(() => RelayMongoOutput());
+				return true;
+			}
+			catch (Exception ex)
+			{
+				logger.LogWarning(ex, "Unable to start Mongo server process");
+				return false;
+			}
+		}
+
+		/// <summary>
+		/// Copies output from the mongo process to the logger
+		/// </summary>
+		/// <returns></returns>
+		async Task RelayMongoOutput()
+		{
+			ILogger mongoLogger = _loggerFactory.CreateLogger("MongoDB");
+
+			Dictionary<string, ILogger> channelLoggers = new Dictionary<string, ILogger>();
+			for (; ; )
+			{
+				string? line = await _mongoProcess!.ReadLineAsync();
+				if (line == null)
+				{
+					break;
+				}
+				if (line.Length > 0)
+				{
+					Match match = Regex.Match(line, @"^\s*[^\s]+\s+([A-Z])\s+([^\s]+)\s+(.*)");
+					if (match.Success)
+					{
+						ILogger? channelLogger;
+						if (!channelLoggers.TryGetValue(match.Groups[2].Value, out channelLogger))
+						{
+							channelLogger = _loggerFactory.CreateLogger($"MongoDB.{match.Groups[2].Value}");
+							channelLoggers.Add(match.Groups[2].Value, channelLogger);
+						}
+						channelLogger.Log(ParseMongoLogLevel(match.Groups[1].Value), "{Message}", match.Groups[3].Value.TrimEnd());
+					}
+					else
+					{
+						mongoLogger.Log(LogLevel.Information, "{Message}", line);
+					}
+				}
+			}
+			mongoLogger.LogInformation("Exit code {ExitCode}", _mongoProcess.ExitCode);
+		}
+
+		static LogLevel ParseMongoLogLevel(string text)
+		{
+			if (text.Equals("I", StringComparison.Ordinal))
+			{
+				return LogLevel.Information;
+			}
+			else if (text.Equals("E", StringComparison.Ordinal))
+			{
+				return LogLevel.Error;
+			}
+			else
+			{
+				return LogLevel.Warning;
+			}
+		}
+
+		/// <summary>
+		/// Logs a mongodb command, removing any fields we don't care about
+		/// </summary>
+		/// <param name="command">The command document</param>
+		void TraceMongoCommand(BsonDocument command)
+		{
+			List<string> names = new List<string>();
+			List<string> values = new List<string>();
+
+			foreach (BsonElement element in command)
+			{
+				if (element.Value != null && !element.Name.Equals("$db", StringComparison.Ordinal) && !element.Name.Equals("lsid", StringComparison.Ordinal))
+				{
+					names.Add($"{element.Name}: {{{values.Count}}}");
+					values.Add(element.Value.ToString()!);
+				}
+			}
+
+#pragma warning disable CA2254 // Template should be a static expression
+			_logger.LogTrace($"MongoDB: {String.Join(", ", names)}", values.ToArray());
+#pragma warning restore CA2254 // Template should be a static expression
+		}
+
+		/// <summary>
+		/// Imports one or more certificates from a single PEM file
+		/// </summary>
+		/// <param name="fileName">File to import</param>
+		/// <returns>Collection of certificates</returns>
+		static X509Certificate2Collection ImportCertificateBundle(string fileName)
+		{
+			X509Certificate2Collection collection = new X509Certificate2Collection();
+
+			string text = File.ReadAllText(fileName);
+			for (int offset = 0; offset < text.Length;)
+			{
+				int nextOffset = text.IndexOf("-----BEGIN CERTIFICATE-----", offset + 1, StringComparison.Ordinal);
+				if (nextOffset == -1)
+				{
+					nextOffset = text.Length;
+				}
+
+				string certificateText = text.Substring(offset, nextOffset - offset);
+				collection.Add(new X509Certificate2(Encoding.UTF8.GetBytes(certificateText)));
+
+				offset = nextOffset;
+			}
+
+			return collection;
+		}
+
+		/// <summary>
+		/// Provides additional diagnostic information for SSL certificate validation
+		/// </summary>
+		/// <param name="sender"></param>
+		/// <param name="certificate"></param>
+		/// <param name="chain"></param>
+		/// <param name="sslPolicyErrors"></param>
+		/// <returns>True if the certificate is allowed, false otherwise</returns>
+		bool CertificateValidationCallBack(object sender, X509Certificate? certificate, X509Chain? chain, SslPolicyErrors sslPolicyErrors)
+		{
+			// If the certificate is a valid, signed certificate, return true.
+			if (sslPolicyErrors == SslPolicyErrors.None)
+			{
+				return true;
+			}
+
+			// Generate diagnostic information
+			StringBuilder builder = new StringBuilder();
+			if (sender != null)
+			{
+				string senderInfo = StringUtils.Indent(sender.ToString() ?? String.Empty, "    ");
+				builder.Append(CultureInfo.InvariantCulture, $"\nSender:\n{senderInfo}");
+			}
+			if (certificate != null)
+			{
+				builder.Append(CultureInfo.InvariantCulture, $"\nCertificate: {certificate.Subject}");
+			}
+			if (chain != null)
+			{
+				if (chain.ChainStatus != null && chain.ChainStatus.Length > 0)
+				{
+					builder.Append("\nChain status:");
+					foreach (X509ChainStatus status in chain.ChainStatus)
+					{
+						builder.Append(CultureInfo.InvariantCulture, $"\n  {status.StatusInformation}");
+					}
+				}
+				if (chain.ChainElements != null)
+				{
+					builder.Append("\nChain elements:");
+					for (int idx = 0; idx < chain.ChainElements.Count; idx++)
+					{
+						X509ChainElement element = chain.ChainElements[idx];
+						builder.Append(CultureInfo.InvariantCulture, $"\n  {idx,4} - Certificate: {element.Certificate.Subject}");
+						if (element.ChainElementStatus != null && element.ChainElementStatus.Length > 0)
+						{
+							foreach (X509ChainStatus status in element.ChainElementStatus)
+							{
+								builder.Append(CultureInfo.InvariantCulture, $"\n         Status: {status.StatusInformation} ({status.Status})");
+							}
+						}
+						if (!String.IsNullOrEmpty(element.Information))
+						{
+							builder.Append(CultureInfo.InvariantCulture, $"\n         Info: {element.Information}");
+						}
+					}
+				}
+			}
+
+			// Print out additional diagnostic information
+			_logger.LogError("TLS certificate validation failed ({Errors}).{AdditionalInfo}", sslPolicyErrors, StringUtils.Indent(builder.ToString(), "    "));
+			return false;
+		}
+
+		/// <summary>
+		/// Get a MongoDB collection from database
+		/// </summary>
+		/// <param name="name">Name of collection</param>
+		/// <typeparam name="T">A MongoDB document</typeparam>
+		/// <returns></returns>
+		public IMongoCollection<T> GetCollection<T>(string name)
+		{
+			return new MongoTracingCollection<T>(Database.GetCollection<T>(name));
+		}
+
+		/// <summary>
+		/// Gets a singleton document by id
+		/// </summary>
+		/// <returns>The globals document</returns>
+		public Task<Globals> GetGlobalsAsync()
+		{
+			return GetSingletonAsync(Globals.StaticId, () => new Globals() { SchemaVersion = UpgradeService.LatestSchemaVersion });
+		}
+
+		/// <summary>
+		/// Gets a singleton document by id
+		/// </summary>
+		/// <param name="id">Id of the singleton document</param>
+		/// <returns>The document</returns>
+		public async Task<T> GetSingletonAsync<T>(ObjectId id) where T : SingletonBase, new()
+		{
+			T singleton = await GetSingletonAsync(id, () => new T());
+			singleton.PostLoad();
+			return singleton;
+		}
+
+		/// <summary>
+		/// Gets a singleton document by id
+		/// </summary>
+		/// <param name="id">Id of the singleton document</param>
+		/// <param name="constructor">Method to use to construct a new object</param>
+		/// <returns>The document</returns>
+		public async Task<T> GetSingletonAsync<T>(ObjectId id, Func<T> constructor) where T : SingletonBase, new()
+		{
+			FilterDefinition<BsonDocument> filter = new BsonDocument(new BsonElement("_id", id));
+			for (; ; )
+			{
+				BsonDocument? document = await Singletons.Find(filter).FirstOrDefaultAsync();
+				if (document != null)
+				{
+					T item = BsonSerializer.Deserialize<T>(document);
+					item.PostLoad();
+					return item;
+				}
+
+				T newItem = constructor();
+				newItem.Id = id;
+				await Singletons.InsertOneAsync(newItem.ToBsonDocument());
+			}
+		}
+
+		/// <summary>
+		/// Updates a singleton
+		/// </summary>
+		/// <typeparam name="T"></typeparam>
+		/// <param name="id"></param>
+		/// <param name="updater"></param>
+		/// <returns></returns>
+		public async Task UpdateSingletonAsync<T>(ObjectId id, Action<T> updater) where T : SingletonBase, new()
+		{
+			for (; ; )
+			{
+				T document = await GetSingletonAsync(id, () => new T());
+				updater(document);
+
+				if (await TryUpdateSingletonAsync(document))
+				{
+					break;
+				}
+			}
+		}
+
+		/// <summary>
+		/// Attempts to update a singleton object
+		/// </summary>
+		/// <param name="singletonObject">The singleton object</param>
+		/// <returns>True if the singleton document was updated</returns>
+		public async Task<bool> TryUpdateSingletonAsync<T>(T singletonObject) where T : SingletonBase
+		{
+			int prevRevision = singletonObject.Revision++;
+
+			BsonDocument filter = new BsonDocument { new BsonElement("_id", singletonObject.Id), new BsonElement(nameof(SingletonBase.Revision), prevRevision) };
+			try
+			{
+				ReplaceOneResult result = await Singletons.ReplaceOneAsync(filter, singletonObject.ToBsonDocument(), new ReplaceOptions { IsUpsert = true });
+				return result.MatchedCount > 0;
+			}
+			catch (MongoWriteException ex)
+			{
+				// Duplicate key error occurs if filter fails to match because revision is not the same.
+				if (ex.WriteError != null && ex.WriteError.Category == ServerErrorCategory.DuplicateKey)
+				{
+					return false;
+				}
+				else
+				{
+					throw;
+				}
+			}
+		}
+	}
+}
diff --git a/Engine/Source/Programs/Horde/Horde.Build/Server/RedisService.cs b/Engine/Source/Programs/Horde/Horde.Build/Server/RedisService.cs
--- a/Engine/Source/Programs/Horde/Horde.Build/Server/RedisService.cs
+++ b/Engine/Source/Programs/Horde/Horde.Build/Server/RedisService.cs
@@ -14,7 +14,7 @@ using Microsoft.Extensions.Logging;
 using Microsoft.Extensions.Options;
 using StackExchange.Redis;
 
-namespace Horde.Build.Services
+namespace Horde.Build.Server
 {
 	/// <summary>
 	/// Manages the lifetime of a bundled Redis instance
@@ -66,7 +66,7 @@ namespace Horde.Build.Services
 		/// </summary>
 		/// <param name="settings"></param>
 		public RedisService(ServerSettings settings)
-			: this(Options.Create<ServerSettings>(settings), new Serilog.Extensions.Logging.SerilogLoggerFactory(Serilog.Log.Logger))
+			: this(Options.Create(settings), new Serilog.Extensions.Logging.SerilogLoggerFactory(Serilog.Log.Logger))
 		{
 		}
 
@@ -158,7 +158,7 @@ namespace Horde.Build.Services
 			DirectoryReference.CreateDirectory(redisDir);
 
 			FileReference redisConfigFile = FileReference.Combine(redisDir, "redis.conf");
-			if(!FileReference.Exists(redisConfigFile))
+			if (!FileReference.Exists(redisConfigFile))
 			{
 				using (StreamWriter writer = new StreamWriter(redisConfigFile.FullName))
 				{
diff --git a/Engine/Source/Programs/Horde/Horde.Build/Services/HealthService.cs b/Engine/Source/Programs/Horde/Horde.Build/Services/HealthService.cs
--- a/Engine/Source/Programs/Horde/Horde.Build/Services/HealthService.cs
+++ b/Engine/Source/Programs/Horde/Horde.Build/Services/HealthService.cs
@@ -3,6 +3,7 @@
 using System.Threading.Tasks;
 using Grpc.Core;
 using Grpc.Health.V1;
+using Horde.Build.Server;
 using static Grpc.Health.V1.HealthCheckResponse.Types;
 
 namespace Horde.Build.Services
diff --git a/Engine/Source/Programs/Horde/Horde.Build/Services/RpcService.cs b/Engine/Source/Programs/Horde/Horde.Build/Services/RpcService.cs
--- a/Engine/Source/Programs/Horde/Horde.Build/Services/RpcService.cs
+++ b/Engine/Source/Programs/Horde/Horde.Build/Services/RpcService.cs
@@ -16,6 +16,7 @@ using Horde.Build.Acls;
 using Horde.Build.Collections;
 using Horde.Build.Jobs;
 using Horde.Build.Models;
+using Horde.Build.Server;
 using Horde.Build.Tasks.Impl;
 using Horde.Build.Utilities;
 using HordeCommon;
diff --git a/Engine/Source/Programs/Horde/Horde.Build/Startup.cs b/Engine/Source/Programs/Horde/Horde.Build/Startup.cs
--- a/Engine/Source/Programs/Horde/Horde.Build/Startup.cs
+++ b/Engine/Source/Programs/Horde/Horde.Build/Startup.cs
@@ -38,6 +38,7 @@ using Horde.Build.Logs.Storage.Impl;
 using Horde.Build.Models;
 using Horde.Build.Notifications;
 using Horde.Build.Notifications.Impl;
+using Horde.Build.Server;
 using Horde.Build.Services;
 using Horde.Build.Services.Impl;
 using Horde.Build.Storage;
@@ -339,7 +340,7 @@ namespace Horde.Build
 			services.AddSingleton<ConsistencyService>();
 			services.AddSingleton<RequestTrackerService>();
 			services.AddSingleton<CredentialService>();
-			services.AddSingleton<DatabaseService>();
+			services.AddSingleton<MongoService>();
 			services.AddSingleton<ConfigService>();
 			services.AddSingleton<IDogStatsd>(ctx =>
 			{
diff --git a/Engine/Source/Programs/Horde/Horde.Build/Storage/Services/BasicStorageClient.cs b/Engine/Source/Programs/Horde/Horde.Build/Storage/Services/BasicStorageClient.cs
--- a/Engine/Source/Programs/Horde/Horde.Build/Storage/Services/BasicStorageClient.cs
+++ b/Engine/Source/Programs/Horde/Horde.Build/Storage/Services/BasicStorageClient.cs
@@ -9,7 +9,7 @@ using System.Threading.Tasks;
 using EpicGames.Core;
 using EpicGames.Horde.Storage;
 using EpicGames.Serialization;
-using Horde.Build.Services;
+using Horde.Build.Server;
 using MongoDB.Bson.Serialization.Attributes;
 using MongoDB.Driver;
 
@@ -49,11 +49,11 @@ namespace Horde.Build.Storage.Services
 		/// <summary>
 		/// 
 		/// </summary>
-		/// <param name="databaseService"></param>
+		/// <param name="mongoService"></param>
 		/// <param name="storageBackend"></param>
-		public BasicStorageClient(DatabaseService databaseService, IStorageBackend<BasicStorageClient> storageBackend)
+		public BasicStorageClient(MongoService mongoService, IStorageBackend<BasicStorageClient> storageBackend)
 		{
-			_refs = databaseService.GetCollection<RefImpl>("Storage.RefsV2");
+			_refs = mongoService.GetCollection<RefImpl>("Storage.RefsV2");
 			_storageBackend = storageBackend;
 		}
 
diff --git a/Engine/Source/Programs/Horde/Horde.Build/Ugs/Impl/UgsMetadataCollection.cs b/Engine/Source/Programs/Horde/Horde.Build/Ugs/Impl/UgsMetadataCollection.cs
--- a/Engine/Source/Programs/Horde/Horde.Build/Ugs/Impl/UgsMetadataCollection.cs
+++ b/Engine/Source/Programs/Horde/Horde.Build/Ugs/Impl/UgsMetadataCollection.cs
@@ -7,7 +7,7 @@ using System.Linq;
 using System.Threading.Tasks;
 using Horde.Build.Api;
 using Horde.Build.Models;
-using Horde.Build.Services;
+using Horde.Build.Server;
 using Horde.Build.Utilities;
 using MongoDB.Bson;
 using MongoDB.Bson.Serialization.Attributes;
@@ -92,12 +92,12 @@ namespace Horde.Build.Collections.Impl
 		/// <summary>
 		/// Constructor
 		/// </summary>
-		/// <param name="databaseService">Database service instance</param>
-		public UgsMetadataCollection(DatabaseService databaseService)
+		/// <param name="mongoService">Database service instance</param>
+		public UgsMetadataCollection(MongoService mongoService)
 		{
-			_collection = databaseService.GetCollection<UgsMetadataDocument>("UgsMetadata");
+			_collection = mongoService.GetCollection<UgsMetadataDocument>("UgsMetadata");
 			
-			if (!databaseService.ReadOnlyMode)
+			if (!mongoService.ReadOnlyMode)
 			{
 				void CreateIndex()
 				{
diff --git a/Engine/Source/Programs/Horde/Horde.Build/Users/Impl/ServiceAccountCollection.cs b/Engine/Source/Programs/Horde/Horde.Build/Users/Impl/ServiceAccountCollection.cs
--- a/Engine/Source/Programs/Horde/Horde.Build/Users/Impl/ServiceAccountCollection.cs
+++ b/Engine/Source/Programs/Horde/Horde.Build/Users/Impl/ServiceAccountCollection.cs
@@ -5,7 +5,7 @@ using System.Collections.Generic;
 using System.Linq;
 using System.Threading.Tasks;
 using Horde.Build.Models;
-using Horde.Build.Services;
+using Horde.Build.Server;
 using MongoDB.Bson;
 using MongoDB.Bson.Serialization.Attributes;
 using MongoDB.Driver;
@@ -104,12 +104,12 @@ namespace Horde.Build.Collections.Impl
 		/// <summary>
 		/// Constructor
 		/// </summary>
-		/// <param name="databaseService">The database service</param>
-		public ServiceAccountCollection(DatabaseService databaseService)
+		/// <param name="mongoService">The database service</param>
+		public ServiceAccountCollection(MongoService mongoService)
 		{
-			_serviceAccounts = databaseService.GetCollection<ServiceAccountDocument>("ServiceAccounts");
+			_serviceAccounts = mongoService.GetCollection<ServiceAccountDocument>("ServiceAccounts");
 
-			if (!databaseService.ReadOnlyMode)
+			if (!mongoService.ReadOnlyMode)
 			{
 				_serviceAccounts.Indexes.CreateOne(new CreateIndexModel<ServiceAccountDocument>(Builders<ServiceAccountDocument>.IndexKeys.Ascending(x => x.SecretToken)));
 			}
diff --git a/Engine/Source/Programs/Horde/Horde.Build/Users/Impl/UserCollection.cs b/Engine/Source/Programs/Horde/Horde.Build/Users/Impl/UserCollection.cs
--- a/Engine/Source/Programs/Horde/Horde.Build/Users/Impl/UserCollection.cs
+++ b/Engine/Source/Programs/Horde/Horde.Build/Users/Impl/UserCollection.cs
@@ -6,7 +6,7 @@ using System.Linq;
 using System.Security.Claims;
 using System.Threading.Tasks;
 using Horde.Build.Models;
-using Horde.Build.Services;
+using Horde.Build.Server;
 using Horde.Build.Utilities;
 using MongoDB.Bson;
 using MongoDB.Bson.Serialization;
@@ -85,12 +85,12 @@ namespace Horde.Build.Collections.Impl
 		/// <summary>
 		/// Constructor
 		/// </summary>
-		/// <param name="databaseService"></param>
-		public UserCollectionV1(DatabaseService databaseService)
+		/// <param name="mongoService"></param>
+		public UserCollectionV1(MongoService mongoService)
 		{
-			_users = databaseService.GetCollection<UserDocument>("Users");
+			_users = mongoService.GetCollection<UserDocument>("Users");
 
-			if (!databaseService.ReadOnlyMode)
+			if (!mongoService.ReadOnlyMode)
 			{
 				_users.Indexes.CreateOne(new CreateIndexModel<UserDocument>(Builders<UserDocument>.IndexKeys.Ascending(x => x.PrimaryClaim), new CreateIndexOptions { Unique = true }));
 			}
diff --git a/Engine/Source/Programs/Horde/Horde.Build/Users/Impl/UserCollectionV2.cs b/Engine/Source/Programs/Horde/Horde.Build/Users/Impl/UserCollectionV2.cs
--- a/Engine/Source/Programs/Horde/Horde.Build/Users/Impl/UserCollectionV2.cs
+++ b/Engine/Source/Programs/Horde/Horde.Build/Users/Impl/UserCollectionV2.cs
@@ -5,7 +5,7 @@ using System.Collections.Generic;
 using System.Linq;
 using System.Threading.Tasks;
 using Horde.Build.Models;
-using Horde.Build.Services;
+using Horde.Build.Server;
 using Horde.Build.Utilities;
 using Microsoft.Extensions.Caching.Memory;
 using Microsoft.Extensions.Logging;
@@ -163,17 +163,17 @@ namespace Horde.Build.Collections.Impl
 		/// <summary>
 		/// Constructor
 		/// </summary>
-		/// <param name="databaseService"></param>
+		/// <param name="mongoService"></param>
 		/// <param name="logger"></param>
-		public UserCollectionV2(DatabaseService databaseService, ILogger<UserCollectionV2> logger)
+		public UserCollectionV2(MongoService mongoService, ILogger<UserCollectionV2> logger)
 		{
 			_logger = logger;
 
-			_users = databaseService.GetCollection<UserDocument>("UsersV2");
-			_userClaims = databaseService.GetCollection<UserClaimsDocument>("UserClaimsV2");
-			_userSettings = databaseService.GetCollection<UserSettingsDocument>("UserSettingsV2");
+			_users = mongoService.GetCollection<UserDocument>("UsersV2");
+			_userClaims = mongoService.GetCollection<UserClaimsDocument>("UserClaimsV2");
+			_userSettings = mongoService.GetCollection<UserSettingsDocument>("UserSettingsV2");
 
-			if (!databaseService.ReadOnlyMode)
+			if (!mongoService.ReadOnlyMode)
 			{
 				_users.Indexes.CreateOne(new CreateIndexModel<UserDocument>(Builders<UserDocument>.IndexKeys.Ascending(x => x.LoginUpper), new CreateIndexOptions { Unique = true }));
 				_users.Indexes.CreateOne(new CreateIndexModel<UserDocument>(Builders<UserDocument>.IndexKeys.Ascending(x => x.EmailUpper)));
diff --git a/Engine/Source/Programs/Horde/Horde.Build/Utilities/Clock.cs b/Engine/Source/Programs/Horde/Horde.Build/Utilities/Clock.cs
--- a/Engine/Source/Programs/Horde/Horde.Build/Utilities/Clock.cs
+++ b/Engine/Source/Programs/Horde/Horde.Build/Utilities/Clock.cs
@@ -6,7 +6,7 @@ using System.Diagnostics;
 using System.Threading;
 using System.Threading.Tasks;
 using EpicGames.Redis.Utility;
-using Horde.Build.Services;
+using Horde.Build.Server;
 using Microsoft.Extensions.Logging;
 using StackExchange.Redis;
 
diff --git a/Engine/Source/Programs/Horde/Horde.Build/Utilities/SingletonDocument.cs b/Engine/Source/Programs/Horde/Horde.Build/Utilities/SingletonDocument.cs
--- a/Engine/Source/Programs/Horde/Horde.Build/Utilities/SingletonDocument.cs
+++ b/Engine/Source/Programs/Horde/Horde.Build/Utilities/SingletonDocument.cs
@@ -4,7 +4,7 @@ using System;
 using System.Reflection;
 using System.Threading.Tasks;
 using Horde.Build.Models;
-using Horde.Build.Services;
+using Horde.Build.Server;
 using MongoDB.Bson;
 
 namespace Horde.Build.Utilities
@@ -59,7 +59,7 @@ namespace Horde.Build.Utilities
 		/// <summary>
 		/// The database service instance
 		/// </summary>
-		readonly DatabaseService _databaseService;
+		readonly MongoService _mongoService;
 
 		/// <summary>
 		/// Unique id for the singleton document
@@ -69,10 +69,10 @@ namespace Horde.Build.Utilities
 		/// <summary>
 		/// Constructor
 		/// </summary>
-		/// <param name="databaseService">The database service instance</param>
-		public SingletonDocument(DatabaseService databaseService)
+		/// <param name="mongoService">The database service instance</param>
+		public SingletonDocument(MongoService mongoService)
 		{
-			_databaseService = databaseService;
+			_mongoService = mongoService;
 
 			SingletonDocumentAttribute? attribute = typeof(T).GetCustomAttribute<SingletonDocumentAttribute>();
 			if (attribute == null)
@@ -86,18 +86,18 @@ namespace Horde.Build.Utilities
 		/// <summary>
 		/// Constructor
 		/// </summary>
-		/// <param name="databaseService">The database service instance</param>
+		/// <param name="mongoService">The database service instance</param>
 		/// <param name="objectId">The singleton document object id</param>
-		public SingletonDocument(DatabaseService databaseService, ObjectId objectId)
+		public SingletonDocument(MongoService mongoService, ObjectId objectId)
 		{
-			_databaseService = databaseService;
+			_mongoService = mongoService;
 			_objectId = objectId;
 		}
 
 		/// <inheritdoc/>
 		public async Task<T> GetAsync()
 		{
-			T value = await _databaseService.GetSingletonAsync<T>(_objectId);
+			T value = await _mongoService.GetSingletonAsync<T>(_objectId);
 			value.Id = _objectId;
 			return value;
 		}
@@ -105,7 +105,7 @@ namespace Horde.Build.Utilities
 		/// <inheritdoc/>
 		public Task<bool> TryUpdateAsync(T value)
 		{
-			return _databaseService.TryUpdateSingletonAsync<T>(value);
+			return _mongoService.TryUpdateSingletonAsync<T>(value);
 		}
 	}
 
diff --git a/Engine/Source/Programs/Horde/Horde.Build/Utilities/VersionedCollection.cs b/Engine/Source/Programs/Horde/Horde.Build/Utilities/VersionedCollection.cs
--- a/Engine/Source/Programs/Horde/Horde.Build/Utilities/VersionedCollection.cs
+++ b/Engine/Source/Programs/Horde/Horde.Build/Utilities/VersionedCollection.cs
@@ -3,7 +3,7 @@
 using System;
 using System.Collections.Generic;
 using System.Threading.Tasks;
-using Horde.Build.Services;
+using Horde.Build.Server;
 using Horde.Build.Utilities;
 using MongoDB.Bson;
 using MongoDB.Bson.Serialization;
@@ -83,12 +83,12 @@ namespace Horde.Build.Tools.Impl
 		/// <summary>
 		/// Constructor
 		/// </summary>
-		/// <param name="databaseService">The database service</param>
+		/// <param name="mongoService">The database service</param>
 		/// <param name="collectionName">Name of the collection of documents to manage</param>
 		/// <param name="redis">Instance of the redis service</param>
 		/// <param name="baseKey">Prefix for key types</param>
-		public VersionedCollection(DatabaseService databaseService, string collectionName, IDatabase redis, RedisKey baseKey)
-			: this(databaseService.Database.GetCollection<VersionedDocument<TId, TLatest>>(collectionName), redis, baseKey)
+		public VersionedCollection(MongoService mongoService, string collectionName, IDatabase redis, RedisKey baseKey)
+			: this(mongoService.Database.GetCollection<VersionedDocument<TId, TLatest>>(collectionName), redis, baseKey)
 		{
 		}
 
