Horde Storage - Improved name for the blob store consistency scope in datadog

diff --git a/Engine/Source/Programs/Horde/HordeStorage/Horde.Storage/Implementation/Blob/BlobStoreConsistencyCheckService.cs b/Engine/Source/Programs/Horde/HordeStorage/Horde.Storage/Implementation/Blob/BlobStoreConsistencyCheckService.cs
--- a/Engine/Source/Programs/Horde/HordeStorage/Horde.Storage/Implementation/Blob/BlobStoreConsistencyCheckService.cs
+++ b/Engine/Source/Programs/Horde/HordeStorage/Horde.Storage/Implementation/Blob/BlobStoreConsistencyCheckService.cs
@@ -89,7 +89,7 @@ namespace Horde.Storage.Implementation
 
                     await foreach ((BlobIdentifier blob, DateTime lastModified) in blobStore.ListObjects(ns))
                     {
-                        using IScope scope = Tracer.Instance.StartActive("consistency_check.blob");
+                        using IScope scope = Tracer.Instance.StartActive("consistency_check.blob_store");
                         scope.Span.ResourceName = $"{ns}.{blob}";
                         scope.Span.SetTag("BlobStore", blobStoreName);
 
