Horde Storage - Made the random start offset for polling service optional and removed it for most services (it doesnt add much for them and makes it harder to reason about when they should start).
Kept it for leader election and replication as these will hit external services and thus all starting at the same time causes unhealth spikes.

diff --git a/Engine/Source/Programs/Horde/HordeStorage/Horde.Storage/Implementation/LeaderElection/KubernetesLeaderElection.cs b/Engine/Source/Programs/Horde/HordeStorage/Horde.Storage/Implementation/LeaderElection/KubernetesLeaderElection.cs
--- a/Engine/Source/Programs/Horde/HordeStorage/Horde.Storage/Implementation/LeaderElection/KubernetesLeaderElection.cs
+++ b/Engine/Source/Programs/Horde/HordeStorage/Horde.Storage/Implementation/LeaderElection/KubernetesLeaderElection.cs
@@ -42,7 +42,7 @@ namespace Horde.Storage.Implementation.LeaderElection
         private ConfigMapLock _configMapLock;
         private string _identity;
 
-        public KubernetesLeaderElection(IOptionsMonitor<KubernetesLeaderElectionSettings> leaderSettings) : base("Kubernetes Leader Election", TimeSpan.FromSeconds(1), new KubernetesLeaderElectionState())
+        public KubernetesLeaderElection(IOptionsMonitor<KubernetesLeaderElectionSettings> leaderSettings) : base("Kubernetes Leader Election", TimeSpan.FromSeconds(1), new KubernetesLeaderElectionState(), startAtRandomTime: true)
         {
             KubernetesLeaderElectionSettings settings = leaderSettings.CurrentValue;
             // As we are determining if we are the leader we just assume we are running in a kubernetes cluster
diff --git a/Engine/Source/Programs/Horde/HordeStorage/Horde.Storage/Implementation/Replication/ReplicationService.cs b/Engine/Source/Programs/Horde/HordeStorage/Horde.Storage/Implementation/Replication/ReplicationService.cs
--- a/Engine/Source/Programs/Horde/HordeStorage/Horde.Storage/Implementation/Replication/ReplicationService.cs
+++ b/Engine/Source/Programs/Horde/HordeStorage/Horde.Storage/Implementation/Replication/ReplicationService.cs
@@ -34,7 +34,7 @@ namespace Horde.Storage.Implementation
         }
 
         public ReplicationService(IOptionsMonitor<ReplicationSettings> settings, IServiceProvider provider, ILeaderElection leaderElection) :
-            base(serviceName: nameof(ReplicationService), TimeSpan.FromSeconds(settings.CurrentValue.ReplicationPollFrequencySeconds), new ReplicationState())
+            base(serviceName: nameof(ReplicationService), TimeSpan.FromSeconds(settings.CurrentValue.ReplicationPollFrequencySeconds), new ReplicationState(), startAtRandomTime: true)
         {
             _settings = settings;
             _leaderElection = leaderElection;
diff --git a/Engine/Source/Programs/Horde/HordeStorage/Jupiter.Common/PollingService.cs b/Engine/Source/Programs/Horde/HordeStorage/Jupiter.Common/PollingService.cs
--- a/Engine/Source/Programs/Horde/HordeStorage/Jupiter.Common/PollingService.cs
+++ b/Engine/Source/Programs/Horde/HordeStorage/Jupiter.Common/PollingService.cs
@@ -25,17 +25,19 @@ namespace Jupiter
         private readonly string _serviceName;
         private readonly TimeSpan _pollFrequency;
         private readonly T _state;
+        private readonly bool _startAtRandomTime;
         private readonly CancellationTokenSource _stopPolling = new CancellationTokenSource();
         private readonly ManualResetEvent _hasFinishedRunning = new ManualResetEvent(true);
         private volatile bool _alreadyRunning = false;
         private Timer? _timer;
         private bool _disposed = false;
 
-        protected PollingService(string serviceName, TimeSpan pollFrequency, T state)
+        protected PollingService(string serviceName, TimeSpan pollFrequency, T state, bool startAtRandomTime = false)
         {
             _serviceName = serviceName;
             _pollFrequency = pollFrequency;
             _state = state;
+            _startAtRandomTime = startAtRandomTime;
 
             _timer = new Timer(OnUpdate, new ThreadState
             {
@@ -69,9 +71,13 @@ namespace Jupiter
 
             if (shouldPoll)
             {
-                // start at a random time between now and the poll frequency
-                int randomStartOffset = Random.Shared.Next(0, (int)_pollFrequency.TotalSeconds);
-                _timer?.Change(TimeSpan.FromSeconds(randomStartOffset), _pollFrequency);
+                int startOffset = 0;
+                if (_startAtRandomTime)
+                {
+                    // start at a random time between now and the poll frequency
+                    startOffset = Random.Shared.Next(0, (int)_pollFrequency.TotalSeconds);
+                }
+                _timer?.Change(TimeSpan.FromSeconds(startOffset), _pollFrequency);
             }
 
             return Task.CompletedTask;
