Horde Storage - Removed the remaining top level scopes for long running tasks as memory usage for datadog tracing is still to large.

diff --git a/Engine/Source/Programs/Horde/HordeStorage/Horde.Storage/Implementation/LastAccess/LastAccessService.cs b/Engine/Source/Programs/Horde/HordeStorage/Horde.Storage/Implementation/LastAccess/LastAccessService.cs
--- a/Engine/Source/Programs/Horde/HordeStorage/Horde.Storage/Implementation/LastAccess/LastAccessService.cs
+++ b/Engine/Source/Programs/Horde/HordeStorage/Horde.Storage/Implementation/LastAccess/LastAccessService.cs
@@ -65,7 +65,6 @@ namespace Horde.Storage.Implementation
 
         internal async Task<List<(RefRecord, DateTime)>> ProcessLastAccessRecords()
         {
-            using IScope _ = Tracer.Instance.StartActive("lastAccess.update");
             _logger.Information("Running Last Access Aggregation");
             List<(RefRecord, DateTime)> records = await _lastAccessCacheRecord.GetLastAccessedRecords();
             foreach ((RefRecord record, DateTime lastAccessTime) in records)
@@ -73,7 +72,7 @@ namespace Horde.Storage.Implementation
                 if (!ShouldTrackLastAccess(record.Namespace))
                     continue;
 
-                using IScope scope = Tracer.Instance.StartActive("lastAccess.update.record");
+                using IScope scope = Tracer.Instance.StartActive("lastAccess.update");
                 scope.Span.ResourceName = $"{record.Namespace}:{record.Bucket}.{record.RefName}";
                 _logger.Debug("Updating last access time to {LastAccessTime} for {Record}", lastAccessTime, record);
                 await _refsStore.UpdateLastAccessTime(record, lastAccessTime);
diff --git a/Engine/Source/Programs/Horde/HordeStorage/Horde.Storage/Implementation/LastAccess/LastAccessServiceReferences.cs b/Engine/Source/Programs/Horde/HordeStorage/Horde.Storage/Implementation/LastAccess/LastAccessServiceReferences.cs
--- a/Engine/Source/Programs/Horde/HordeStorage/Horde.Storage/Implementation/LastAccess/LastAccessServiceReferences.cs
+++ b/Engine/Source/Programs/Horde/HordeStorage/Horde.Storage/Implementation/LastAccess/LastAccessServiceReferences.cs
@@ -72,14 +72,13 @@ namespace Horde.Storage.Implementation
 
         internal async Task<List<(LastAccessRecord, DateTime)>> ProcessLastAccessRecords()
         {
-            using IScope _ = Tracer.Instance.StartActive("lastAccessRefs.update");
             _logger.Information("Running Last Access Aggregation for refs");
             List<(LastAccessRecord, DateTime)> records = await _lastAccessCacheRecord.GetLastAccessedRecords();
             foreach ((LastAccessRecord record, DateTime lastAccessTime) in records)
             {
                 if (!ShouldTrackLastAccess(record.Namespace))
                     continue;
-                using IScope scope = Tracer.Instance.StartActive("lastAccess.update.record");
+                using IScope scope = Tracer.Instance.StartActive("lastAccess.update");
                 scope.Span.ResourceName = $"{record.Namespace}:{record.Bucket}.{record.Key}";
                 _logger.Debug("Updating last access time to {LastAccessTime} for {Record}", lastAccessTime, record);
                 await _referencesStore.UpdateLastAccessTime(record.Namespace, record.Bucket, record.Key, lastAccessTime);
diff --git a/Engine/Source/Programs/Horde/HordeStorage/Horde.Storage/Implementation/Replication/RefsReplicator.cs b/Engine/Source/Programs/Horde/HordeStorage/Horde.Storage/Implementation/Replication/RefsReplicator.cs
--- a/Engine/Source/Programs/Horde/HordeStorage/Horde.Storage/Implementation/Replication/RefsReplicator.cs
+++ b/Engine/Source/Programs/Horde/HordeStorage/Horde.Storage/Implementation/Replication/RefsReplicator.cs
@@ -157,8 +157,8 @@ namespace Horde.Storage.Implementation
 
             try
             {
-                using IScope scope = Tracer.Instance.StartActive("replicator.run");
-                scope.Span.ResourceName =_name;
+                //using IScope scope = Tracer.Instance.StartActive("replicator.run");
+                //scope.Span.ResourceName =_name;
 
                 _replicationTokenSource.TryReset();
                 _replicationRunning = true;
@@ -391,7 +391,7 @@ namespace Horde.Storage.Implementation
                     cancellationTokenSource.Cancel();
                     return;
                 }
-                using IScope scope = Tracer.Instance.StartActive("replicator.replicate_op_incremental", new SpanCreationSettings {Parent = parentSpanContext});
+                using IScope scope = Tracer.Instance.StartActive("replicator.replicate_op_incremental");
                 scope.Span.ResourceName = $"[REDACTED_EMAIL]}[REDACTED_EMAIL]}";
 
                 _logger.Information("{Name} New transaction to replicate found. New op: {@Op} . Count of running replications: {CurrentReplications}", _name, @event, replicationTasks.Count);
diff --git a/Engine/Source/Programs/Horde/HordeStorage/Horde.Storage/Implementation/Replication/Replicator.cs b/Engine/Source/Programs/Horde/HordeStorage/Horde.Storage/Implementation/Replication/Replicator.cs
--- a/Engine/Source/Programs/Horde/HordeStorage/Horde.Storage/Implementation/Replication/Replicator.cs
+++ b/Engine/Source/Programs/Horde/HordeStorage/Horde.Storage/Implementation/Replication/Replicator.cs
@@ -171,8 +171,8 @@ namespace Horde.Storage.Implementation
             long countOfReplicatedEvents = 0L;
             try
             {
-                using IScope scope = Tracer.Instance.StartActive("replicator.run");
-                scope.Span.ResourceName =_name;
+                //using IScope scope = Tracer.Instance.StartActive("replicator.run");
+                //scope.Span.ResourceName =_name;
 
                 _replicationRunning = true;
                 _replicationFinishedEvent.Reset();
