Horde Storage - Fix for failing tests (null access then resolving parent dd context which didnt exist)

diff --git a/Engine/Source/Programs/Horde/HordeStorage/Horde.Storage/Implementation/Replication/RefsReplicator.cs b/Engine/Source/Programs/Horde/HordeStorage/Horde.Storage/Implementation/Replication/RefsReplicator.cs
--- a/Engine/Source/Programs/Horde/HordeStorage/Horde.Storage/Implementation/Replication/RefsReplicator.cs
+++ b/Engine/Source/Programs/Horde/HordeStorage/Horde.Storage/Implementation/Replication/RefsReplicator.cs
@@ -378,7 +378,6 @@ namespace Horde.Storage.Implementation
 
             CancellationTokenSource cancellationTokenSource = new CancellationTokenSource();
             CancellationTokenSource linkedTokenSource = CancellationTokenSource.CreateLinkedTokenSource(cancellationTokenSource.Token, replicationToken);
-            ISpanContext parentSpanContext = Tracer.Instance.ActiveScope.Span.Context;
 
             if (replicationToken.IsCancellationRequested)
                 return countOfReplicationsDone;
