Horde: Support for declaratively specifying indexes for Horde collections, which will be created/removed asynchronously by the server.

- MongoService.GetCollection() has an overload that takes a list of indexes.
- A background task will add and remove indexes to the collection to ensure that this list of indexes is reflected in the database.
- A shared Redis lock is acquired to ensure only one server instance can upgrade collections at a time.
- The application version is now expected to be in SemVer format, and is saved into a singleton document whenever an upgrade is made. Upgrades are only performed if the current server version is greater or equal to the value stored here, in order to prevent servers of different versions trying to perform conflicting operations.

diff --git a/Engine/Source/Programs/Horde/Horde.Build/Agents/Impl/AgentCollection.cs b/Engine/Source/Programs/Horde/Horde.Build/Agents/Impl/AgentCollection.cs
--- a/Engine/Source/Programs/Horde/Horde.Build/Agents/Impl/AgentCollection.cs
+++ b/Engine/Source/Programs/Horde/Horde.Build/Agents/Impl/AgentCollection.cs
@@ -153,15 +153,13 @@ namespace Horde.Build.Collections.Impl
 		/// </summary>
 		public AgentCollection(MongoService mongoService, RedisService redisService, IAuditLog<AgentId> auditLog)
 		{
-			_agents = mongoService.GetCollection<AgentDocument>("Agents");
+			List<MongoIndex<AgentDocument>> indexes = new List<MongoIndex<AgentDocument>>();
+			indexes.Add(keys => keys.Ascending(x => x.Deleted).Ascending(x => x.Id).Ascending(x => x.Pools));
+
+			_agents = mongoService.GetCollection<AgentDocument>("Agents", indexes);
 			_redisService = redisService;
 			_updateEventChannel = new RedisChannel<AgentId>("agents/notify");
 			_auditLog = auditLog;
-
-			if (!mongoService.ReadOnlyMode)
-			{
-				_agents.Indexes.CreateOne(new CreateIndexModel<AgentDocument>(Builders<AgentDocument>.IndexKeys.Ascending(x => x.Deleted).Ascending(x => x.Id).Ascending(x => x.Pools)));
-			}
 		}
 
 		/// <inheritdoc/>
diff --git a/Engine/Source/Programs/Horde/Horde.Build/Agents/Impl/LeaseCollection.cs b/Engine/Source/Programs/Horde/Horde.Build/Agents/Impl/LeaseCollection.cs
--- a/Engine/Source/Programs/Horde/Horde.Build/Agents/Impl/LeaseCollection.cs
+++ b/Engine/Source/Programs/Horde/Horde.Build/Agents/Impl/LeaseCollection.cs
@@ -69,33 +69,8 @@ namespace Horde.Build.Collections.Impl
 			}
 		}
 		
-		class DatabaseIndexes : BaseDatabaseIndexes
-		{
-			public readonly string AgentId;
-			public readonly string SessionId;
-			public readonly string StartTime;
-			public readonly string FinishTime;
-			public readonly string FinishTimeStartTimeCompound;
-
-			public DatabaseIndexes(IMongoCollection<LeaseDocument> leases, bool databaseReadOnlyMode) : base(databaseReadOnlyMode)
-			{
-				AgentId = CreateOrGetIndex(leases, "AgentId_1", Builders<LeaseDocument>.IndexKeys.Ascending(x => x.AgentId));
-				SessionId = CreateOrGetIndex(leases, "SessionId_1", Builders<LeaseDocument>.IndexKeys.Ascending(x => x.SessionId));
-				StartTime = CreateOrGetIndex(leases, "StartTime_1", Builders<LeaseDocument>.IndexKeys.Ascending(x => x.StartTime));
-				FinishTime = CreateOrGetIndex(leases, "FinishTime_1", Builders<LeaseDocument>.IndexKeys.Ascending(x => x.FinishTime));
-				FinishTimeStartTimeCompound = CreateOrGetIndex(leases, "FinishTime_1_StartTime_-1", Builders<LeaseDocument>.IndexKeys.Ascending(x => x.FinishTime).Descending(x => x.StartTime));
-			}
-		}
-
-		/// <summary>
-		/// Collection of lease documents
-		/// </summary>
 		readonly IMongoCollection<LeaseDocument> _leases;
-
-		/// <summary>
-		/// Creation and lookup of indexes
-		/// </summary>
-		readonly DatabaseIndexes _indexes;
+		readonly MongoIndex<LeaseDocument> _finishTimeStartTimeCompoundIndex;
 
 		/// <summary>
 		/// Constructor
@@ -103,8 +78,14 @@ namespace Horde.Build.Collections.Impl
 		/// <param name="mongoService">The database service instance</param>
 		public LeaseCollection(MongoService mongoService)
 		{
-			_leases = mongoService.GetCollection<LeaseDocument>("Leases");
-			_indexes = new DatabaseIndexes(_leases, mongoService.ReadOnlyMode);
+			List<MongoIndex<LeaseDocument>> indexes = new List<MongoIndex<LeaseDocument>>();
+			indexes.Add(keys => keys.Ascending(x => x.AgentId));
+			indexes.Add(keys => keys.Ascending(x => x.SessionId));
+			indexes.Add(keys => keys.Ascending(x => x.StartTime));
+			indexes.Add(keys => keys.Ascending(x => x.FinishTime));
+			indexes.Add(_finishTimeStartTimeCompoundIndex = MongoIndex.Create<LeaseDocument>(keys => keys.Ascending(x => x.FinishTime).Descending(x => x.StartTime)));
+
+			_leases = mongoService.GetCollection<LeaseDocument>("Leases", indexes);
 		}
 
 		/// <inheritdoc/>
@@ -190,7 +171,7 @@ namespace Horde.Build.Collections.Impl
 		/// <inheritdoc/>
 		public async Task<List<ILease>> FindLeasesAsync(DateTime? minTime, DateTime? maxTime)
 		{
-			return await FindLeasesAsync(null, null, minTime, maxTime, null, null, _indexes.FinishTimeStartTimeCompound, false);
+			return await FindLeasesAsync(null, null, minTime, maxTime, null, null, _finishTimeStartTimeCompoundIndex.Name, false);
 		}
 
 		/// <inheritdoc/>
diff --git a/Engine/Source/Programs/Horde/Horde.Build/Agents/Impl/SessionCollection.cs b/Engine/Source/Programs/Horde/Horde.Build/Agents/Impl/SessionCollection.cs
--- a/Engine/Source/Programs/Horde/Horde.Build/Agents/Impl/SessionCollection.cs
+++ b/Engine/Source/Programs/Horde/Horde.Build/Agents/Impl/SessionCollection.cs
@@ -71,14 +71,12 @@ namespace Horde.Build.Collections.Impl
 		/// <param name="mongoService">The database service</param>
 		public SessionCollection(MongoService mongoService)
 		{
-			_sessions = mongoService.GetCollection<SessionDocument>("Sessions");
+			List<MongoIndex<SessionDocument>> indexes = new List<MongoIndex<SessionDocument>>();
+			indexes.Add(keys => keys.Ascending(x => x.AgentId));
+			indexes.Add(keys => keys.Ascending(x => x.StartTime));
+			indexes.Add(keys => keys.Ascending(x => x.FinishTime));
 
-			if (!mongoService.ReadOnlyMode)
-			{
-				_sessions.Indexes.CreateOne(new CreateIndexModel<SessionDocument>(Builders<SessionDocument>.IndexKeys.Ascending(x => x.AgentId)));
-				_sessions.Indexes.CreateOne(new CreateIndexModel<SessionDocument>(Builders<SessionDocument>.IndexKeys.Ascending(x => x.StartTime)));
-				_sessions.Indexes.CreateOne(new CreateIndexModel<SessionDocument>(Builders<SessionDocument>.IndexKeys.Ascending(x => x.FinishTime)));
-			}
+			_sessions = mongoService.GetCollection<SessionDocument>("Sessions", indexes);
 		}
 
 		/// <inheritdoc/>
diff --git a/Engine/Source/Programs/Horde/Horde.Build/Auditing/Impl/AuditLog.cs b/Engine/Source/Programs/Horde/Horde.Build/Auditing/Impl/AuditLog.cs
--- a/Engine/Source/Programs/Horde/Horde.Build/Auditing/Impl/AuditLog.cs
+++ b/Engine/Source/Programs/Horde/Horde.Build/Auditing/Impl/AuditLog.cs
@@ -100,16 +100,15 @@ namespace Horde.Build.Collections.Impl
 
 		public AuditLog(MongoService mongoService, string collectionName, string subjectProperty, ILogger logger)
 		{
-			_messages = mongoService.GetCollection<AuditLogMessage>(collectionName);
+			List<MongoIndex<AuditLogMessage>> indexes = new List<MongoIndex<AuditLogMessage>>();
+			indexes.Add(builder => builder.Ascending(x => x.Subject).Descending(x => x.TimeUtc));
+
+			_messages = mongoService.GetCollection<AuditLogMessage>(collectionName, indexes);
+
 			_messageChannel = Channel.CreateUnbounded<AuditLogMessage>();
 			_subjectProperty = subjectProperty;
 			_logger = logger;
 
-			if (!mongoService.ReadOnlyMode)
-			{
-				_messages.Indexes.CreateOne(new CreateIndexModel<AuditLogMessage>(Builders<AuditLogMessage>.IndexKeys.Ascending(x => x.Subject).Descending(x => x.TimeUtc)));
-			}
-
 			_backgroundTask = Task.Run(() => WriteMessagesAsync());
 		}
 
diff --git a/Engine/Source/Programs/Horde/Horde.Build/Devices/Impl/DeviceCollection.cs b/Engine/Source/Programs/Horde/Horde.Build/Devices/Impl/DeviceCollection.cs
--- a/Engine/Source/Programs/Horde/Horde.Build/Devices/Impl/DeviceCollection.cs
+++ b/Engine/Source/Programs/Horde/Horde.Build/Devices/Impl/DeviceCollection.cs
@@ -250,17 +250,10 @@ namespace Horde.Build.Collections.Impl
 		/// </summary>
 		public DeviceCollection(MongoService mongoService)
 		{
-			_devices = mongoService.GetCollection<DeviceDocument>("Devices");
-			_platforms = mongoService.GetCollection<DevicePlatformDocument>("Devices.Platforms");
-			_pools = mongoService.GetCollection<DevicePoolDocument>("Devices.Pools");
+			_devices = mongoService.GetCollection<DeviceDocument>("Devices", keys => keys.Ascending(x => x.Name), unique: true);
+			_platforms = mongoService.GetCollection<DevicePlatformDocument>("Devices.Platforms", keys => keys.Ascending(x => x.Name), unique: true);
+			_pools = mongoService.GetCollection<DevicePoolDocument>("Devices.Pools", keys => keys.Ascending(x => x.Name), unique: true);
 			_reservations = mongoService.GetCollection<DeviceReservationDocument>("Devices.Reservations");
-
-			if (!mongoService.ReadOnlyMode)
-			{
-				_platforms.Indexes.CreateOne(new CreateIndexModel<DevicePlatformDocument>(Builders<DevicePlatformDocument>.IndexKeys.Ascending(x => x.Name), new CreateIndexOptions { Unique = true }));
-				_pools.Indexes.CreateOne(new CreateIndexModel<DevicePoolDocument>(Builders<DevicePoolDocument>.IndexKeys.Ascending(x => x.Name), new CreateIndexOptions { Unique = true }));
-				_devices.Indexes.CreateOne(new CreateIndexModel<DeviceDocument>(Builders<DeviceDocument>.IndexKeys.Ascending(x => x.Name), new CreateIndexOptions { Unique = true }));
-			}
 		}
 
 		/// <inheritdoc/>
diff --git a/Engine/Source/Programs/Horde/Horde.Build/Fleet/Impl/TelemetryCollection.cs b/Engine/Source/Programs/Horde/Horde.Build/Fleet/Impl/TelemetryCollection.cs
--- a/Engine/Source/Programs/Horde/Horde.Build/Fleet/Impl/TelemetryCollection.cs
+++ b/Engine/Source/Programs/Horde/Horde.Build/Fleet/Impl/TelemetryCollection.cs
@@ -105,12 +105,7 @@ namespace Horde.Build.Collections.Impl
 		/// <param name="database"></param>
 		public TelemetryCollection(MongoService database)
 		{
-			_utilization = database.GetCollection<UtilizationDocument>("Utilization");
-
-			if (!database.ReadOnlyMode)
-			{
-				_utilization.Indexes.CreateOne(new CreateIndexModel<UtilizationDocument>(Builders<UtilizationDocument>.IndexKeys.Ascending(x => x.FinishTime).Ascending(x => x.StartTime)));
-			}
+			_utilization = database.GetCollection<UtilizationDocument>("Utilization", keys => keys.Ascending(x => x.FinishTime).Ascending(x => x.StartTime));
 		}
 
 		/// <inheritdoc/>
diff --git a/Engine/Source/Programs/Horde/Horde.Build/Issues/Impl/IssueCollection.cs b/Engine/Source/Programs/Horde/Horde.Build/Issues/Impl/IssueCollection.cs
--- a/Engine/Source/Programs/Horde/Horde.Build/Issues/Impl/IssueCollection.cs
+++ b/Engine/Source/Programs/Horde/Horde.Build/Issues/Impl/IssueCollection.cs
@@ -402,27 +402,34 @@ namespace Horde.Build.Collections.Impl
 
 			_ledgerSingleton = new SingletonDocument<IssueLedger>(mongoService);
 
+			List<MongoIndex<Issue>> issueIndexes = new List<MongoIndex<Issue>>();
+			issueIndexes.Add(keys => keys.Ascending(x => x.ResolvedAt));
+			issueIndexes.Add(keys => keys.Ascending(x => x.VerifiedAt));
 			_issues = mongoService.GetCollection<Issue>("IssuesV2");
-			_issueSpans = mongoService.GetCollection<IssueSpan>("IssuesV2.Spans");
-			_issueSteps = mongoService.GetCollection<IssueStep>("IssuesV2.Steps");
-			_issueSuspects = mongoService.GetCollection<IssueSuspect>("IssuesV2.Suspects");
+
+			List<MongoIndex<IssueSpan>> issueSpanIndexes = new List<MongoIndex<IssueSpan>>();
+			issueSpanIndexes.Add(keys => keys.Ascending(x => x.IssueId));
+			issueSpanIndexes.Add(keys => keys.Ascending(x => x.StreamId).Ascending(x => x.MinChange).Ascending(x => x.MaxChange));
+			issueSpanIndexes.Add("StreamChanges", keys => keys.Ascending(x => x.StreamId).Ascending(x => x.TemplateRefId).Ascending(x => x.NodeName).Ascending(x => x.MinChange).Ascending(x => x.MaxChange));
+			_issueSpans = mongoService.GetCollection<IssueSpan>("IssuesV2.Spans", issueSpanIndexes);
+
+			List<MongoIndex<IssueStep>> issueStepIndexes = new List<MongoIndex<IssueStep>>();
+			issueStepIndexes.Add(keys => keys.Ascending(x => x.SpanId));
+			issueStepIndexes.Add(keys => keys.Ascending(x => x.JobId).Ascending(x => x.BatchId).Ascending(x => x.StepId));
+			_issueSteps = mongoService.GetCollection<IssueStep>("IssuesV2.Steps", issueStepIndexes);
+
+			List<MongoIndex<IssueSuspect>> issueSuspectIndexes = new List<MongoIndex<IssueSuspect>>();
+			issueSuspectIndexes.Add(keys => keys.Ascending(x => x.Change));
+			issueSuspectIndexes.Add(keys => keys.Ascending(x => x.AuthorId).Ascending(x => x.ResolvedAt));
+			issueSuspectIndexes.Add(keys => keys.Ascending(x => x.IssueId).Ascending(x => x.Change), unique: true);
+			_issueSuspects = mongoService.GetCollection<IssueSuspect>("IssuesV2.Suspects", issueSuspectIndexes);
+
 			_auditLog = auditLogFactory.Create("IssuesV2.History", "IssueId");
 
 			if (!mongoService.ReadOnlyMode)
 			{
-				_issues.Indexes.CreateOne(new CreateIndexModel<Issue>(Builders<Issue>.IndexKeys.Ascending(x => x.ResolvedAt)));
-				_issues.Indexes.CreateOne(new CreateIndexModel<Issue>(Builders<Issue>.IndexKeys.Ascending(x => x.VerifiedAt)));
-
-				_issueSpans.Indexes.CreateOne(new CreateIndexModel<IssueSpan>(Builders<IssueSpan>.IndexKeys.Ascending(x => x.IssueId)));
-				_issueSpans.Indexes.CreateOne(new CreateIndexModel<IssueSpan>(Builders<IssueSpan>.IndexKeys.Ascending(x => x.StreamId).Ascending(x => x.MinChange).Ascending(x => x.MaxChange)));
-				_issueSpans.Indexes.CreateOne(new CreateIndexModel<IssueSpan>(Builders<IssueSpan>.IndexKeys.Ascending(x => x.StreamId).Ascending(x => x.TemplateRefId).Ascending(x => x.NodeName).Ascending(x => x.MinChange).Ascending(x => x.MaxChange), new CreateIndexOptions { Name = "StreamChanges" }));
 				
-				_issueSteps.Indexes.CreateOne(new CreateIndexModel<IssueStep>(Builders<IssueStep>.IndexKeys.Ascending(x => x.SpanId)));
-				_issueSteps.Indexes.CreateOne(new CreateIndexModel<IssueStep>(Builders<IssueStep>.IndexKeys.Ascending(x => x.JobId).Ascending(x => x.BatchId).Ascending(x => x.StepId)));
 
-				_issueSuspects.Indexes.CreateOne(new CreateIndexModel<IssueSuspect>(Builders<IssueSuspect>.IndexKeys.Ascending(x => x.Change)));
-				_issueSuspects.Indexes.CreateOne(new CreateIndexModel<IssueSuspect>(Builders<IssueSuspect>.IndexKeys.Ascending(x => x.AuthorId).Ascending(x => x.ResolvedAt)));
-				_issueSuspects.Indexes.CreateOne(new CreateIndexModel<IssueSuspect>(Builders<IssueSuspect>.IndexKeys.Ascending(x => x.IssueId).Ascending(x => x.Change), new CreateIndexOptions { Unique = true }));
 			}
 		}
 
diff --git a/Engine/Source/Programs/Horde/Horde.Build/Jobs/Impl/ArtifactCollection.cs b/Engine/Source/Programs/Horde/Horde.Build/Jobs/Impl/ArtifactCollection.cs
--- a/Engine/Source/Programs/Horde/Horde.Build/Jobs/Impl/ArtifactCollection.cs
+++ b/Engine/Source/Programs/Horde/Horde.Build/Jobs/Impl/ArtifactCollection.cs
@@ -74,11 +74,7 @@ namespace Horde.Build.Services
 			_storageBackend = storageBackend;
 
 			// Initialize Artifacts table
-			_artifacts = mongoService.GetCollection<Artifact>("Artifacts");
-			if (!mongoService.ReadOnlyMode)
-			{
-				_artifacts.Indexes.CreateOne(new CreateIndexModel<Artifact>(Builders<Artifact>.IndexKeys.Ascending(x => x.JobId)));
-			}
+			_artifacts = mongoService.GetCollection<Artifact>("Artifacts", keys => keys.Ascending(x => x.JobId));
 		}
 
 		/// <summary>
diff --git a/Engine/Source/Programs/Horde/Horde.Build/Jobs/Impl/JobCollection.cs b/Engine/Source/Programs/Horde/Horde.Build/Jobs/Impl/JobCollection.cs
--- a/Engine/Source/Programs/Horde/Horde.Build/Jobs/Impl/JobCollection.cs
+++ b/Engine/Source/Programs/Horde/Horde.Build/Jobs/Impl/JobCollection.cs
@@ -318,32 +318,6 @@ namespace Horde.Build.Collections.Impl
 			}
 		}
 
-		class DatabaseIndexes : BaseDatabaseIndexes
-		{
-			public readonly string StreamId;
-			public readonly string Change;
-			public readonly string PreflightChange;
-			public readonly string CreateTimeUtc;
-			public readonly string UpdateTimeUtc;
-			public readonly string Name;
-			public readonly string StartedByUserId;
-			public readonly string TemplateId;
-			public readonly string SchedulePriority;
-
-			public DatabaseIndexes(IMongoCollection<JobDocument> jobs, bool databaseReadOnlyMode) : base(databaseReadOnlyMode)
-			{
-				StreamId = CreateOrGetIndex(jobs, "StreamId_1", Builders<JobDocument>.IndexKeys.Ascending(x => x.StreamId));
-				Change = CreateOrGetIndex(jobs, "Change_1", Builders<JobDocument>.IndexKeys.Ascending(x => x.Change));
-				PreflightChange = CreateOrGetIndex(jobs, "PreflightChange_1", Builders<JobDocument>.IndexKeys.Ascending(x => x.PreflightChange));
-				CreateTimeUtc = CreateOrGetIndex(jobs, "CreateTimeUtc_-1", Builders<JobDocument>.IndexKeys.Descending(x => x.CreateTimeUtc));
-				UpdateTimeUtc = CreateOrGetIndex(jobs, "UpdateTimeUtc_-1", Builders<JobDocument>.IndexKeys.Descending(x => x.UpdateTimeUtc));
-				Name = CreateOrGetIndex(jobs, "Name_1", Builders<JobDocument>.IndexKeys.Ascending(x => x.Name));
-				StartedByUserId = CreateOrGetIndex(jobs, "StartedByUserId_1", Builders<JobDocument>.IndexKeys.Ascending(x => x.StartedByUserId));
-				TemplateId = CreateOrGetIndex(jobs, "TemplateId_1", Builders<JobDocument>.IndexKeys.Ascending(x => x.TemplateId));
-				SchedulePriority = CreateOrGetIndex(jobs, "SchedulePriority_-1", Builders<JobDocument>.IndexKeys.Descending(x => x.SchedulePriority));
-			}
-		}
-
 		/// <summary>
 		/// Projection of a job definition to just include permissions info
 		/// </summary>
@@ -361,26 +335,12 @@ namespace Horde.Build.Collections.Impl
 		/// </summary>
 		const int MaxRetries = 2;
 
-		/// <summary>
-		/// The jobs collection
-		/// </summary>
 		readonly IMongoCollection<JobDocument> _jobs;
-
-		/// <summary>
-		/// Clock representing wall-clock time
-		/// </summary>
+		readonly MongoIndex<JobDocument> _createTimeIndex;
+		readonly MongoIndex<JobDocument> _updateTimeIndex;
 		readonly IClock _clock;
-
-		/// <summary>
-		/// Logger for output
-		/// </summary>
 		readonly ILogger<JobCollection> _logger;
 
-		/// <summary>
-		/// Creation and lookup of indexes
-		/// </summary>
-		readonly DatabaseIndexes _indexes;
-
 		/// <summary>
 		/// Constructor
 		/// </summary>
@@ -392,8 +352,17 @@ namespace Horde.Build.Collections.Impl
 			_clock = clock;
 			_logger = logger;
 
-			_jobs = mongoService.GetCollection<JobDocument>("Jobs");
-			_indexes = new DatabaseIndexes(_jobs, mongoService.ReadOnlyMode);
+			List<MongoIndex<JobDocument>> indexes = new List<MongoIndex<JobDocument>>();
+			indexes.Add(keys => keys.Ascending(x => x.StreamId));
+			indexes.Add(keys => keys.Ascending(x => x.Change));
+			indexes.Add(keys => keys.Ascending(x => x.PreflightChange));
+			indexes.Add(_createTimeIndex = MongoIndex.Create<JobDocument>(keys => keys.Descending(x => x.CreateTimeUtc)));
+			indexes.Add(_updateTimeIndex = MongoIndex.Create<JobDocument>(keys => keys.Descending(x => x.UpdateTimeUtc)));
+			indexes.Add(keys => keys.Ascending(x => x.Name));
+			indexes.Add(keys => keys.Ascending(x => x.StartedByUserId));
+			indexes.Add(keys => keys.Ascending(x => x.TemplateId));
+			indexes.Add(keys => keys.Descending(x => x.SchedulePriority));
+			_jobs = mongoService.GetCollection<JobDocument>("Jobs", indexes);
 		}
 
 		static Task PostLoadAsync(JobDocument job)
@@ -579,10 +548,10 @@ namespace Horde.Build.Collections.Impl
 		/// <inheritdoc/>
 		public async Task<List<IJob>> FindLatestByStreamWithTemplatesAsync(StreamId streamId, TemplateRefId[] templates, UserId? preflightStartedByUser, DateTimeOffset? maxCreateTime, DateTimeOffset? modifiedAfter, int? index, int? count, bool consistentRead)
 		{
-			string indexHint = _indexes.CreateTimeUtc;
+			string indexHint = _createTimeIndex.Name;
 			if (modifiedAfter != null)
 			{
-				indexHint = _indexes.UpdateTimeUtc;
+				indexHint = _updateTimeIndex.Name;
 			}
 			
 			// This find call uses an index hint. Modifying the parameter passed to FindAsync can affect execution time a lot as the query planner is forced to use the specified index.
diff --git a/Engine/Source/Programs/Horde/Horde.Build/Jobs/Impl/JobStepRefCollection.cs b/Engine/Source/Programs/Horde/Horde.Build/Jobs/Impl/JobStepRefCollection.cs
--- a/Engine/Source/Programs/Horde/Horde.Build/Jobs/Impl/JobStepRefCollection.cs
+++ b/Engine/Source/Programs/Horde/Horde.Build/Jobs/Impl/JobStepRefCollection.cs
@@ -96,12 +96,7 @@ namespace Horde.Build.Collections.Impl
 		/// <param name="mongoService">The database service instance</param>
 		public JobStepRefCollection(MongoService mongoService)
 		{
-			_jobStepRefs = mongoService.GetCollection<JobStepRef>("JobStepRefs");
-
-			if (!mongoService.ReadOnlyMode)
-			{
-				_jobStepRefs.Indexes.CreateOne(new CreateIndexModel<JobStepRef>(Builders<JobStepRef>.IndexKeys.Ascending(x => x.StreamId).Ascending(x => x.TemplateId).Ascending(x => x.Name).Descending(x => x.Change)));
-			}
+			_jobStepRefs = mongoService.GetCollection<JobStepRef>("JobStepRefs", keys => keys.Ascending(x => x.StreamId).Ascending(x => x.TemplateId).Ascending(x => x.Name).Descending(x => x.Change));
 		}
 
 		/// <inheritdoc/>
diff --git a/Engine/Source/Programs/Horde/Horde.Build/Jobs/Impl/TestDataCollection.cs b/Engine/Source/Programs/Horde/Horde.Build/Jobs/Impl/TestDataCollection.cs
--- a/Engine/Source/Programs/Horde/Horde.Build/Jobs/Impl/TestDataCollection.cs
+++ b/Engine/Source/Programs/Horde/Horde.Build/Jobs/Impl/TestDataCollection.cs
@@ -65,13 +65,10 @@ namespace Horde.Build.Collections.Impl
 		/// <param name="mongoService">The database service instance</param>
 		public TestDataCollection(MongoService mongoService)
 		{
-			_testDataDocuments = mongoService.GetCollection<TestDataDocument>("TestData");
-
-			if (!mongoService.ReadOnlyMode)
-			{
-				_testDataDocuments.Indexes.CreateOne(new CreateIndexModel<TestDataDocument>(Builders<TestDataDocument>.IndexKeys.Ascending(x => x.StreamId).Ascending(x => x.Change).Ascending(x => x.Key)));
-				_testDataDocuments.Indexes.CreateOne(new CreateIndexModel<TestDataDocument>(Builders<TestDataDocument>.IndexKeys.Ascending(x => x.JobId).Ascending(x => x.StepId).Ascending(x => x.Key), new CreateIndexOptions { Unique = true }));
-			}
+			List<MongoIndex<TestDataDocument>> indexes = new List<MongoIndex<TestDataDocument>>();
+			indexes.Add(keys => keys.Ascending(x => x.StreamId).Ascending(x => x.Change).Ascending(x => x.Key));
+			indexes.Add(keys => keys.Ascending(x => x.JobId).Ascending(x => x.StepId).Ascending(x => x.Key), unique: true);
+			_testDataDocuments = mongoService.GetCollection<TestDataDocument>("TestData", indexes);
 		}
 
 		/// <inheritdoc/>
diff --git a/Engine/Source/Programs/Horde/Horde.Build/Logs/Impl/LogEventCollection.cs b/Engine/Source/Programs/Horde/Horde.Build/Logs/Impl/LogEventCollection.cs
--- a/Engine/Source/Programs/Horde/Horde.Build/Logs/Impl/LogEventCollection.cs
+++ b/Engine/Source/Programs/Horde/Horde.Build/Logs/Impl/LogEventCollection.cs
@@ -112,15 +112,12 @@ namespace Horde.Build.Collections.Impl
 		{
 			_logger = logger;
 
-			_logEvents = mongoService.GetCollection<LogEventDocument>("LogEvents");
-			_legacyEvents = mongoService.GetCollection<LegacyEventDocument>("Events");
+			List<MongoIndex<LogEventDocument>> logEventIndexes = new List<MongoIndex<LogEventDocument>>();
+			logEventIndexes.Add(keys => keys.Ascending(x => x.Id.LogId));
+			logEventIndexes.Add(keys => keys.Ascending(x => x.SpanId).Ascending(x => x.Id));
+			_logEvents = mongoService.GetCollection<LogEventDocument>("LogEvents", logEventIndexes);
 
-			if (!mongoService.ReadOnlyMode)
-			{
-				_logEvents.Indexes.CreateOne(new CreateIndexModel<LogEventDocument>(Builders<LogEventDocument>.IndexKeys.Ascending(x => x.Id.LogId)));
-				_logEvents.Indexes.CreateOne(new CreateIndexModel<LogEventDocument>(Builders<LogEventDocument>.IndexKeys.Ascending(x => x.SpanId).Ascending(x => x.Id)));
-				_legacyEvents.Indexes.CreateOne(new CreateIndexModel<LegacyEventDocument>(Builders<LegacyEventDocument>.IndexKeys.Ascending(x => x.LogId)));
-			}
+			_legacyEvents = mongoService.GetCollection<LegacyEventDocument>("Events", keys => keys.Ascending(x => x.LogId));
 		}
 
 		/// <inheritdoc/>
diff --git a/Engine/Source/Programs/Horde/Horde.Build/Notifications/Impl/SubscriptionCollection.cs b/Engine/Source/Programs/Horde/Horde.Build/Notifications/Impl/SubscriptionCollection.cs
--- a/Engine/Source/Programs/Horde/Horde.Build/Notifications/Impl/SubscriptionCollection.cs
+++ b/Engine/Source/Programs/Horde/Horde.Build/Notifications/Impl/SubscriptionCollection.cs
@@ -187,12 +187,10 @@ namespace Horde.Build.Collections.Impl
 		/// <param name="mongoService">The database service</param>
 		public SubscriptionCollection(MongoService mongoService)
 		{
-			_collection = mongoService.GetCollection<Subscription>("SubscriptionsV2");
-			if (!mongoService.ReadOnlyMode)
-			{
-				_collection.Indexes.CreateOne(new CreateIndexModel<Subscription>(Builders<Subscription>.IndexKeys.Ascending(x => x.Event)));
-				_collection.Indexes.CreateOne(new CreateIndexModel<Subscription>(Builders<Subscription>.IndexKeys.Ascending(x => x.UserId)));
-			}
+			List<MongoIndex<Subscription>> indexes = new List<MongoIndex<Subscription>>();
+			indexes.Add(keys => keys.Ascending(x => x.Event));
+			indexes.Add(keys => keys.Ascending(x => x.UserId));
+			_collection = mongoService.GetCollection<Subscription>("SubscriptionsV2", indexes);
 		}
 
 		/// <inheritdoc/>
diff --git a/Engine/Source/Programs/Horde/Horde.Build/Perforce/Impl/CommitCollection.cs b/Engine/Source/Programs/Horde/Horde.Build/Perforce/Impl/CommitCollection.cs
--- a/Engine/Source/Programs/Horde/Horde.Build/Perforce/Impl/CommitCollection.cs
+++ b/Engine/Source/Programs/Horde/Horde.Build/Perforce/Impl/CommitCollection.cs
@@ -84,12 +84,7 @@ namespace Horde.Build.Commits.Impl
 		/// </summary>
 		public CommitCollection(MongoService mongoService)
 		{
-			_commits = mongoService.Database.GetCollection<Commit>("Commits");
-
-			if (!mongoService.ReadOnlyMode)
-			{
-				_commits.Indexes.CreateOne(new CreateIndexModel<Commit>(Builders<Commit>.IndexKeys.Ascending(x => x.StreamId).Descending(x => x.Change), new CreateIndexOptions { Unique = true }));
-			}
+			_commits = mongoService.GetCollection<Commit>("Commits", keys => keys.Ascending(x => x.StreamId).Descending(x => x.Change), unique: true);
 		}
 
 		/// <inheritdoc/>
diff --git a/Engine/Source/Programs/Horde/Horde.Build/Program.cs b/Engine/Source/Programs/Horde/Horde.Build/Program.cs
--- a/Engine/Source/Programs/Horde/Horde.Build/Program.cs
+++ b/Engine/Source/Programs/Horde/Horde.Build/Program.cs
@@ -2,6 +2,7 @@
 
 using System;
 using System.Collections.Generic;
+using System.Diagnostics;
 using System.IO;
 using System.Reflection;
 using System.Runtime.InteropServices;
@@ -101,6 +102,8 @@ namespace Horde.Build
 
 	class Program
 	{
+		public static SemVer Version => _version;
+
 		public static DirectoryReference AppDir { get; } = GetAppDir();
 
 		public static DirectoryReference DataDir { get; } = GetDefaultDataDir();
@@ -109,6 +112,8 @@ namespace Horde.Build
 
 		public static Type[] ConfigSchemas = FindSchemaTypes();
 
+		static SemVer _version;
+
 		static Type[] FindSchemaTypes()
 		{
 			List<Type> schemaTypes = new List<Type>();
@@ -124,6 +129,16 @@ namespace Horde.Build
 
 		public static async Task<int> Main(string[] args)
 		{
+			FileVersionInfo versionInfo = FileVersionInfo.GetVersionInfo(Assembly.GetExecutingAssembly().Location);
+			if (String.IsNullOrEmpty(versionInfo.ProductVersion))
+			{
+				_version = SemVer.Parse("0.0.0");
+			}
+			else
+			{
+				_version = SemVer.Parse(versionInfo.ProductVersion);
+			}
+
 			CommandLineArguments arguments = new CommandLineArguments(args);
 
 			IConfiguration config = new ConfigurationBuilder()
@@ -156,6 +171,8 @@ namespace Horde.Build
 				.ReadFrom.Configuration(config)
 				.CreateLogger();
 
+			Serilog.Log.Logger.Information("Server version: {Version}", Version);
+
 			if (hordeSettings.WithDatadog)
 			{
 				GlobalTracer.Register(Datadog.Trace.OpenTracing.OpenTracingTracerFactory.WrapTracer(Datadog.Trace.Tracer.Instance));
diff --git a/Engine/Source/Programs/Horde/Horde.Build/Server/MongoService.cs b/Engine/Source/Programs/Horde/Horde.Build/Server/MongoService.cs
--- a/Engine/Source/Programs/Horde/Horde.Build/Server/MongoService.cs
+++ b/Engine/Source/Programs/Horde/Horde.Build/Server/MongoService.cs
@@ -13,21 +13,188 @@ using System.Runtime.InteropServices;
 using System.Security.Cryptography.X509Certificates;
 using System.Text;
 using System.Text.RegularExpressions;
+using System.Threading;
+using System.Threading.Channels;
 using System.Threading.Tasks;
 using EpicGames.Core;
+using EpicGames.Redis.Utility;
 using Horde.Build.Models;
 using Horde.Build.Services;
+using Horde.Build.Utilities;
 using Horde.Build.Utiltiies;
+using Microsoft.Extensions.Hosting;
 using Microsoft.Extensions.Logging;
 using Microsoft.Extensions.Options;
 using Microsoft.IdentityModel.Tokens;
 using MongoDB.Bson;
 using MongoDB.Bson.Serialization;
+using MongoDB.Bson.Serialization.Attributes;
 using MongoDB.Driver;
 using MongoDB.Driver.Core.Events;
+using StackExchange.Redis;
 
 namespace Horde.Build.Server
 {
+	/// <summary>
+	/// Constrained set of parameters for building a mongo index
+	/// </summary>
+	[DebuggerDisplay("{Name}")]
+	public class MongoIndex : IEquatable<MongoIndex>
+	{
+		/// <summary>
+		/// Name of the index
+		/// </summary>
+		[BsonElement("name"), BsonRequired]
+		public string Name { get; private set; }
+
+		/// <summary>
+		/// Keys for the index
+		/// </summary>
+		[BsonElement("key"), BsonRequired]
+		public BsonDocument KeysDocument { get; private set; }
+
+		/// <summary>
+		/// The index should be unique
+		/// </summary>
+		[BsonElement("unique")]
+		public bool Unique { get; private set; }
+
+		/// <summary>
+		/// Constructor
+		/// </summary>
+		[BsonConstructor]
+		private MongoIndex()
+		{
+			Name = String.Empty;
+			KeysDocument = new BsonDocument();
+		}
+
+		/// <summary>
+		/// Constructor
+		/// </summary>
+		/// <param name="keysDocument">Keys for the index</param>
+		/// <param name="unique">Whether the index should be unique</param>
+		public MongoIndex(BsonDocument keysDocument, bool unique)
+			: this(GetDefaultName(keysDocument), keysDocument, unique)
+		{
+		}
+
+		/// <summary>
+		/// Constructor
+		/// </summary>
+		/// <param name="name">Name of the index</param>
+		/// <param name="keysDocument">Keys for the index</param>
+		/// <param name="unique">Whether the index should be unique</param>
+		public MongoIndex(string name, BsonDocument keysDocument, bool unique)
+		{
+			Name = name;
+			KeysDocument = keysDocument;
+			Unique = unique;
+		}
+
+		/// <summary>
+		/// Constructor
+		/// </summary>
+		/// <param name="keysFunc">Callback to configure keys for this document</param>
+		/// <param name="unique">Whether the index should be unique</param>
+		public static MongoIndex<T> Create<T>(Func<IndexKeysDefinitionBuilder<T>, IndexKeysDefinition<T>> keysFunc, bool unique = false)
+		{
+			return new MongoIndex<T>(keysFunc(Builders<T>.IndexKeys), unique);
+		}
+
+		/// <summary>
+		/// Constructor
+		/// </summary>
+		/// <param name="name">Name of the index</param>
+		/// <param name="keysFunc">Callback to configure keys for this document</param>
+		/// <param name="unique">Whether the index should be unique</param>
+		public static MongoIndex<T> Create<T>(string name, Func<IndexKeysDefinitionBuilder<T>, IndexKeysDefinition<T>> keysFunc, bool unique = false)
+		{
+			return new MongoIndex<T>(name, keysFunc(Builders<T>.IndexKeys), unique);
+		}
+
+		/// <inheritdoc/>
+		public override bool Equals(object? obj) => obj is MongoIndex other && Equals(other);
+
+		/// <inheritdoc/>
+		public override int GetHashCode() => Name.GetHashCode(StringComparison.Ordinal);
+
+		/// <inheritdoc/>
+		public bool Equals(MongoIndex? other) => other != null && Name.Equals(other.Name, StringComparison.Ordinal) && KeysDocument.Equals(other.KeysDocument) && Unique == other.Unique;
+
+		/// <summary>
+		/// Gets the default name for an index based on its keys
+		/// </summary>
+		/// <param name="keys">Keys for the index</param>
+		/// <returns>Name of the index</returns>
+		protected static string GetDefaultName(BsonDocument keys)
+		{
+			StringBuilder name = new StringBuilder();
+			foreach (BsonElement element in keys.Elements)
+			{
+				if (name.Length > 0)
+				{
+					name.Append('_');
+				}
+				name.Append(element.Name);
+				name.Append('_');
+				name.Append(element.Value.ToString());
+			}
+			return name.ToString();
+		}
+	}
+
+	/// <summary>
+	/// Strongly typed index document
+	/// </summary>
+	/// <typeparam name="T"></typeparam>
+	public class MongoIndex<T> : MongoIndex
+	{
+		/// <summary>
+		/// Keys for the index
+		/// </summary>
+		public IndexKeysDefinition<T> Keys { get; }
+
+		/// <summary>
+		/// Constructor
+		/// </summary>
+		/// <param name="keys">Keys for the index</param>
+		/// <param name="unique">Whether the index should be unique</param>
+		public MongoIndex(IndexKeysDefinition<T> keys, bool unique)
+			: base(keys.Render(BsonSerializer.LookupSerializer<T>(), BsonSerializer.SerializerRegistry), unique)
+		{
+			Keys = keys;
+		}
+
+		/// <summary>
+		/// Constructor
+		/// </summary>
+		/// <param name="name">Name of the index</param>
+		/// <param name="keys">Keys for the index</param>
+		/// <param name="unique">Whether the index should be unique</param>
+		public MongoIndex(string name, IndexKeysDefinition<T> keys, bool unique)
+			: base(name, keys.Render(BsonSerializer.LookupSerializer<T>(), BsonSerializer.SerializerRegistry), unique)
+		{
+			Keys = keys;
+		}
+	}
+
+	/// <summary>
+	/// Extension methods for <see cref="MongoIndex{T}"/>
+	/// </summary>
+	static class MongoIndexExtensions
+	{
+		public static void Add<T>(this List<MongoIndex<T>> list, Func<IndexKeysDefinitionBuilder<T>, IndexKeysDefinition<T>> keyFunc, bool unique = false)
+		{
+			list.Add(MongoIndex.Create(keyFunc, unique));
+		}
+
+		public static void Add<T>(this List<MongoIndex<T>> list, string name, Func<IndexKeysDefinitionBuilder<T>, IndexKeysDefinition<T>> keyFunc, bool unique = false)
+		{
+			list.Add(MongoIndex.Create(name, keyFunc, unique));
+		}
+	}
+
 	/// <summary>
 	/// Singleton for accessing the database
 	/// </summary>
@@ -98,6 +265,9 @@ namespace Horde.Build.Server
 		/// </summary>
 		const int DefaultMongoPort = 27017;
 
+		readonly HashSet<string> _collectionNames = new HashSet<string>(StringComparer.Ordinal);
+		readonly Channel<Func<CancellationToken, Task>> _updateIndexesChannel = Channel.CreateUnbounded<Func<CancellationToken, Task>>();
+
 		/// <summary>
 		/// Constructor
 		/// </summary>
@@ -483,6 +653,125 @@ namespace Horde.Build.Server
 			return new MongoTracingCollection<T>(Database.GetCollection<T>(name));
 		}
 
+		/// <summary>
+		/// Get a MongoDB collection from database with a single index
+		/// </summary>
+		/// <param name="name">Name of collection</param>
+		/// <param name="keysFunc">Method to configure keys for the collection</param>
+		/// <param name="unique">Whether a unique index is required</param>
+		/// <typeparam name="T">A MongoDB document</typeparam>
+		/// <returns></returns>
+		public IMongoCollection<T> GetCollection<T>(string name, Func<IndexKeysDefinitionBuilder<T>, IndexKeysDefinition<T>> keysFunc, bool unique = false)
+		{
+			List<MongoIndex<T>> indexes = new List<MongoIndex<T>>();
+			indexes.Add(MongoIndex.Create<T>(keysFunc, unique));
+			return GetCollection(name, indexes);
+		}
+
+		/// <summary>
+		/// Get a MongoDB collection from database
+		/// </summary>
+		/// <param name="name">Name of collection</param>
+		/// <param name="indexes">Indexes for the collection</param>
+		/// <typeparam name="T">A MongoDB document</typeparam>
+		/// <returns></returns>
+		public IMongoCollection<T> GetCollection<T>(string name, IEnumerable<MongoIndex<T>> indexes)
+		{
+			IMongoCollection<T> collection = GetCollection<T>(name);
+			lock (_collectionNames)
+			{
+				if (!_collectionNames.Add(name))
+				{
+					throw new NotImplementedException();
+				}
+
+				MongoIndex<T>[] indexesCopy = indexes.ToArray();
+				_updateIndexesChannel.Writer.TryWrite(ctx => UpdateIndexesAsync(name, collection, indexesCopy, ctx));
+			}
+			return collection;
+		}
+
+		/// <summary>
+		/// Pops an upgrade task from the queue. This method should only be called by <see cref="MongoUpgradeService"/>
+		/// </summary>
+		/// <param name="cancellationToken"></param>
+		/// <returns></returns>
+		public ValueTask<Func<CancellationToken, Task>> ReadNextUpgradeTask(CancellationToken cancellationToken)
+		{
+			return _updateIndexesChannel.Reader.ReadAsync(cancellationToken);
+		}
+
+		private async Task UpdateIndexesAsync<T>(string collectionName, IMongoCollection<T> collection, MongoIndex<T>[] newIndexes, CancellationToken cancellationToken)
+		{
+			// Find all the current indexes, excluding the default
+			Dictionary<string, MongoIndex> nameToExistingIndex = new Dictionary<string, MongoIndex>(StringComparer.Ordinal);
+			using (IAsyncCursor<BsonDocument> cursor = await collection.Indexes.ListAsync(cancellationToken))
+			{
+				while (await cursor.MoveNextAsync(cancellationToken))
+				{
+					foreach (BsonDocument document in cursor.Current)
+					{
+						MongoIndex indexDocument = BsonSerializer.Deserialize<MongoIndex>(document);
+						if (!indexDocument.Name.Equals("_id_", StringComparison.Ordinal))
+						{
+							nameToExistingIndex.Add(indexDocument.Name, indexDocument);
+						}
+					}
+				}
+			}
+
+			// Figure out which indexes to drop
+			List<MongoIndex<T>> createIndexes = new List<MongoIndex<T>>();
+			HashSet<string> removeIndexNames = new HashSet<string>(nameToExistingIndex.Keys, StringComparer.Ordinal);
+			foreach (MongoIndex<T> newIndex in newIndexes)
+			{
+				MongoIndex? existingIndex;
+				if (nameToExistingIndex.TryGetValue(newIndex.Name, out existingIndex))
+				{
+					if (existingIndex.Equals(newIndex))
+					{
+						removeIndexNames.Remove(newIndex.Name);
+						continue;
+					}
+				}
+				createIndexes.Add(newIndex);
+			}
+
+			// Drop any indexes that are no longer needed
+			foreach (string removeIndexName in removeIndexNames)
+			{
+				if (ReadOnlyMode)
+				{
+					_logger.LogWarning("Would drop unused index {CollectionName}.{IndexName} - skipping due to read-only setting.", collectionName, removeIndexName);
+				}
+				else
+				{
+					_logger.LogInformation("Dropping unused index {IndexName}", removeIndexName);
+					await collection.Indexes.DropOneAsync(removeIndexName, cancellationToken);
+				}
+			}
+
+			// Create all the new indexes
+			foreach (MongoIndex<T> createIndex in createIndexes)
+			{
+				if (ReadOnlyMode)
+				{
+					_logger.LogWarning("Would create index {CollectionName}.{IndexName} - skipping due to read-only setting.", collectionName, createIndex.Name);
+				}
+				else
+				{
+					_logger.LogInformation("Creating index {CollectionName}: {IndexName}", collectionName, createIndex.Name);
+
+					CreateIndexOptions<T> options = new CreateIndexOptions<T>();
+					options.Name = createIndex.Name;
+					options.Unique = createIndex.Unique;
+
+					CreateIndexModel<T> model = new CreateIndexModel<T>(createIndex.Keys, options);
+					await collection.Indexes.CreateOneAsync(model, cancellationToken: cancellationToken);
+				}
+			}
+		}
+
 		/// <summary>
 		/// Gets a singleton document by id
 		/// </summary>
@@ -579,4 +868,128 @@ namespace Horde.Build.Server
 			}
 		}
 	}
+
+	/// <summary>
+	/// Stores the version number of the latest server instance to have upgraded the database schema
+	/// </summary>
+	[SingletonDocument("62470d8508d48eddfac7b55d")]
+	public class MongoSchemaDocument : SingletonBase
+	{
+		/// <summary>
+		/// Current version number
+		/// </summary>
+		public string? Version { get; set; }
+	}
+
+	/// <summary>
+	/// Service which updates the database state
+	/// </summary>
+	public sealed class MongoUpgradeService : IDisposable, IHostedService
+	{
+		static readonly RedisKey s_schemaLockKey = new RedisKey("server/schema-upgrade/lock");
+
+		readonly MongoService _mongoService;
+		readonly RedisService _redisService;
+		readonly ISingletonDocument<MongoSchemaDocument> _schemaDocument;
+		readonly ILogger _logger;
+		readonly BackgroundTask _updateIndexesTask;
+
+		/// <summary>
+		/// Constructor
+		/// </summary>
+		public MongoUpgradeService(MongoService mongoService, RedisService redisService, ISingletonDocument<MongoSchemaDocument> schemaDocument, ILogger<MongoUpgradeService> logger)
+		{
+			_mongoService = mongoService;
+			_redisService = redisService;
+			_schemaDocument = schemaDocument;
+			_logger = logger;
+			_updateIndexesTask = new BackgroundTask(UpdateAsync);
+		}
+
+		/// <inheritdoc/>
+		public void Dispose()
+		{
+			_updateIndexesTask.Dispose();
+		}
+
+		/// <inheritdoc/>
+		public Task StartAsync(CancellationToken cancellationToken)
+		{
+			_updateIndexesTask.Start();
+			return Task.CompletedTask;
+		}
+
+		/// <inheritdoc/>
+		public async Task StopAsync(CancellationToken cancellationToken)
+		{
+			await _updateIndexesTask.StopAsync();
+		}
+
+		async Task UpdateAsync(CancellationToken cancellationToken)
+		{
+			for (; ; )
+			{
+				Func<CancellationToken, Task> updateIndexTask = await _mongoService.ReadNextUpgradeTask(cancellationToken);
+				for (; ; )
+				{
+					using (RedisLock schemaLock = new RedisLock(_redisService.Database, s_schemaLockKey))
+					{
+						if (await schemaLock.AcquireAsync(TimeSpan.FromMinutes(5.0)) && await UpdateOneAsync(updateIndexTask, cancellationToken))
+						{
+							break;
+						}
+					}
+					await Task.Delay(TimeSpan.FromMinutes(1.0), cancellationToken);
+				}
+			}
+		}
+
+		async ValueTask<bool> UpdateOneAsync(Func<CancellationToken, Task> taskAsync, CancellationToken cancellationToken)
+		{
+			try
+			{
+				// Check we're not downgrading the data
+				for (; ; )
+				{
+					MongoSchemaDocument currentSchema = await _schemaDocument.GetAsync();
+					if (!String.IsNullOrEmpty(currentSchema.Version))
+					{
+						SemVer currentVersion = SemVer.Parse(currentSchema.Version);
+						if (Program.Version < currentVersion)
+						{
+							return false;
+						}
+						if (Program.Version == currentVersion)
+						{
+							break;
+						}
+					}
+					currentSchema.Version = Program.Version.ToString();
+
+					if (_mongoService.ReadOnlyMode) // Pretend like we've upgraded the document so that the task can run and log its changes; it will guard writing while this flag is set itself.
+					{
+						break;
+					}
+					if (await _schemaDocument.TryUpdateAsync(currentSchema))
+					{
+						break;
+					}
+				}
+
+				// Perform the upgrade
+				await taskAsync(cancellationToken);
+				return true;
+			}
+			catch (MongoCommandException ex)
+			{
+				_logger.LogWarning(ex, "Command exception while attempting to update indexes ({Code})", ex.Code);
+				return false;
+			}
+			catch (Exception ex)
+			{
+				_logger.LogWarning(ex, "Exception while attempting to update indexes");
+				return false;
+			}
+		}
+	}
 }
diff --git a/Engine/Source/Programs/Horde/Horde.Build/Startup.cs b/Engine/Source/Programs/Horde/Horde.Build/Startup.cs
--- a/Engine/Source/Programs/Horde/Horde.Build/Startup.cs
+++ b/Engine/Source/Programs/Horde/Horde.Build/Startup.cs
@@ -547,6 +547,8 @@ namespace Horde.Build
 
 			if (settings.IsRunModeActive(RunMode.Worker))
 			{
+				services.AddHostedService<MongoUpgradeService>();
+
 				services.AddHostedService(provider => provider.GetRequiredService<AutoscaleServiceV2>());
 				
 				services.AddHostedService(provider => provider.GetRequiredService<AgentService>());
diff --git a/Engine/Source/Programs/Horde/Horde.Build/Ugs/Impl/UgsMetadataCollection.cs b/Engine/Source/Programs/Horde/Horde.Build/Ugs/Impl/UgsMetadataCollection.cs
--- a/Engine/Source/Programs/Horde/Horde.Build/Ugs/Impl/UgsMetadataCollection.cs
+++ b/Engine/Source/Programs/Horde/Horde.Build/Ugs/Impl/UgsMetadataCollection.cs
@@ -95,27 +95,10 @@ namespace Horde.Build.Collections.Impl
 		/// <param name="mongoService">Database service instance</param>
 		public UgsMetadataCollection(MongoService mongoService)
 		{
-			_collection = mongoService.GetCollection<UgsMetadataDocument>("UgsMetadata");
-			
-			if (!mongoService.ReadOnlyMode)
-			{
-				void CreateIndex()
-				{
-					_collection.Indexes.CreateOne(new CreateIndexModel<UgsMetadataDocument>(Builders<UgsMetadataDocument>.IndexKeys.Ascending(x => x.Stream).Descending(x => x.Change).Ascending(x => x.Project), new CreateIndexOptions { Unique = true }));
-					_collection.Indexes.CreateOne(new CreateIndexModel<UgsMetadataDocument>(Builders<UgsMetadataDocument>.IndexKeys.Ascending(x => x.Stream).Descending(x => x.Change).Descending(x => x.UpdateTicks)));
-				}
-				
-				try
-				{
-					CreateIndex();	
-				}
-				catch (MongoCommandException)
-				{
-					// If index creation fails, drop all indexes and retry.
-					_collection.Indexes.DropAll();
-					CreateIndex();
-				}
-			}
+			List<MongoIndex<UgsMetadataDocument>> indexes = new List<MongoIndex<UgsMetadataDocument>>();
+			indexes.Add(keys => keys.Ascending(x => x.Stream).Descending(x => x.Change).Ascending(x => x.Project), unique: true);
+			indexes.Add(keys => keys.Ascending(x => x.Stream).Descending(x => x.Change).Descending(x => x.UpdateTicks));
+			_collection = mongoService.GetCollection<UgsMetadataDocument>("UgsMetadata", indexes);
 		}
 
 		/// <summary>
diff --git a/Engine/Source/Programs/Horde/Horde.Build/Users/Impl/ServiceAccountCollection.cs b/Engine/Source/Programs/Horde/Horde.Build/Users/Impl/ServiceAccountCollection.cs
--- a/Engine/Source/Programs/Horde/Horde.Build/Users/Impl/ServiceAccountCollection.cs
+++ b/Engine/Source/Programs/Horde/Horde.Build/Users/Impl/ServiceAccountCollection.cs
@@ -107,12 +107,7 @@ namespace Horde.Build.Collections.Impl
 		/// <param name="mongoService">The database service</param>
 		public ServiceAccountCollection(MongoService mongoService)
 		{
-			_serviceAccounts = mongoService.GetCollection<ServiceAccountDocument>("ServiceAccounts");
-
-			if (!mongoService.ReadOnlyMode)
-			{
-				_serviceAccounts.Indexes.CreateOne(new CreateIndexModel<ServiceAccountDocument>(Builders<ServiceAccountDocument>.IndexKeys.Ascending(x => x.SecretToken)));
-			}
+			_serviceAccounts = mongoService.GetCollection<ServiceAccountDocument>("ServiceAccounts", keys => keys.Ascending(x => x.SecretToken));
 		}
 
 		/// <inheritdoc/>
diff --git a/Engine/Source/Programs/Horde/Horde.Build/Users/Impl/UserCollection.cs b/Engine/Source/Programs/Horde/Horde.Build/Users/Impl/UserCollection.cs
--- a/Engine/Source/Programs/Horde/Horde.Build/Users/Impl/UserCollection.cs
+++ b/Engine/Source/Programs/Horde/Horde.Build/Users/Impl/UserCollection.cs
@@ -88,12 +88,7 @@ namespace Horde.Build.Collections.Impl
 		/// <param name="mongoService"></param>
 		public UserCollectionV1(MongoService mongoService)
 		{
-			_users = mongoService.GetCollection<UserDocument>("Users");
-
-			if (!mongoService.ReadOnlyMode)
-			{
-				_users.Indexes.CreateOne(new CreateIndexModel<UserDocument>(Builders<UserDocument>.IndexKeys.Ascending(x => x.PrimaryClaim), new CreateIndexOptions { Unique = true }));
-			}
+			_users = mongoService.GetCollection<UserDocument>("Users", keys => keys.Ascending(x => x.PrimaryClaim), unique: true);
 		}
 
 		/// <inheritdoc/>
diff --git a/Engine/Source/Programs/Horde/Horde.Build/Users/Impl/UserCollectionV2.cs b/Engine/Source/Programs/Horde/Horde.Build/Users/Impl/UserCollectionV2.cs
--- a/Engine/Source/Programs/Horde/Horde.Build/Users/Impl/UserCollectionV2.cs
+++ b/Engine/Source/Programs/Horde/Horde.Build/Users/Impl/UserCollectionV2.cs
@@ -169,16 +169,14 @@ namespace Horde.Build.Collections.Impl
 		{
 			_logger = logger;
 
-			_users = mongoService.GetCollection<UserDocument>("UsersV2");
+			List<MongoIndex<UserDocument>> userIndexes = new List<MongoIndex<UserDocument>>();
+			userIndexes.Add(keys => keys.Ascending(x => x.LoginUpper), unique: true);
+			userIndexes.Add(keys => keys.Ascending(x => x.EmailUpper));
+			_users = mongoService.GetCollection<UserDocument>("UsersV2", userIndexes);
+
 			_userClaims = mongoService.GetCollection<UserClaimsDocument>("UserClaimsV2");
 			_userSettings = mongoService.GetCollection<UserSettingsDocument>("UserSettingsV2");
 
-			if (!mongoService.ReadOnlyMode)
-			{
-				_users.Indexes.CreateOne(new CreateIndexModel<UserDocument>(Builders<UserDocument>.IndexKeys.Ascending(x => x.LoginUpper), new CreateIndexOptions { Unique = true }));
-				_users.Indexes.CreateOne(new CreateIndexModel<UserDocument>(Builders<UserDocument>.IndexKeys.Ascending(x => x.EmailUpper)));
-			}
-
 			MemoryCacheOptions options = new MemoryCacheOptions();
 			_userCache = new MemoryCache(options);
 		}
diff --git a/Engine/Source/Programs/Horde/Horde.Build/Utilities/BaseDatabaseIndexes.cs b/Engine/Source/Programs/Horde/Horde.Build/Utilities/BaseDatabaseIndexes.cs
deleted file mode 100644
