Horde: Retry queries that fail due to a hinted index not existing. Now that we create and update indexes asynchronously, we may fail requests while doing so.

diff --git a/Engine/Source/Programs/Horde/Horde.Build/Agents/Impl/LeaseCollection.cs b/Engine/Source/Programs/Horde/Horde.Build/Agents/Impl/LeaseCollection.cs
--- a/Engine/Source/Programs/Horde/Horde.Build/Agents/Impl/LeaseCollection.cs
+++ b/Engine/Source/Programs/Horde/Horde.Build/Agents/Impl/LeaseCollection.cs
@@ -131,13 +131,8 @@ namespace Horde.Build.Collections.Impl
 			{
 				filter &= filterBuilder.Lt(x => x.StartTime, maxTime.Value);
 			}
-			FindOptions? findOptions = null;
-			if (indexHint != null)
-			{
-				findOptions = new FindOptions { Hint = new BsonString(indexHint) };
-			}
 
-			List<LeaseDocument> results = await collection.Find(filter, findOptions).SortByDescending(x => x.StartTime).Range(index, count).ToListAsync();
+			List<LeaseDocument> results = await collection.FindWithHint(filter, indexHint, x => x.SortByDescending(x => x.StartTime).Range(index, count).ToListAsync());
 			return results.ConvertAll<ILease>(x => x);
 		}
 		
@@ -162,9 +157,7 @@ namespace Horde.Build.Collections.Impl
 				filter &= filterBuilder.Lt(x => x.FinishTime, maxFinishTime.Value);
 			}
 
-			FindOptions? findOptions = indexHint == null ? null : new FindOptions { Hint = new BsonString(indexHint) };
-			List<LeaseDocument> results = await collection.Find(filter, findOptions).SortByDescending(x => x.FinishTime).Range(index, count).ToListAsync();
-
+			List<LeaseDocument> results = await collection.FindWithHint(filter, indexHint, x => x.SortByDescending(x => x.FinishTime).Range(index, count).ToListAsync());
 			return results.ConvertAll<ILease>(x => x);
 		}
 
diff --git a/Engine/Source/Programs/Horde/Horde.Build/Jobs/Impl/JobCollection.cs b/Engine/Source/Programs/Horde/Horde.Build/Jobs/Impl/JobCollection.cs
--- a/Engine/Source/Programs/Horde/Horde.Build/Jobs/Impl/JobCollection.cs
+++ b/Engine/Source/Programs/Horde/Horde.Build/Jobs/Impl/JobCollection.cs
@@ -519,24 +519,7 @@ namespace Horde.Build.Collections.Impl
 			using (IScope scope = GlobalTracer.Instance.BuildSpan("Jobs.Find").StartActive())
 			{
 				IMongoCollection<JobDocument> collection = consistentRead ? _jobs : _jobs.WithReadPreference(ReadPreference.SecondaryPreferred);
-
-				FindOptions? findOptions = null;
-				if (indexHint != null)
-				{
-					findOptions = new FindOptions { Hint = new BsonString(indexHint) };
-				}
-				
-				IFindFluent<JobDocument, JobDocument> query = collection.Find<JobDocument>(filter, findOptions).SortByDescending(x => x.CreateTimeUtc!);
-				
-				if (index != null)
-				{
-					query = query.Skip(index.Value);
-				}
-				if (count != null)
-				{
-					query = query.Limit(count.Value);
-				}
-				results = await query.ToListAsync();
+				results = await collection.FindWithHint(filter, indexHint, x => x.SortByDescending(x => x.CreateTimeUtc!).Range(index, count).ToListAsync());
 			}
 			foreach (JobDocument result in results)
 			{
diff --git a/Engine/Source/Programs/Horde/Horde.Build/Utilities/MongoExtensions.cs b/Engine/Source/Programs/Horde/Horde.Build/Utilities/MongoExtensions.cs
--- a/Engine/Source/Programs/Horde/Horde.Build/Utilities/MongoExtensions.cs
+++ b/Engine/Source/Programs/Horde/Horde.Build/Utilities/MongoExtensions.cs
@@ -16,6 +16,33 @@ namespace Horde.Build.Utilities
 	/// </summary>
 	public static class MongoExtensions
 	{
+		/// <summary>
+		/// Executes a query with a particular index hint
+		/// </summary>
+		/// <typeparam name="TDoc"></typeparam>
+		/// <typeparam name="TResult"></typeparam>
+		/// <param name="collection"></param>
+		/// <param name="filter"></param>
+		/// <param name="indexHint"></param>
+		/// <param name="processAsync"></param>
+		/// <returns></returns>
+		public static async Task<TResult> FindWithHint<TDoc, TResult>(this IMongoCollection<TDoc> collection, FilterDefinition<TDoc> filter, string? indexHint, Func<IFindFluent<TDoc, TDoc>, Task<TResult>> processAsync)
+		{
+			FindOptions? findOptions = null;
+			if (indexHint != null)
+			{
+				findOptions = new FindOptions { Hint = new BsonString(indexHint) };
+			}
+			try
+			{
+				return await processAsync(collection.Find(filter, findOptions));
+			}
+			catch (MongoCommandException ex) when (indexHint != null && ex.Code == 2)
+			{
+				return await processAsync(collection.Find(filter));
+			}
+		}
+
 		/// <summary>
 		/// Filters the documents returned from a search
 		/// </summary>
