Horde: Add a basic executor that will extract data from a bundle.

Also add matching grpc wrapper types for namespace, bucket and ref identifiers.

diff --git a/Engine/Source/Programs/Horde/Horde.Agent/AgentSettings.cs b/Engine/Source/Programs/Horde/Horde.Agent/AgentSettings.cs
--- a/Engine/Source/Programs/Horde/Horde.Agent/AgentSettings.cs
+++ b/Engine/Source/Programs/Horde/Horde.Agent/AgentSettings.cs
@@ -29,6 +29,11 @@ namespace Horde.Agent
 		/// Manage workspaces and run steps in a synced branch
 		/// </summary>
 		Perforce,
+
+		/// <summary>
+		/// Pull replicated data from Horde Storage
+		/// </summary>
+		Replicated,
 	}
 
 	/// <summary>
diff --git a/Engine/Source/Programs/Horde/Horde.Agent/Execution/ReplicatedExecutor.cs b/Engine/Source/Programs/Horde/Horde.Agent/Execution/ReplicatedExecutor.cs
new file mode 100644
--- /dev/null
+++ b/Engine/Source/Programs/Horde/Horde.Agent/Execution/ReplicatedExecutor.cs
@@ -0,0 +1,66 @@
+// Copyright Epic Games, Inc. All Rights Reserved.
+
+using System.Collections.Generic;
+using System.Threading;
+using System.Threading.Tasks;
+using EpicGames.Core;
+using EpicGames.Horde.Bundles;
+using EpicGames.Horde.Bundles.Nodes;
+using EpicGames.Horde.Storage;
+using Horde.Agent.Utility;
+using HordeCommon.Rpc;
+using Microsoft.Extensions.Caching.Memory;
+using Microsoft.Extensions.Logging;
+
+namespace Horde.Agent.Execution
+{
+	class ReplicatedExecutor : BuildGraphExecutor
+	{
+		protected readonly DirectoryReference _rootDir;
+		protected readonly IStorageClient _storageClient;
+		protected readonly QualifiedRefId _contentRef;
+		protected DirectoryReference _workspaceDir = null!;
+		protected Dictionary<string, string> _envVars = new Dictionary<string, string>();
+
+		public ReplicatedExecutor(IRpcConnection rpcConnection, string jobId, string batchId, string agentTypeName, IStorageClient storageClient, QualifiedRefId contentRef, DirectoryReference rootDir)
+			: base(rpcConnection, jobId, batchId, agentTypeName)
+		{
+			_rootDir = rootDir;
+			_storageClient = storageClient;
+			_contentRef = contentRef;
+		}
+
+		public override async Task InitializeAsync(ILogger logger, CancellationToken cancellationToken)
+		{
+			await base.InitializeAsync(logger, cancellationToken);
+
+			_workspaceDir = DirectoryReference.Combine(_rootDir, "Replicated");
+			FileUtils.ForceDeleteDirectoryContents(_workspaceDir);
+
+			using (IMemoryCache cache = new MemoryCache(new MemoryCacheOptions { SizeLimit = 1024 * 1024 * 1000 }))
+			{
+				BucketId bucketId = new BucketId(_job.StreamId.ToString());
+				using Bundle<DirectoryNode> bundle = await Bundle.ReadAsync<DirectoryNode>(_storageClient, _contentRef.NamespaceId, _contentRef.BucketId, _contentRef.RefId, new BundleOptions(), cache);
+				await bundle.Root.CopyToDirectoryAsync(bundle, _workspaceDir.ToDirectoryInfo(), logger);
+			}
+
+			// Get all the environment variables for jobs
+			_envVars["IsBuildMachine"] = "1";
+		}
+
+		protected override Task<bool> SetupAsync(BeginStepResponse step, ILogger logger, CancellationToken cancellationToken)
+		{
+			return SetupAsync(step, _workspaceDir, null, _envVars, logger, cancellationToken);
+		}
+
+		protected override Task<bool> ExecuteAsync(BeginStepResponse step, ILogger logger, CancellationToken cancellationToken)
+		{
+			return ExecuteAsync(step, _workspaceDir, null, _envVars, logger, cancellationToken);
+		}
+
+		public override Task FinalizeAsync(ILogger logger, CancellationToken cancellationToken)
+		{
+			return Task.CompletedTask;
+		}
+	}
+}
diff --git a/Engine/Source/Programs/Horde/Horde.Agent/Horde.Agent.csproj b/Engine/Source/Programs/Horde/Horde.Agent/Horde.Agent.csproj
--- a/Engine/Source/Programs/Horde/Horde.Agent/Horde.Agent.csproj
+++ b/Engine/Source/Programs/Horde/Horde.Agent/Horde.Agent.csproj
@@ -73,24 +73,24 @@
     <Content Include="..\Horde.Build\Protos\horde\horde_rpc_messages.proto">
       <Link>Protos\horde_rpc_messages.proto</Link>
     </Content>
-    <Content Include="..\Horde.Build\Protos\horde\common\**">
+    <Content Include="..\Horde.Build\Protos\horde\common\**\*.proto">
       <Link>Protos\horde\common\%(RecursiveDir)%(FileName)%(Extension)</Link>
     </Content>
-    <Content Include="..\Horde.Build\Protos\horde\messages\**">
+    <Content Include="..\Horde.Build\Protos\horde\messages\**\*.proto">
       <Link>Protos\horde\messages\%(RecursiveDir)%(FileName)%(Extension)</Link>
     </Content>
-    <Content Include="..\Horde.Build\Protos\horde\tasks\**">
+    <Content Include="..\Horde.Build\Protos\horde\tasks\**\*.proto">
       <Link>Protos\horde\tasks\%(RecursiveDir)%(FileName)%(Extension)</Link>
     </Content>
+    <Compile Include="..\Horde.Build\Protos\horde\**\*.cs">
+      <Link>Protos\horde\%(RecursiveDir)%(FileName)%(Extension)</Link>
+    </Compile>
     <Compile Include="../../../../Platforms/*/Source/Programs/Horde/Horde.Agent/**/*.cs">
       <Link>$([System.Text.RegularExpressions.Regex]::Replace(%(FullPath), '^.+?[\\/]Horde.Agent[\\/]', ''))</Link>
     </Compile>
     <Compile Include="../../../../Restricted/*/Source/Programs/Horde/Horde.Agent/**/*.cs">
       <Link>$([System.Text.RegularExpressions.Regex]::Replace(%(FullPath), '^.+?[\\/]Horde.Agent[\\/]', ''))</Link>
     </Compile>
-    <Compile Include="..\Horde.Build\Utilities\ProtoExtensions.cs">
-      <Link>Utility\ProtoExtensions.cs</Link>
-    </Compile>
   </ItemGroup>
 
   <ItemGroup>
@@ -104,15 +104,18 @@
     <Protobuf Include="..\Horde.Build\Protos\horde\horde_rpc.proto" ProtoRoot="..\Horde.Build\Protos" GrpcServices="Client" Access="Internal" />
     <Protobuf Include="..\Horde.Build\Protos\horde\horde_rpc_messages.proto" ProtoRoot="..\Horde.Build\Protos" Access="Internal" />
     <Protobuf Include="..\Horde.Build\Protos\horde\common\agent_status.proto" ProtoRoot="..\Horde.Build\Protos" Access="Internal" />
+    <Protobuf Include="..\Horde.Build\Protos\horde\common\bucket_id.proto" ProtoRoot="..\Horde.Build\Protos" Access="Internal" />
     <Protobuf Include="..\Horde.Build\Protos\horde\common\event_severity.proto" ProtoRoot="..\Horde.Build\Protos" Access="Internal" />
-    <Protobuf Include="..\Horde.Build\Protos\horde\common\iohash_wrapper.proto" ProtoRoot="..\Horde.Build\Protos" />
+    <Protobuf Include="..\Horde.Build\Protos\horde\common\io_hash.proto" ProtoRoot="..\Horde.Build\Protos" Access="Internal" />
     <Protobuf Include="..\Horde.Build\Protos\horde\common\jobstep_outcome.proto" ProtoRoot="..\Horde.Build\Protos" Access="Internal" />
     <Protobuf Include="..\Horde.Build\Protos\horde\common\jobstep_state.proto" ProtoRoot="..\Horde.Build\Protos" Access="Internal" />
     <Protobuf Include="..\Horde.Build\Protos\horde\common\label_change.proto" ProtoRoot="..\Horde.Build\Protos" Access="Internal" />
     <Protobuf Include="..\Horde.Build\Protos\horde\common\lease_outcome.proto" ProtoRoot="..\Horde.Build\Protos" Access="Internal" />
     <Protobuf Include="..\Horde.Build\Protos\horde\common\lease_state.proto" ProtoRoot="..\Horde.Build\Protos" Access="Internal" />
+    <Protobuf Include="..\Horde.Build\Protos\horde\common\namespace_id.proto" ProtoRoot="..\Horde.Build\Protos" Access="Internal" />
     <Protobuf Include="..\Horde.Build\Protos\horde\common\priority.proto" ProtoRoot="..\Horde.Build\Protos" Access="Internal" />
-    <Protobuf Include="..\Horde.Build\Protos\horde\common\refid_wrapper.proto" ProtoRoot="..\Horde.Build\Protos" />
+    <Protobuf Include="..\Horde.Build\Protos\horde\common\qualified_ref_id.proto" ProtoRoot="..\Horde.Build\Protos" Access="Internal" />
+    <Protobuf Include="..\Horde.Build\Protos\horde\common\ref_id.proto" ProtoRoot="..\Horde.Build\Protos" Access="Internal" />
     <Protobuf Include="..\Horde.Build\Protos\horde\messages\agent_capabilities.proto" ProtoRoot="..\Horde.Build\Protos" Access="Internal" />
     <Protobuf Include="..\Horde.Build\Protos\horde\messages\agent_requirements.proto" ProtoRoot="..\Horde.Build\Protos" Access="Internal" />
     <Protobuf Include="..\Horde.Build\Protos\horde\messages\agent_workspace.proto" ProtoRoot="..\Horde.Build\Protos" Access="Internal" />
diff --git a/Engine/Source/Programs/Horde/Horde.Agent/Services/WorkerService.cs b/Engine/Source/Programs/Horde/Horde.Agent/Services/WorkerService.cs
--- a/Engine/Source/Programs/Horde/Horde.Agent/Services/WorkerService.cs
+++ b/Engine/Source/Programs/Horde/Horde.Agent/Services/WorkerService.cs
@@ -258,6 +258,9 @@ namespace Horde.Agent.Services
 						new LocalExecutor(rpcClient, executeTask.JobId, executeTask.BatchId, batch.AgentType, _settings.LocalExecutor),
 					ExecutorType.Perforce => (rpcClient, executeTask, batch) =>
 						new PerforceExecutor(rpcClient, executeTask.JobId, executeTask.BatchId, batch.AgentType, executeTask.AutoSdkWorkspace, executeTask.Workspace, _workingDir),
+					ExecutorType.Replicated => (rpcClient, executeTask, batch) =>
+						new ReplicatedExecutor(rpcClient, executeTask.JobId, executeTask.BatchId, batch.AgentType, _storageClient, executeTask.ContentRef, _workingDir),
+
 					_ => throw new InvalidDataException($"Unknown executor type '{_settings.Executor}'")
 				};
 			}
diff --git a/Engine/Source/Programs/Horde/Horde.Agent/Utilities/ProtoMethods.cs b/Engine/Source/Programs/Horde/Horde.Agent/Utilities/ProtoMethods.cs
--- a/Engine/Source/Programs/Horde/Horde.Agent/Utilities/ProtoMethods.cs
+++ b/Engine/Source/Programs/Horde/Horde.Agent/Utilities/ProtoMethods.cs
@@ -1,6 +1,7 @@
 // Copyright Epic Games, Inc. All Rights Reserved.
 
 using EpicGames.Horde.Compute;
+using EpicGames.Horde.Storage;
 
 namespace HordeCommon.Rpc.Tasks
 {
@@ -12,7 +13,7 @@ namespace HordeCommon.Rpc.Tasks
 		/// <summary>
 		/// Constructor
 		/// </summary>
-		public ComputeTaskResultMessage(RefIdWrapper resultRefId)
+		public ComputeTaskResultMessage(RefId resultRefId)
 		{
 			ResultRefId = resultRefId;
 		}
diff --git a/Engine/Source/Programs/Horde/Horde.Build.Tests/TestSetup.cs b/Engine/Source/Programs/Horde/Horde.Build.Tests/TestSetup.cs
--- a/Engine/Source/Programs/Horde/Horde.Build.Tests/TestSetup.cs
+++ b/Engine/Source/Programs/Horde/Horde.Build.Tests/TestSetup.cs
@@ -11,6 +11,8 @@ using EpicGames.Horde.Storage;
 using Horde.Build.Acls;
 using Horde.Build.Collections;
 using Horde.Build.Collections.Impl;
+using Horde.Build.Commits;
+using Horde.Build.Commits.Impl;
 using Horde.Build.Controllers;
 using Horde.Build.Jobs;
 using Horde.Build.Logs;
@@ -32,6 +34,7 @@ using Horde.Build.Utilities;
 using HordeCommon;
 using Microsoft.AspNetCore.Http;
 using Microsoft.AspNetCore.Mvc;
+using Microsoft.Extensions.Caching.Memory;
 using Microsoft.Extensions.Configuration;
 using Microsoft.Extensions.DependencyInjection;
 using Microsoft.Extensions.Hosting;
@@ -128,6 +131,7 @@ namespace Horde.Build.Tests
 			services.AddSingleton<IConfiguration>(config);
 
 			services.AddLogging(builder => builder.AddConsole());
+			services.AddSingleton<IMemoryCache>(sp => new MemoryCache(new MemoryCacheOptions { }));
 
 			services.AddSingleton<RedisService>();
 
@@ -137,6 +141,7 @@ namespace Horde.Build.Tests
 			services.AddSingleton<IAgentCollection, AgentCollection>();
 			services.AddSingleton<IAgentSoftwareCollection, AgentSoftwareCollection>();
 			services.AddSingleton<IArtifactCollection, ArtifactCollection>();
+			services.AddSingleton<ICommitCollection, CommitCollection>();
 			services.AddSingleton<IGraphCollection, GraphCollection>();
 			services.AddSingleton<IIssueCollection, IssueCollection>();
 			services.AddSingleton<IJobCollection, JobCollection>();
@@ -191,6 +196,7 @@ namespace Horde.Build.Tests
 			services.AddSingleton<UpgradeService>();
 
 			services.AddSingleton<ConformTaskSource>();
+			services.AddSingleton<ICommitService, CommitService>();
 
 			services.AddSingleton(new Startup.StorageBackendSettings<PersistentLogStorage> { Type = StorageProviderType.Transient });
 			services.AddSingleton(new Startup.StorageBackendSettings<ArtifactCollection> { Type = StorageProviderType.Transient });
diff --git a/Engine/Source/Programs/Horde/Horde.Build/Horde.Build.csproj b/Engine/Source/Programs/Horde/Horde.Build/Horde.Build.csproj
--- a/Engine/Source/Programs/Horde/Horde.Build/Horde.Build.csproj
+++ b/Engine/Source/Programs/Horde/Horde.Build/Horde.Build.csproj
@@ -1,4 +1,4 @@
-<Project Sdk="Microsoft.NET.Sdk.Web">
+﻿<Project Sdk="Microsoft.NET.Sdk.Web">
 
   <PropertyGroup>
     <TargetFramework>net6.0</TargetFramework>
@@ -72,15 +72,18 @@
     <Protobuf Include="Protos\horde\horde_rpc.proto" ProtoRoot="Protos" GrpcServices="Server" />
     <Protobuf Include="Protos\horde\horde_rpc_messages.proto" ProtoRoot="Protos" />
     <Protobuf Include="Protos\horde\common\agent_status.proto" ProtoRoot="Protos" />
+    <Protobuf Include="Protos\horde\common\bucket_id.proto" ProtoRoot="Protos" />
     <Protobuf Include="Protos\horde\common\event_severity.proto" ProtoRoot="Protos" />
-    <Protobuf Include="Protos\horde\common\iohash_wrapper.proto" ProtoRoot="Protos" />
+    <Protobuf Include="Protos\horde\common\io_hash.proto" ProtoRoot="Protos" />
     <Protobuf Include="Protos\horde\common\jobstep_outcome.proto" ProtoRoot="Protos" />
     <Protobuf Include="Protos\horde\common\jobstep_state.proto" ProtoRoot="Protos" />
     <Protobuf Include="Protos\horde\common\label_change.proto" ProtoRoot="Protos" />
     <Protobuf Include="Protos\horde\common\lease_outcome.proto" ProtoRoot="Protos" />
     <Protobuf Include="Protos\horde\common\lease_state.proto" ProtoRoot="Protos" />
+    <Protobuf Include="Protos\horde\common\namespace_id.proto" ProtoRoot="Protos" />
     <Protobuf Include="Protos\horde\common\priority.proto" ProtoRoot="Protos" />
-    <Protobuf Include="Protos\horde\common\refid_wrapper.proto" ProtoRoot="Protos" />
+    <Protobuf Include="Protos\horde\common\qualified_ref_id.proto" ProtoRoot="Protos" />
+    <Protobuf Include="Protos\horde\common\ref_id.proto" ProtoRoot="Protos" />
     <Protobuf Include="Protos\horde\messages\agent_capabilities.proto" ProtoRoot="Protos" />
     <Protobuf Include="Protos\horde\messages\agent_requirements.proto" ProtoRoot="Protos" />
     <Protobuf Include="Protos\horde\messages\agent_workspace.proto" ProtoRoot="Protos" />
diff --git a/Engine/Source/Programs/Horde/Horde.Build/Jobs/Impl/JobTaskSource.cs b/Engine/Source/Programs/Horde/Horde.Build/Jobs/Impl/JobTaskSource.cs
--- a/Engine/Source/Programs/Horde/Horde.Build/Jobs/Impl/JobTaskSource.cs
+++ b/Engine/Source/Programs/Horde/Horde.Build/Jobs/Impl/JobTaskSource.cs
@@ -14,6 +14,7 @@ using Google.Protobuf;
 using Google.Protobuf.WellKnownTypes;
 using Horde.Build.Api;
 using Horde.Build.Collections;
+using Horde.Build.Commits;
 using Horde.Build.Models;
 using Horde.Build.Server;
 using Horde.Build.Services;
@@ -191,6 +192,7 @@ namespace Horde.Build.Tasks.Impl
 		readonly IUgsMetadataCollection _ugsMetadataCollection;
 		readonly PerforceLoadBalancer _perforceLoadBalancer;
 		readonly IOptionsMonitor<ServerSettings> _settings;
+		readonly ICommitService _commitService;
 		readonly ILogger<JobTaskSource> _logger;
 		readonly ITicker _ticker;
 
@@ -231,7 +233,7 @@ namespace Horde.Build.Tasks.Impl
 		/// <summary>
 		/// Constructor
 		/// </summary>
-		public JobTaskSource(MongoService mongoService, IAgentCollection agents, IJobCollection jobs, IJobStepRefCollection jobStepRefs, IGraphCollection graphs, IPoolCollection pools, IUgsMetadataCollection ugsMetadataCollection, StreamService streamService, ILogFileService logFileService, PerforceLoadBalancer perforceLoadBalancer, IClock clock, IOptionsMonitor<ServerSettings> settings, ILogger<JobTaskSource> logger)
+		public JobTaskSource(MongoService mongoService, IAgentCollection agents, IJobCollection jobs, IJobStepRefCollection jobStepRefs, IGraphCollection graphs, IPoolCollection pools, IUgsMetadataCollection ugsMetadataCollection, StreamService streamService, ILogFileService logFileService, PerforceLoadBalancer perforceLoadBalancer, ICommitService commitService, IClock clock, IOptionsMonitor<ServerSettings> settings, ILogger<JobTaskSource> logger)
 		{
 			_mongoService = mongoService;
 			_agentsCollection = agents;
@@ -245,6 +247,7 @@ namespace Horde.Build.Tasks.Impl
 			_perforceLoadBalancer = perforceLoadBalancer;
 			_ticker = clock.AddTicker<JobTaskSource>(s_refreshInterval, TickAsync, logger);
 			_settings = settings;
+			_commitService = commitService;
 			_logger = logger;
 
 			OnLeaseStartedProperties.Add(nameof(ExecuteJobTask.JobId), x => new JobId(x.JobId)).Add(nameof(ExecuteJobTask.BatchId), x => SubResourceId.Parse(x.BatchId)).Add(nameof(ExecuteJobTask.LogId), x => new LogId(x.LogId));
@@ -844,6 +847,13 @@ namespace Horde.Build.Tasks.Impl
 			}
 
 			task.Workspace = workspaces.Last();
+
+			IStream? Stream;
+			if (_streams.TryGetValue(job.StreamId, out Stream))
+			{
+				task.ContentRef = _commitService.GetReplicatedContentRef(Stream, job.Change);
+			}
+
 			return task;
 		}
 
diff --git a/Engine/Source/Programs/Horde/Horde.Build/Perforce/ICommit.cs b/Engine/Source/Programs/Horde/Horde.Build/Perforce/ICommit.cs
--- a/Engine/Source/Programs/Horde/Horde.Build/Perforce/ICommit.cs
+++ b/Engine/Source/Programs/Horde/Horde.Build/Perforce/ICommit.cs
@@ -64,7 +64,7 @@ namespace Horde.Build.Commits
 	/// <summary>
 	/// New commit object
 	/// </summary>
-	class NewCommit
+	public class NewCommit
 	{
 		/// <inheritdoc cref="ICommit.StreamId"/>
 		public StreamId StreamId { get; set; }
diff --git a/Engine/Source/Programs/Horde/Horde.Build/Perforce/ICommitCollection.cs b/Engine/Source/Programs/Horde/Horde.Build/Perforce/ICommitCollection.cs
--- a/Engine/Source/Programs/Horde/Horde.Build/Perforce/ICommitCollection.cs
+++ b/Engine/Source/Programs/Horde/Horde.Build/Perforce/ICommitCollection.cs
@@ -13,7 +13,7 @@ namespace Horde.Build.Commits
 	/// <summary>
 	/// Stores a collection of commits
 	/// </summary>
-	interface ICommitCollection
+	public interface ICommitCollection
 	{
 		/// <summary>
 		/// Adds or replaces an existing commit
diff --git a/Engine/Source/Programs/Horde/Horde.Build/Perforce/ICommitService.cs b/Engine/Source/Programs/Horde/Horde.Build/Perforce/ICommitService.cs
--- a/Engine/Source/Programs/Horde/Horde.Build/Perforce/ICommitService.cs
+++ b/Engine/Source/Programs/Horde/Horde.Build/Perforce/ICommitService.cs
@@ -1,9 +1,14 @@
 // Copyright Epic Games, Inc. All Rights Reserved.
 
 using System;
+using EpicGames.Horde.Storage;
+using Horde.Build.Models;
+using Horde.Build.Utilities;
 
 namespace Horde.Build.Commits
 {
+	using StreamId = StringId<IStream>;
+
 	/// <summary>
 	/// Provides information about commits
 	/// </summary>
@@ -15,5 +20,13 @@ namespace Horde.Build.Commits
 		/// <param name="onAddCommit">Callback for a new commit being added</param>
 		/// <returns>Disposable handler.</returns>
 		IDisposable AddListener(Action<ICommit> onAddCommit);
+
+		/// <summary>
+		/// Gets reference to replicated content for the given stream and change
+		/// </summary>
+		/// <param name="stream"></param>
+		/// <param name="change"></param>
+		/// <returns></returns>
+		QualifiedRefId? GetReplicatedContentRef(IStream stream, int change);
 	}
 }
diff --git a/Engine/Source/Programs/Horde/Horde.Build/Perforce/Impl/CommitService.cs b/Engine/Source/Programs/Horde/Horde.Build/Perforce/Impl/CommitService.cs
--- a/Engine/Source/Programs/Horde/Horde.Build/Perforce/Impl/CommitService.cs
+++ b/Engine/Source/Programs/Horde/Horde.Build/Perforce/Impl/CommitService.cs
@@ -85,6 +85,7 @@ namespace Horde.Build.Commits.Impl
 		/// <summary>
 		/// Filter for paths in the depot included in this tree
 		/// </summary>
+		[CbField]
 		public string? Filter { get; set; }
 
 		/// <summary>
@@ -120,6 +121,11 @@ namespace Horde.Build.Commits.Impl
 			CbObject obj = CbSerializer.Serialize(this);
 			return new RefId(IoHash.Compute(obj.GetView().Span));
 		}
+
+		/// <summary>
+		/// Gets the ref id for a particular commit
+		/// </summary>
+		public static RefId GetRefId(StreamId streamId, int change, string? filter, bool revisionsOnly) => new CommitKey(streamId, change, filter, revisionsOnly).GetRefId();
 	}
 
 	/// <summary>
@@ -977,6 +983,19 @@ namespace Horde.Build.Commits.Impl
 			_logger.LogInformation("Written ref {RefId} for {StreamId} change {Change}. Download: {NewBytes:n0} bytes ({NewRefCount} refs => {NewRefBytes:n0} bytes, {NewBlobCount} blobs => {NewBlobBytes:n0} bytes, {NewNodeCount:n0} nodes => {NewNodeBytes:n0} bytes).", refId, stream.Id, change, stats.NewRefBytes + stats.NewBlobBytes, stats.NewRefCount, stats.NewRefBytes, stats.NewBlobCount, stats.NewBlobBytes, stats.NewNodeCount, stats.NewNodeBytes);
 		}
 
+		public QualifiedRefId? GetReplicatedContentRef(IStream stream, int change)
+		{
+			if (stream.ReplicationMode == ContentReplicationMode.None)
+			{
+				return null;
+			}
+			else
+			{
+				CommitKey key = new CommitKey(stream.Id, change, stream.ReplicationFilter, stream.ReplicationMode == ContentReplicationMode.RevisionsOnly);
+				return new QualifiedRefId(Options.NamespaceId, GetStreamBucketId(stream.Id), key.GetRefId());
+			}
+		}
+
 		async Task FlushWorkspaceAsync(ReplicationClient clientInfo, IPerforceConnection perforce, int change)
 		{
 			if (clientInfo.Change != change)
diff --git a/Engine/Source/Programs/Horde/Horde.Build/Protos/horde/common/bucket_id.cs b/Engine/Source/Programs/Horde/Horde.Build/Protos/horde/common/bucket_id.cs
new file mode 100644
--- /dev/null
+++ b/Engine/Source/Programs/Horde/Horde.Build/Protos/horde/common/bucket_id.cs
@@ -0,0 +1,23 @@
+// Copyright Epic Games, Inc. All Rights Reserved.
+
+namespace HordeCommon.Wrapper
+{
+	using NativeBucketId = EpicGames.Horde.Storage.BucketId;
+
+#pragma warning disable CS1591 // Missing XML comment for publicly visible type or member
+	partial class BucketId
+	{
+		public static implicit operator NativeBucketId(BucketId wrapper)
+		{
+			return new NativeBucketId(wrapper.Data);
+		}
+
+		public static implicit operator BucketId(NativeBucketId source)
+		{
+			BucketId result = new BucketId();
+			result.Data = source.ToString();
+			return result;
+		}
+	}
+#pragma warning restore CS1591 // Missing XML comment for publicly visible type or member
+}
diff --git a/Engine/Source/Programs/Horde/Horde.Build/Protos/horde/common/io_hash.cs b/Engine/Source/Programs/Horde/Horde.Build/Protos/horde/common/io_hash.cs
new file mode 100644
--- /dev/null
+++ b/Engine/Source/Programs/Horde/Horde.Build/Protos/horde/common/io_hash.cs
@@ -0,0 +1,23 @@
+// Copyright Epic Games, Inc. All Rights Reserved.
+
+namespace HordeCommon.Wrapper
+{
+	using NativeIoHash = EpicGames.Core.IoHash;
+
+#pragma warning disable CS1591 // Missing XML comment for publicly visible type or member
+	partial class IoHash
+	{
+		public static implicit operator NativeIoHash(IoHash hash)
+		{
+			return new NativeIoHash(hash.Data.ToByteArray());
+		}
+
+		public static implicit operator IoHash(NativeIoHash hash)
+		{
+			IoHash result = new IoHash();
+			result.Data = Google.Protobuf.ByteString.CopyFrom(hash.ToByteArray());
+			return result;
+		}
+	}
+#pragma warning restore CS1591 // Missing XML comment for publicly visible type or member
+}
diff --git a/Engine/Source/Programs/Horde/Horde.Build/Protos/horde/common/namespace_id.cs b/Engine/Source/Programs/Horde/Horde.Build/Protos/horde/common/namespace_id.cs
new file mode 100644
--- /dev/null
+++ b/Engine/Source/Programs/Horde/Horde.Build/Protos/horde/common/namespace_id.cs
@@ -0,0 +1,23 @@
+// Copyright Epic Games, Inc. All Rights Reserved.
+
+namespace HordeCommon.Wrapper
+{
+	using NativeNamespaceId = EpicGames.Horde.Storage.NamespaceId;
+
+#pragma warning disable CS1591 // Missing XML comment for publicly visible type or member
+	partial class NamespaceId
+	{
+		public static implicit operator NativeNamespaceId(NamespaceId wrapper)
+		{
+			return new NativeNamespaceId(wrapper.Data);
+		}
+
+		public static implicit operator NamespaceId(NativeNamespaceId source)
+		{
+			NamespaceId result = new NamespaceId();
+			result.Data = source.ToString();
+			return result;
+		}
+	}
+#pragma warning restore CS1591 // Missing XML comment for publicly visible type or member
+}
diff --git a/Engine/Source/Programs/Horde/Horde.Build/Protos/horde/common/qualified_ref_id.cs b/Engine/Source/Programs/Horde/Horde.Build/Protos/horde/common/qualified_ref_id.cs
new file mode 100644
--- /dev/null
+++ b/Engine/Source/Programs/Horde/Horde.Build/Protos/horde/common/qualified_ref_id.cs
@@ -0,0 +1,42 @@
+// Copyright Epic Games, Inc. All Rights Reserved.
+
+namespace HordeCommon.Wrapper
+{
+	using System.Diagnostics.CodeAnalysis;
+	using NativeQualifiedRefId = EpicGames.Horde.Storage.QualifiedRefId;
+
+#pragma warning disable CS1591 // Missing XML comment for publicly visible type or member
+	partial class QualifiedRefId
+	{
+		[return: NotNullIfNotNull("wrapper")]
+		public static implicit operator NativeQualifiedRefId?(QualifiedRefId? wrapper)
+		{
+			if (wrapper == null)
+			{
+				return null;
+			}
+			else
+			{
+				return new NativeQualifiedRefId(wrapper.NamespaceId, wrapper.BucketId, wrapper.RefId);
+			}
+		}
+
+		[return: NotNullIfNotNull("source")]
+		public static implicit operator QualifiedRefId?(NativeQualifiedRefId? source)
+		{
+			if (source == null)
+			{
+				return null;
+			}
+			else
+			{
+				QualifiedRefId result = new QualifiedRefId();
+				result.NamespaceId = source.NamespaceId;
+				result.BucketId = source.BucketId;
+				result.RefId = source.RefId;
+				return result;
+			}
+		}
+	}
+#pragma warning restore CS1591 // Missing XML comment for publicly visible type or member
+}
diff --git a/Engine/Source/Programs/Horde/Horde.Build/Protos/horde/common/ref_id.cs b/Engine/Source/Programs/Horde/Horde.Build/Protos/horde/common/ref_id.cs
new file mode 100644
--- /dev/null
+++ b/Engine/Source/Programs/Horde/Horde.Build/Protos/horde/common/ref_id.cs
@@ -0,0 +1,28 @@
+// Copyright Epic Games, Inc. All Rights Reserved.
+
+namespace HordeCommon.Wrapper
+{
+	using NativeRefId = EpicGames.Horde.Storage.RefId;
+
+#pragma warning disable CS1591 // Missing XML comment for publicly visible type or member
+	partial class RefId
+	{
+		public NativeRefId AsRefId()
+		{
+			return new NativeRefId(Data);
+		}
+
+		public static implicit operator NativeRefId(RefId wrapper)
+		{
+			return new NativeRefId(wrapper.Data);
+		}
+
+		public static implicit operator RefId(NativeRefId source)
+		{
+			RefId refId = new RefId();
+			refId.Data = source.Hash;
+			return refId;
+		}
+	}
+#pragma warning restore CS1591 // Missing XML comment for publicly visible type or member
+}
diff --git a/Engine/Source/Programs/Horde/Horde.Build/Utilities/ProtoExtensions.cs b/Engine/Source/Programs/Horde/Horde.Build/Utilities/ProtoExtensions.cs
deleted file mode 100644
diff --git a/Engine/Source/Programs/Shared/EpicGames.Horde/Storage/QualifiedRefId.cs b/Engine/Source/Programs/Shared/EpicGames.Horde/Storage/QualifiedRefId.cs
new file mode 100644
--- /dev/null
+++ b/Engine/Source/Programs/Shared/EpicGames.Horde/Storage/QualifiedRefId.cs
@@ -0,0 +1,38 @@
+// Copyright Epic Games, Inc. All Rights Reserved.
+
+namespace EpicGames.Horde.Storage
+{
+	/// <summary>
+	/// Qualified locator for a ref
+	/// </summary>
+	public class QualifiedRefId
+	{
+		/// <summary>
+		/// The namespace containing the ref
+		/// </summary>
+		public NamespaceId NamespaceId { get; }
+
+		/// <summary>
+		/// Bucket containing the ref
+		/// </summary>
+		public BucketId BucketId { get; }
+
+		/// <summary>
+		/// The ref itself
+		/// </summary>
+		public RefId RefId { get; }
+
+		/// <summary>
+		/// Constructor
+		/// </summary>
+		/// <param name="namespaceId"></param>
+		/// <param name="bucketId"></param>
+		/// <param name="refId"></param>
+		public QualifiedRefId(NamespaceId namespaceId, BucketId bucketId, RefId refId)
+		{
+			NamespaceId = namespaceId;
+			BucketId = bucketId;
+			RefId = refId;
+		}
+	}
+}
