Horde: Add a subsystem for hosting tools on Horde.

- Each tool can have a number of deployments which are rolled out over time. Clients can specify a "phase" value indicating where in the queue they want to take updates (ie. whether they want to be early- or late-adopters), and deployments can be paused and cancelled.
- Tool channels are configured through the global configuration file. Tools can be marked "public", meaning that they do not require authentication against Horde for access. This can be useful for auto-update scenarios.
- The most recent 5 deployments are kept in the tool at any time.

The collection and controller implementation here is an experiment in reducing the amount of boilerplate currently adopted as a pattern in Horde. Notably:

- Model and collection classes are concrete rather than interfaces (since we don't generally mock them anyway).
- The same model class with annotations is used for response objects (paving the way to supporting patching and filtering using the same public model as the internal model).
- The VersionedCollection class is used to handle automatic migration of documents across schema versions, and cache document values with Redis.
- MongoDB property names are explicit and short, rather than just taking variable names.

diff --git a/Engine/Source/Programs/Horde/Horde.Build.Tests/DatabaseIntegrationTest.cs b/Engine/Source/Programs/Horde/Horde.Build.Tests/DatabaseIntegrationTest.cs
--- a/Engine/Source/Programs/Horde/Horde.Build.Tests/DatabaseIntegrationTest.cs
+++ b/Engine/Source/Programs/Horde/Horde.Build.Tests/DatabaseIntegrationTest.cs
@@ -139,9 +139,7 @@ namespace Horde.Build.Tests
 		private readonly LoggerFactory _loggerFactory = new LoggerFactory();
 
 		private static RedisRunner? s_redisRunner;
-        private IDatabase? _redisDb;
-        private RedisConnectionPool? _redisConnectionPool;
-        private ConnectionMultiplexer? _conMux;
+		private RedisService? _redisService;
         public const int RedisDbNum = 15;
 
 		public DatabaseIntegrationTest()
@@ -151,7 +149,7 @@ namespace Horde.Build.Tests
 		protected override void ConfigureServices(IServiceCollection services)
 		{
 			services.AddSingleton(GetMongoServiceSingleton());
-			services.AddSingleton(GetRedisDatabase());
+			services.AddSingleton(GetRedisServiceSingleton());
 		}
 
 		protected override void Dispose(bool disposing)
@@ -160,9 +158,8 @@ namespace Horde.Build.Tests
 
 			_mongoDbInstance?.Dispose();
 			_mongoService?.Dispose();
+			_redisService?.Dispose();
 			_loggerFactory.Dispose();
-
-			_conMux?.Dispose();
 		}
 
 		public MongoService GetMongoServiceSingleton()
@@ -197,30 +194,23 @@ namespace Horde.Build.Tests
 			return s_redisRunner;
 		}
 
-		public IDatabase GetRedisDatabase()
+		public RedisService GetRedisServiceSingleton()
         {
-			RedisRunner redisRunner = GetRedisRunner();
-			(string host, int port) = redisRunner.GetListenAddress();
-
-			_conMux = ConnectionMultiplexer.Connect($"{host}:{port},allowAdmin=true");
-			foreach (EndPoint endpoint in _conMux.GetEndPoints())
+			if (_redisService == null)
 			{
-				_conMux.GetServer(endpoint).FlushDatabase(RedisDbNum);
-			}
+				RedisRunner redisRunner = GetRedisRunner();
 
-			_redisDb = _conMux.GetDatabase(RedisDbNum);
-	        return _redisDb;
-        }
-        
-        public RedisConnectionPool GetRedisConnectionPool()
-        {
-			RedisRunner redisRunner = GetRedisRunner();
-			(string host, int port) = redisRunner.GetListenAddress();
+				(string host, int port) = redisRunner.GetListenAddress();
+				_redisService = new RedisService($"{host}:{port},allowAdmin=true", RedisDbNum, _loggerFactory);
 
-	        _redisConnectionPool = new RedisConnectionPool(3, $"{host}:{port},allowAdmin=true", RedisDbNum);
-	        return _redisConnectionPool;
+				foreach (EndPoint endpoint in _redisService.Multiplexer.GetEndPoints())
+				{
+					_redisService.Multiplexer.GetServer(endpoint).FlushDatabase(RedisDbNum);
+				}
+			}
+			return _redisService;
         }
-
+        
 		public static T Deref<T>(T? item)
 		{
 			Assert.IsNotNull(item);
diff --git a/Engine/Source/Programs/Horde/Horde.Build.Tests/LogBuilderTests.cs b/Engine/Source/Programs/Horde/Horde.Build.Tests/LogBuilderTests.cs
--- a/Engine/Source/Programs/Horde/Horde.Build.Tests/LogBuilderTests.cs
+++ b/Engine/Source/Programs/Horde/Horde.Build.Tests/LogBuilderTests.cs
@@ -8,6 +8,7 @@ using Horde.Build.Api;
 using Horde.Build.Logs;
 using Horde.Build.Logs.Builder;
 using Horde.Build.Models;
+using Horde.Build.Server;
 using Horde.Build.Utilities;
 using Microsoft.Extensions.Logging.Abstractions;
 using Microsoft.VisualStudio.TestTools.UnitTesting;
@@ -29,7 +30,8 @@ namespace Horde.Build.Tests
 		[TestMethod]
 		public async Task TestRedisLogBuilder()
 		{
-			ILogBuilder builder = new RedisLogBuilder(GetRedisConnectionPool(), NullLogger.Instance);
+			RedisService redisService = GetRedisServiceSingleton();
+			ILogBuilder builder = new RedisLogBuilder(redisService.ConnectionPool, NullLogger.Instance);
 			await TestBuilder(builder);
 		}
 
diff --git a/Engine/Source/Programs/Horde/Horde.Build.Tests/LogFileServiceTests.cs b/Engine/Source/Programs/Horde/Horde.Build.Tests/LogFileServiceTests.cs
--- a/Engine/Source/Programs/Horde/Horde.Build.Tests/LogFileServiceTests.cs
+++ b/Engine/Source/Programs/Horde/Horde.Build.Tests/LogFileServiceTests.cs
@@ -43,7 +43,7 @@ namespace Horde.Build.Tests
 			_loggerFactory = new LoggerFactory();
             ILogger<LogFileService> logger = _loggerFactory.CreateLogger<LogFileService>();
 
-			ILogBuilder logBuilder = new RedisLogBuilder(GetRedisConnectionPool(), NullLogger.Instance);
+			ILogBuilder logBuilder = new RedisLogBuilder(GetRedisServiceSingleton().ConnectionPool, NullLogger.Instance);
 			_logStorage = new PersistentLogStorage(new TransientStorageBackend().ForType<PersistentLogStorage>(), NullLogger<PersistentLogStorage>.Instance);
 			_clock = new FakeClock();
 			_logFileService = new LogFileService(logFileCollection, null!, logBuilder, _logStorage, _clock, logger);
diff --git a/Engine/Source/Programs/Horde/Horde.Build.Tests/TestSetup.cs b/Engine/Source/Programs/Horde/Horde.Build.Tests/TestSetup.cs
--- a/Engine/Source/Programs/Horde/Horde.Build.Tests/TestSetup.cs
+++ b/Engine/Source/Programs/Horde/Horde.Build.Tests/TestSetup.cs
@@ -30,6 +30,7 @@ using Horde.Build.Storage.Services;
 using Horde.Build.Tasks.Impl;
 using Horde.Build.Tests.Stubs.Collections;
 using Horde.Build.Tests.Stubs.Services;
+using Horde.Build.Tools;
 using Horde.Build.Utilities;
 using HordeCommon;
 using Microsoft.AspNetCore.Http;
@@ -133,8 +134,6 @@ namespace Horde.Build.Tests
 			services.AddLogging(builder => builder.AddConsole());
 			services.AddSingleton<IMemoryCache>(sp => new MemoryCache(new MemoryCacheOptions { }));
 
-			services.AddSingleton<RedisService>();
-
 			services.AddSingleton(typeof(IAuditLogFactory<>), typeof(AuditLogFactory<>));
 			services.AddSingleton<IAuditLog<AgentId>>(sp => sp.GetRequiredService<IAuditLogFactory<AgentId>>().Create("Agents.Log", "AgentId"));
 
@@ -163,6 +162,8 @@ namespace Horde.Build.Tests
 			services.AddSingleton<IUgsMetadataCollection, UgsMetadataCollection>();
 			services.AddSingleton<IUserCollection, UserCollectionV1>();
 
+			services.AddSingleton<ToolCollection>();
+
 			services.AddSingleton<FakeClock>();
 			services.AddSingleton<IClock>(sp => sp.GetRequiredService<FakeClock>());
 			services.AddSingleton<IHostApplicationLifetime, AppLifetimeStub>();
diff --git a/Engine/Source/Programs/Horde/Horde.Build.Tests/ToolTests.cs b/Engine/Source/Programs/Horde/Horde.Build.Tests/ToolTests.cs
new file mode 100644
--- /dev/null
+++ b/Engine/Source/Programs/Horde/Horde.Build.Tests/ToolTests.cs
@@ -0,0 +1,98 @@
+// Copyright Epic Games, Inc. All Rights Reserved.
+
+using System;
+using System.IO;
+using System.Text;
+using System.Threading.Tasks;
+using Horde.Build.Tools;
+using Horde.Build.Utilities;
+using HordeCommon;
+using Microsoft.Extensions.DependencyInjection;
+using Microsoft.VisualStudio.TestTools.UnitTesting;
+
+namespace Horde.Build.Tests
+{
+	using ToolId = StringId<Tool>;
+	using ToolDeploymentId = ObjectId<ToolDeployment>;
+
+	[TestClass]
+    public class ToolTests : TestSetup
+    {
+        [TestMethod]
+        public async Task AddTool()
+        {
+			ToolCollection collection = ServiceProvider.GetRequiredService<ToolCollection>();
+			Tool tool = await collection.AddOrUpdateAsync(new ToolOptions(new ToolId("ugs")) { Name = "UnrealGameSync", Description = "Tool for syncing content from source control" });
+			Assert.AreEqual(tool.Id, new ToolId("ugs"));
+			Assert.AreEqual(tool.Name, "UnrealGameSync");
+			Assert.AreEqual(tool.Description, "Tool for syncing content from source control");
+
+			Tool tool2 = Deref(await collection.GetAsync(tool.Id));
+			Assert.AreEqual(tool.Id, tool2.Id);
+			Assert.AreEqual(tool.LastUpdateTime, tool2.LastUpdateTime);
+			Assert.AreEqual(tool.Name, tool2.Name);
+			Assert.AreEqual(tool.Description, tool2.Description);
+        }
+
+		[TestMethod]
+		public async Task AddDeployment()
+		{
+			ToolCollection collection = ServiceProvider.GetRequiredService<ToolCollection>();
+
+			Tool tool = await collection.AddOrUpdateAsync(new ToolOptions(new ToolId("ugs")) { Name = "UnrealGameSync", Description = "Tool for syncing content from source control" });
+			Assert.AreEqual(new ToolId("ugs"), tool.Id);
+			Assert.AreEqual("UnrealGameSync", tool.Name);
+			Assert.AreEqual("Tool for syncing content from source control", tool.Description);
+			Assert.AreEqual(0, tool.Deployments.Count);
+
+			byte[] deploymentData = Encoding.UTF8.GetBytes("hello world");
+
+			ToolDeploymentId deploymentId;
+			using (MemoryStream stream = new MemoryStream(deploymentData))
+			{
+				tool = Deref(await collection.CreateDeploymentAsync(tool, new ToolDeploymentOptions { Version = "1.0", Duration = TimeSpan.FromMinutes(5.0), CreatePaused = true }, stream));
+				Assert.AreEqual(1, tool.Deployments.Count);
+				Assert.IsNull(tool.Deployments[0].StartedAt);
+				deploymentId = tool.Deployments[^1].Id;
+			}
+
+			// Check that the deployment doesn't do anything until started
+			FakeClock clock = ServiceProvider.GetRequiredService<FakeClock>();
+			await clock.AdvanceAsync(TimeSpan.FromHours(1.0));
+
+			tool = Deref(await collection.GetAsync(tool.Id));
+			Assert.AreEqual(1, tool.Deployments.Count);
+			Assert.IsNull(tool.Deployments[0].StartedAt);
+
+			// Start the deployment
+			tool = Deref(await collection.UpdateDeploymentAsync(tool, deploymentId, ToolDeploymentState.Resume));
+			Assert.AreEqual(1, tool.Deployments.Count);
+			Assert.IsNotNull(tool.Deployments[0].StartedAt);
+			Assert.IsTrue(Math.Abs((tool.Deployments[0].StartedAt!.Value - clock.UtcNow).TotalSeconds) < 1.0);
+
+			// Check it updates
+			await clock.AdvanceAsync(TimeSpan.FromMinutes(2.5));
+			tool = Deref(await collection.UpdateDeploymentAsync(tool, deploymentId, ToolDeploymentState.Pause));
+			Assert.AreEqual(1, tool.Deployments.Count);
+			Assert.IsNull(tool.Deployments[0].StartedAt);
+			Assert.IsTrue((tool.Deployments[0].Progress - 0.5) < 0.1);
+
+			// Check it stays paused
+			await clock.AdvanceAsync(TimeSpan.FromHours(1.0));
+			tool = Deref(await collection.GetAsync(tool.Id));
+			Assert.AreEqual(1, tool.Deployments.Count);
+			Assert.IsNull(tool.Deployments[0].StartedAt);
+			Assert.IsTrue((tool.Deployments[0].Progress - 0.5) < 0.1);
+
+			// Get the deployment data
+			byte[] outputDeploymentData;
+			using (MemoryStream stream = new MemoryStream())
+			{
+				Stream dataStream = await collection.GetDeploymentPayloadAsync(tool, tool.Deployments[0]);
+				await dataStream.CopyToAsync(stream);
+				outputDeploymentData = stream.ToArray();
+			}
+			Assert.IsTrue(deploymentData.AsSpan().SequenceEqual(outputDeploymentData));
+		}
+	}
+}
\ No newline at end of file
diff --git a/Engine/Source/Programs/Horde/Horde.Build.Tests/VersionedCollectionTests.cs b/Engine/Source/Programs/Horde/Horde.Build.Tests/VersionedCollectionTests.cs
--- a/Engine/Source/Programs/Horde/Horde.Build.Tests/VersionedCollectionTests.cs
+++ b/Engine/Source/Programs/Horde/Horde.Build.Tests/VersionedCollectionTests.cs
@@ -1,8 +1,10 @@
 // Copyright Epic Games, Inc. All Rights Reserved.
 
+using System;
+using System.Collections.Generic;
 using System.Threading.Tasks;
 using Horde.Build.Server;
-using Horde.Build.Tools.Impl;
+using Horde.Build.Utilities;
 using Microsoft.VisualStudio.TestTools.UnitTesting;
 using MongoDB.Driver;
 using StackExchange.Redis;
@@ -12,15 +14,7 @@ namespace Horde.Build.Tests
 	[TestClass]
 	public class VersionedCollectionTests : DatabaseIntegrationTest
 	{
-		abstract class DocumentBase : VersionedDocument<string, DocumentV2>
-		{
-			protected DocumentBase(string id)
-				: base(id)
-			{
-			}
-		}
-
-		class DocumentV1 : DocumentBase
+		class DocumentV1 : VersionedDocument<string, DocumentV2>
 		{
 			public int ValueV1 { get; set; }
 
@@ -33,7 +27,7 @@ namespace Horde.Build.Tests
 			public override DocumentV2 UpgradeToLatest() => new DocumentV2(Id, ValueV1.ToString());
 		}
 
-		class DocumentV2 : DocumentBase
+		class DocumentV2 : VersionedDocument<string, DocumentV2>
 		{
 			public string ValueV2 { get; set; }
 
@@ -46,8 +40,14 @@ namespace Horde.Build.Tests
 			public override DocumentV2 UpgradeToLatest() => this;
 		}
 
+		static readonly IReadOnlyDictionary<int, Type> s_documentTypes = new Dictionary<int, Type>
+		{
+			[1] = typeof(DocumentV1),
+			[2] = typeof(DocumentV2),
+		};
+
 		readonly MongoService _mongoService;
-		readonly IDatabase _redis;
+		readonly RedisService _redisService;
 		readonly RedisKey _baseKey = new RedisKey("versioned/");
 		readonly IMongoCollection<VersionedDocument<string, DocumentV2>> _baseCollection;
 		readonly VersionedCollection<string, DocumentV2> _collection;
@@ -55,9 +55,9 @@ namespace Horde.Build.Tests
 		public VersionedCollectionTests()
 		{
 			_mongoService = GetMongoServiceSingleton();
-			_redis = GetRedisDatabase();
-			_baseCollection = _mongoService.GetCollection<VersionedDocument<string, DocumentV2>>("versioned");
-			_collection = new VersionedCollection<string, DocumentV2>(_baseCollection, _redis, _baseKey);
+			_redisService = GetRedisServiceSingleton();
+			_collection = new VersionedCollection<string, DocumentV2>(_mongoService, "versioned", _redisService, _baseKey, s_documentTypes);
+			_baseCollection = _collection.BaseCollection;
 		}
 
 		[TestMethod]
diff --git a/Engine/Source/Programs/Horde/Horde.Build/Acls/Acl.cs b/Engine/Source/Programs/Horde/Horde.Build/Acls/Acl.cs
--- a/Engine/Source/Programs/Horde/Horde.Build/Acls/Acl.cs
+++ b/Engine/Source/Programs/Horde/Horde.Build/Acls/Acl.cs
@@ -1,4 +1,4 @@
-﻿// Copyright Epic Games, Inc. All Rights Reserved.
+// Copyright Epic Games, Inc. All Rights Reserved.
 
 using System;
 using System.Collections.Generic;
@@ -361,16 +361,33 @@ namespace Horde.Build.Acls
 		/// </summary>
 		DeleteSoftware,
 
-        #endregion
+		#endregion
 
-        //// ADMIN ////
+		//// TOOLS ////
 
-        #region Admin
+		#region Tools
 
-        /// <summary>
-        /// Ability to read any data from the server. Always inherited.
-        /// </summary>
-        AdminRead,
+		/// <summary>
+		/// Ability to download a tool
+		/// </summary>
+		DownloadTool,
+
+		/// <summary>
+		/// Ability to upload new tool versions
+		/// </summary>
+		UploadTool,
+
+		#endregion
+
+
+		//// ADMIN ////
+
+		#region Admin
+
+		/// <summary>
+		/// Ability to read any data from the server. Always inherited.
+		/// </summary>
+		AdminRead,
 
 		/// <summary>
 		/// Ability to write any data to the server.
@@ -485,6 +502,8 @@ namespace Horde.Build.Acls
         #endregion
     }
 
+    #region V1
+
     /// <summary>
     /// Stores information about a claim
     /// </summary>
@@ -752,10 +771,270 @@ namespace Horde.Build.Acls
 		}
 	}
 
-	/// <summary>
-	/// Serializer for JobStepRefId objects
-	/// </summary>
-	public sealed class AclActionSetSerializer : IBsonSerializer<HashSet<AclAction>>
+    #endregion
+
+    #region V2
+
+    /// <summary>
+    /// Describes an entry in the ACL for a particular claim
+    /// </summary>
+    public class AclEntryV2 : IEquatable<AclEntryV2>
+    {
+        /// <summary>
+        /// The claim type, typically a URI
+        /// </summary>
+        [BsonElement("t")]
+        public string Type { get; set; }
+
+        /// <summary>
+        /// The claim value
+        /// </summary>
+        [BsonElement("v")]
+        public string Value { get; set; }
+
+        /// <summary>
+        /// List of allowed operations
+        /// </summary>
+        [BsonElement("a")]
+        [BsonSerializer(typeof(AclActionListSerializer))]
+        public List<AclAction> Actions { get; set; }
+
+        /// <summary>
+        /// Private constructor for serialization
+        /// </summary>
+        [BsonConstructor]
+        private AclEntryV2()
+        {
+            Type = String.Empty;
+            Value = String.Empty;
+            Actions = new List<AclAction>();
+        }
+
+        /// <summary>
+        /// Constructor
+        /// </summary>
+        /// <param name="type">The claim type</param>
+        /// <param name="value">The claim value</param>
+        /// <param name="actions">List of allowed operations</param>
+        public AclEntryV2(string type, string value, IEnumerable<AclAction> actions)
+        {
+            Type = type;
+            Value = value;
+            Actions = new List<AclAction>(actions);
+        }
+
+        /// <inheritdoc/>
+        public override bool Equals(object? obj) => Equals(obj as AclEntryV2);
+
+        /// <inheritdoc/>
+        public override int GetHashCode() => throw new NotSupportedException();
+
+        /// <inheritdoc/>
+        public bool Equals(AclEntryV2? other) => other != null && String.Equals(Type, other.Type, StringComparison.Ordinal) && String.Equals(Value, other.Value, StringComparison.Ordinal) && Actions.SequenceEqual(other.Actions);
+    }
+
+    /// <summary>
+    /// Represents an access control list for an object in the database
+    /// </summary>
+    public class AclV2 : IEquatable<AclV2>
+    {
+        /// <summary>
+        /// List of entries for this ACL
+        /// </summary>
+        [BsonElement("al")]
+        public List<AclEntryV2> Allow { get; set; } = new List<AclEntryV2>();
+
+        /// <summary>
+        /// Whether to inherit permissions from the parent ACL by default
+        /// </summary>
+        [BsonElement("in")]
+        public bool Inherit { get; set; } = true;
+
+        /// <summary>
+        /// Specifies a list of exceptions to the inheritance setting
+        /// </summary>
+        [BsonElement("ex")]
+        [BsonSerializer(typeof(AclActionListSerializer))]
+        public List<AclAction> Exceptions { get; set; } = new List<AclAction>();
+
+        /// <summary>
+        /// Default constructor
+        /// </summary>
+        public AclV2()
+        {
+            Allow = new List<AclEntryV2>();
+            Inherit = true;
+        }
+
+        /// <summary>
+        /// Constructor
+        /// </summary>
+        /// <param name="allow">List of entries for this ACL</param>
+        /// <param name="inherit">Whether to inherit permissions from the parent object by default</param>
+        public AclV2(IEnumerable<AclEntryV2> allow, bool inherit)
+        {
+            Allow.AddRange(allow);
+            Inherit = inherit;
+        }
+
+        /// <summary>
+        /// Constructor
+        /// </summary>
+        /// <param name="allow">List of entries for this ACL</param>
+        public AclV2(params AclEntryV2[] allow)
+            : this(allow.ToList(), true)
+        {
+        }
+
+        /// <summary>
+        /// Checks whether the ACL is set to its default value
+        /// </summary>
+        /// <returns>True if the ACL is its default value</returns>
+        public bool IsDefault() => Allow.Count == 0 && Inherit && Exceptions.Count == 0;
+
+		/// <summary>
+		/// Checks whether an ACL is null or set to its default value
+		/// </summary>
+		/// <returns>True if the ACL is its default value</returns>
+		public static bool IsNullOrDefault(AclV2? acl) => acl is null || acl.IsDefault();
+
+		/// <summary>
+		/// Tests whether a user is authorized to perform the given actions
+		/// </summary>
+		/// <param name="action">Action that is being performed. This should be a single flag.</param>
+		/// <param name="user">The principal to authorize</param>
+		/// <returns>True/false if the action is allowed or denied, null if there is no specific setting for this user</returns>
+		public bool? Authorize(AclAction action, ClaimsPrincipal user)
+        {
+            // Check if there's a specific entry for this action
+            foreach (AclEntryV2 entry in Allow)
+            {
+                if (entry.Actions.Contains(action) && user.HasClaim(entry.Type, entry.Value))
+                {
+                    return true;
+                }
+            }
+
+            // Otherwise check if we're prevented from inheriting permissions
+            if (Inherit)
+            {
+                if (Exceptions != null && Exceptions.Contains(action))
+                {
+                    return false;
+                }
+            }
+            else
+            {
+                if (Exceptions == null || !Exceptions.Contains(action))
+                {
+                    return false;
+                }
+            }
+
+            // Otherwise allow to propagate up the hierarchy
+            return null;
+        }
+
+        /// <inheritdoc/>
+        public override bool Equals(object? obj) => Equals(obj as AclV2);
+
+        /// <inheritdoc/>
+        public bool Equals(AclV2? other) => other != null && Allow.SequenceEqual(other.Allow) && Inherit == other.Inherit && Exceptions.SequenceEqual(other.Exceptions);
+
+		/// <summary>
+		/// Checks whether two acls are equal
+		/// </summary>
+		/// <param name="lhs"></param>
+		/// <param name="rhs"></param>
+		/// <returns></returns>
+		public static bool Equals(AclV2? lhs, AclV2? rhs)
+		{
+			if (lhs is null)
+			{
+				return rhs is null || rhs.IsDefault();
+			}
+			else
+			{
+				return lhs.Equals(rhs);
+			}
+		}
+
+        /// <inheritdoc/>
+        public override int GetHashCode() => throw new NotSupportedException();
+
+        /// <summary>
+        /// Register this type's class map
+        /// </summary>
+        internal static void ConfigureClassMap(BsonClassMap<AclV2> cm)
+        {
+            cm.AutoMap();
+            cm.MapMember(x => x.Allow).SetIgnoreIfDefault(true).SetDefaultValue(new List<AclEntryV2>());
+            cm.MapMember(x => x.Exceptions).SetIgnoreIfDefault(true).SetDefaultValue(new List<AclAction>());
+        }
+    }
+
+    #endregion
+
+    /// <summary>
+    /// Serializer for a list of unique values
+    /// </summary>
+    public sealed class AclActionListSerializer : IBsonSerializer<List<AclAction>>
+    {
+        /// <inheritdoc/>
+        public Type ValueType => typeof(List<AclAction>);
+
+        /// <inheritdoc/>
+        void IBsonSerializer.Serialize(BsonSerializationContext context, BsonSerializationArgs args, object value) => Serialize(context, args, (List<AclAction>)value);
+
+        /// <inheritdoc/>
+        object IBsonSerializer.Deserialize(BsonDeserializationContext context, BsonDeserializationArgs args) => ((IBsonSerializer<List<AclAction>>)this).Deserialize(context, args);
+
+        /// <inheritdoc/>
+        public List<AclAction> Deserialize(BsonDeserializationContext context, BsonDeserializationArgs args)
+        {
+            List<AclAction> values = new();
+
+            context.Reader.ReadStartArray();
+            for (; ; )
+            {
+                BsonType type = context.Reader.ReadBsonType();
+                if (type == BsonType.EndOfDocument)
+                {
+                    break;
+                }
+                else
+                {
+                    values.Add((AclAction)Enum.Parse(typeof(AclAction), context.Reader.ReadString()));
+                }
+            }
+            context.Reader.ReadEndArray();
+
+            return values;
+        }
+
+        /// <inheritdoc/>
+        public void Serialize(BsonSerializationContext context, BsonSerializationArgs args, List<AclAction> values)
+        {
+            context.Writer.WriteStartArray();
+            if (values.Count > 0)
+            {
+                for (int idx = 0; idx < values.Count; idx++)
+                {
+                    AclAction value = values[idx];
+                    if (values.IndexOf(value) == idx)
+                    {
+                        context.Writer.WriteString(value.ToString());
+                    }
+                }
+            }
+            context.Writer.WriteEndArray();
+        }
+    }
+
+    /// <summary>
+    /// Serializer for JobStepRefId objects
+    /// </summary>
+    public sealed class AclActionSetSerializer : IBsonSerializer<HashSet<AclAction>>
 	{
 		/// <inheritdoc/>
 		public Type ValueType => typeof(HashSet<AclAction>);
diff --git a/Engine/Source/Programs/Horde/Horde.Build/Server/Config.cs b/Engine/Source/Programs/Horde/Horde.Build/Server/Config.cs
--- a/Engine/Source/Programs/Horde/Horde.Build/Server/Config.cs
+++ b/Engine/Source/Programs/Horde/Horde.Build/Server/Config.cs
@@ -8,6 +8,7 @@ using EpicGames.Horde.Common;
 using EpicGames.Serialization;
 using Horde.Build.Api;
 using Horde.Build.Models;
+using Horde.Build.Tools;
 using Horde.Build.Utilities;
 
 namespace Horde.Build.Server
@@ -47,6 +48,11 @@ namespace Horde.Build.Server
 		/// </summary>
 		public List<ComputeClusterConfig> Compute { get; set; } = new List<ComputeClusterConfig>();
 
+		/// <summary>
+		/// List of tools hosted by the server
+		/// </summary>
+		public List<ToolOptions> Tools { get; set; } = new List<ToolOptions>();
+
 		/// <summary>
 		/// Maximum number of conforms to run at once
 		/// </summary>
diff --git a/Engine/Source/Programs/Horde/Horde.Build/Server/ConfigService.cs b/Engine/Source/Programs/Horde/Horde.Build/Server/ConfigService.cs
--- a/Engine/Source/Programs/Horde/Horde.Build/Server/ConfigService.cs
+++ b/Engine/Source/Programs/Horde/Horde.Build/Server/ConfigService.cs
@@ -17,6 +17,7 @@ using Horde.Build.Collections;
 using Horde.Build.Models;
 using Horde.Build.Notifications;
 using Horde.Build.Services;
+using Horde.Build.Tools;
 using Horde.Build.Utilities;
 using HordeCommon;
 using Microsoft.AspNetCore.StaticFiles;
@@ -44,6 +45,7 @@ namespace Horde.Build.Server
 		/// </summary>
 		const int Version = 10;
 		readonly MongoService _mongoService;
+		readonly ToolCollection _toolCollection;
 		readonly ProjectService _projectService;
 		readonly StreamService _streamService;
 		readonly IPerforceService _perforceService;
@@ -57,10 +59,11 @@ namespace Horde.Build.Server
 		/// <summary>
 		/// Constructor
 		/// </summary>
-		public ConfigService(MongoService mongoService, IPerforceService perforceService, ProjectService projectService, StreamService streamService, INotificationService notificationService, PoolService poolService, AgentService agentService, IClock clock, IOptionsMonitor<ServerSettings> settings, ILogger<ConfigService> logger)
+		public ConfigService(MongoService mongoService, IPerforceService perforceService, ToolCollection toolCollection, ProjectService projectService, StreamService streamService, INotificationService notificationService, PoolService poolService, AgentService agentService, IClock clock, IOptionsMonitor<ServerSettings> settings, ILogger<ConfigService> logger)
 		{
 			_mongoService = mongoService;
 			_perforceService = perforceService;
+			_toolCollection = toolCollection;
 			_projectService = projectService;
 			_streamService = streamService;
 			_notificationService = notificationService;
@@ -147,6 +150,9 @@ namespace Horde.Build.Server
 			// Update the agent rate table
 			await _agentService.UpdateRateTableAsync(globalConfig.Rates);
 
+			// Update the tools
+			await _toolCollection.ConfigureAsync(globalConfig.Tools);
+
 			// Projects to remove
 			List<IProject> projects = await _projectService.GetProjectsAsync();
 
diff --git a/Engine/Source/Programs/Horde/Horde.Build/Server/RedisService.cs b/Engine/Source/Programs/Horde/Horde.Build/Server/RedisService.cs
--- a/Engine/Source/Programs/Horde/Horde.Build/Server/RedisService.cs
+++ b/Engine/Source/Programs/Horde/Horde.Build/Server/RedisService.cs
@@ -73,16 +73,22 @@ namespace Horde.Build.Server
 		/// <summary>
 		/// Constructor
 		/// </summary>
-		/// <param name="options"></param>
-		/// <param name="loggerFactory"></param>
 		public RedisService(IOptions<ServerSettings> options, ILoggerFactory loggerFactory)
+			: this(options.Value.RedisConnectionConfig, -1, loggerFactory)
+		{
+		}
+
+		/// <summary>
+		/// Constructor
+		/// </summary>
+		/// <param name="connectionString">Redis connection string. If null, we will start a temporary redis instance on the local machine.</param>
+		/// <param name="dbNum">Override for the database to use. Set to -1 to use the default from the connection string.</param>
+		/// <param name="loggerFactory"></param>
+		public RedisService(string? connectionString, int dbNum, ILoggerFactory loggerFactory)
 		{
 			_loggerFactory = loggerFactory;
 			_logger = loggerFactory.CreateLogger<RedisService>();
 
-			ServerSettings settings = options.Value;
-
-			string? connectionString = settings.RedisConnectionConfig;
 			if (connectionString == null)
 			{
 				if (IsRunningOnDefaultPort())
@@ -95,18 +101,20 @@ namespace Horde.Build.Server
 				}
 				else
 				{
-					throw new Exception($"Unable to connect to Redis. Please set {nameof(settings.RedisConnectionConfig)} in {Program.UserConfigFile}");
+					throw new Exception($"Unable to connect to Redis. Please set {nameof(ServerSettings.RedisConnectionConfig)} in {Program.UserConfigFile}");
 				}
 			}
 
 			Multiplexer = ConnectionMultiplexer.Connect(connectionString);
-			Database = Multiplexer.GetDatabase();
-			ConnectionPool = new RedisConnectionPool(20, connectionString);
+			Database = Multiplexer.GetDatabase(dbNum);
+			ConnectionPool = new RedisConnectionPool(20, connectionString, dbNum);
 		}
 
 		/// <inheritdoc/>
 		public void Dispose()
 		{
+			Multiplexer.Dispose();
+
 			if (_redisProcess != null)
 			{
 				_redisProcess.Dispose();
diff --git a/Engine/Source/Programs/Horde/Horde.Build/ServerSettings.cs b/Engine/Source/Programs/Horde/Horde.Build/ServerSettings.cs
--- a/Engine/Source/Programs/Horde/Horde.Build/ServerSettings.cs
+++ b/Engine/Source/Programs/Horde/Horde.Build/ServerSettings.cs
@@ -512,6 +512,11 @@ namespace Horde.Build
 		/// </summary>
 		public StorageOptions? Storage { get; set; }
 
+		/// <summary>
+		/// Namespace to use for storing tools
+		/// </summary>
+		public NamespaceId ToolNamespaceId { get; set; } = new NamespaceId("horde.tools");
+
 		/// <summary>
 		/// Lazily computed timezone value
 		/// </summary>
diff --git a/Engine/Source/Programs/Horde/Horde.Build/Startup.cs b/Engine/Source/Programs/Horde/Horde.Build/Startup.cs
--- a/Engine/Source/Programs/Horde/Horde.Build/Startup.cs
+++ b/Engine/Source/Programs/Horde/Horde.Build/Startup.cs
@@ -45,6 +45,7 @@ using Horde.Build.Storage;
 using Horde.Build.Storage.Backends;
 using Horde.Build.Tasks;
 using Horde.Build.Tasks.Impl;
+using Horde.Build.Tools;
 using Horde.Build.Utilities;
 using HordeCommon;
 using Microsoft.AspNetCore.Authentication;
@@ -312,6 +313,8 @@ namespace Horde.Build
 			services.AddSingleton<IDeviceCollection, DeviceCollection>();
 			services.AddSingleton<INoticeCollection, NoticeCollection>();
 
+			services.AddSingleton<ToolCollection>();
+
 			// Auditing
 			services.AddSingleton<IAuditLog<AgentId>>(sp => sp.GetRequiredService<IAuditLogFactory<AgentId>>().Create("Agents.Log", "AgentId"));
 
@@ -822,6 +825,9 @@ namespace Horde.Build
 				BsonSerializer.RegisterSerializationProvider(new BsonSerializationProvider());
 				BsonSerializer.RegisterSerializationProvider(new StringIdSerializationProvider());
 				BsonSerializer.RegisterSerializationProvider(new ObjectIdSerializationProvider());
+
+                // Register all the custom class maps
+                BsonClassMap.RegisterClassMap<AclV2>(AclV2.ConfigureClassMap);
 			}
 		}
 
diff --git a/Engine/Source/Programs/Horde/Horde.Build/Tools/Impl/ToolsController.cs b/Engine/Source/Programs/Horde/Horde.Build/Tools/Impl/ToolsController.cs
new file mode 100644
--- /dev/null
+++ b/Engine/Source/Programs/Horde/Horde.Build/Tools/Impl/ToolsController.cs
@@ -0,0 +1,298 @@
+﻿// Copyright Epic Games, Inc. All Rights Reserved.
+
+using EpicGames.Core;
+using Horde.Build.Acls;
+using Horde.Build.Utilities;
+using HordeCommon;
+using Microsoft.AspNetCore.Authorization;
+using Microsoft.AspNetCore.Http;
+using Microsoft.AspNetCore.Mvc;
+using Microsoft.Extensions.Logging;
+using System;
+using System.IO;
+using System.Linq;
+using System.Threading.Tasks;
+
+namespace Horde.Build.Tools.Impl
+{
+	using ToolId = StringId<Tool>;
+	using ToolDeploymentId = ObjectId<ToolDeployment>;
+
+	/// <summary>
+	/// Response from creating a deployment
+	/// </summary>
+	public class CreateDeploymentResponse
+	{
+		/// <summary>
+		/// Identifier for the created deployment
+		/// </summary>
+		public ToolDeploymentId Id { get; set; }
+	}
+
+	/// <summary>
+	/// Update an existing deployment
+	/// </summary>
+	public class UpdateDeploymentRequest
+	{
+		/// <summary>
+		/// New state for the deployment
+		/// </summary>
+		public ToolDeploymentState? State { get; set; }
+	}
+
+	/// <summary>
+	/// Action for a deployment
+	/// </summary>
+	public enum GetToolAction
+	{
+		/// <summary>
+		/// Query for information about the deployment 
+		/// </summary>
+		Info,
+
+		/// <summary>
+		/// Download the deployment data
+		/// </summary>
+		Download,
+	}
+
+	/// <summary>
+	/// Controller for the /api/v1/agents endpoint
+	/// </summary>
+	[ApiController]
+	[Authorize]
+	[Route("[controller]")]
+	public class ToolsController : HordeControllerBase
+	{
+		readonly AclService _aclService;
+		readonly ToolCollection _toolCollection;
+
+		/// <summary>
+		/// Constructor
+		/// </summary>
+		public ToolsController(AclService aclService, ToolCollection toolCollection)
+		{
+			_aclService = aclService;
+			_toolCollection = toolCollection;
+		}
+
+		/// <summary>
+		/// Create a new deployment of the given tool.
+		/// </summary>
+		/// <returns>Information about the registered agent</returns>
+		[HttpPost]
+		[Route("/api/v1/tools/{id}/deployments")]
+		public async Task<ActionResult<CreateDeploymentResponse>> CreateDeploymentAsync(ToolId id, [FromForm] ToolDeploymentOptions options, [FromForm] IFormFile file)
+		{
+			Tool? tool = await _toolCollection.GetAsync(id);
+			if (tool == null)
+			{
+				return NotFound(id);
+			}
+			if (!await tool.AuthorizeAsync(AclAction.UploadTool, User, _aclService, null))
+			{
+				return Forbid(AclAction.UploadTool, id);
+			}
+
+			using (Stream stream = file.OpenReadStream())
+			{
+				tool = await _toolCollection.CreateDeploymentAsync(tool, options, stream);
+				if (tool == null)
+				{
+					return NotFound(id);
+				}
+			}
+
+			return new CreateDeploymentResponse { Id = tool.Deployments[^1].Id };
+		}
+
+		/// <summary>
+		/// Register an agent to perform remote work.
+		/// </summary>
+		/// <returns>Information about the registered agent</returns>
+		[HttpGet]
+		[Route("/api/v1/tools/{id}/deployments/{deploymentId}")]
+		public async Task<ActionResult> GetDeploymentAsync(ToolId id, ToolDeploymentId deploymentId, [FromQuery] GetToolAction action = GetToolAction.Info)
+		{
+			Tool? tool = await _toolCollection.GetAsync(id);
+			if (tool == null)
+			{
+				return NotFound(id);
+			}
+			if (!await tool.AuthorizeAsync(AclAction.DownloadTool, User, _aclService, null))
+			{
+				return Forbid(AclAction.DownloadTool, id);
+			}
+
+			ToolDeployment? deployment = tool.Deployments.FirstOrDefault(x => x.Id == deploymentId);
+			if(deployment == null)
+			{
+				return NotFound(id, deploymentId);
+			}
+
+			return await GetDeploymentResponseAsync(tool, deployment, action);
+		}
+
+		private async Task<ActionResult> GetDeploymentResponseAsync(Tool tool, ToolDeployment deployment, GetToolAction action)
+		{
+			if (action == GetToolAction.Info)
+			{
+				return Ok(deployment);
+			}
+			else
+			{
+				Stream stream = await _toolCollection.GetDeploymentPayloadAsync(tool, deployment);
+				return new FileStreamResult(stream, "application/zip");
+			}
+		}
+
+		/// <summary>
+		/// Updates the state of an active deployment.
+		/// </summary>
+		[HttpPatch]
+		[Route("/api/v1/tools/{id}/deployments/{deploymentId}")]
+		public async Task<ActionResult> UpdateDeploymentAsync(ToolId id, ToolDeploymentId deploymentId, [FromBody] UpdateDeploymentRequest request)
+		{
+			Tool? tool = await _toolCollection.GetAsync(id);
+			if (tool == null)
+			{
+				return NotFound(id);
+			}
+			if (!await tool.AuthorizeAsync(AclAction.UploadTool, User, _aclService, null))
+			{
+				return Forbid(AclAction.UploadTool, id);
+			}
+
+			if (request.State != null)
+			{
+				tool = await _toolCollection.UpdateDeploymentAsync(tool, deploymentId, request.State.Value);
+				if (tool == null)
+				{
+					return NotFound(id, deploymentId);
+				}
+			}
+			return Ok();
+		}
+	}
+
+	/// <summary>
+	/// Public methods available without authorization (or with very custom authorization)
+	/// </summary>
+	[ApiController]
+	public class PublicToolsController : HordeControllerBase
+	{
+		readonly AclService _aclService;
+		readonly ToolCollection _toolCollection;
+		readonly IClock _clock;
+
+		/// <summary>
+		/// Constructor
+		/// </summary>
+		public PublicToolsController(AclService aclService, ToolCollection toolCollection, IClock clock)
+		{
+			_aclService = aclService;
+			_toolCollection = toolCollection;
+			_clock = clock;
+		}
+
+		/// <summary>
+		/// Register an agent to perform remote work.
+		/// </summary>
+		/// <returns>Information about the registered agent</returns>
+		[HttpGet]
+		[Route("/api/v1/tools/{id}")]
+		public async Task<ActionResult> GetToolAsync(ToolId id, GetToolAction action = GetToolAction.Info)
+		{
+			Tool? tool = await _toolCollection.GetAsync(id);
+			if (tool == null)
+			{
+				return NotFound(id);
+			}
+			if (!await AuthorizeDownloadAsync(tool))
+			{
+				return Forbid(AclAction.DownloadTool, id);
+			}
+			if(tool.Deployments.Count == 0)
+			{
+				return NotFound(LogEvent.Create(LogLevel.Error, "Tool {ToolId} does not currently have any deployments", id));
+			}
+
+			if (action == GetToolAction.Info)
+			{
+				return Ok(tool);
+			}
+			else
+			{
+				return await GetDeploymentPayloadAsync(tool, tool.Deployments[^1]);
+			}
+		}
+
+		/// <summary>
+		/// Register an agent to perform remote work.
+		/// </summary>
+		/// <param name="id">The tool identifier</param>
+		/// <param name="phase">Value indicating the client's preference for deployment to receive.</param>
+		/// <param name="action">Information about the returned deployment</param>
+		/// <returns>Information about the registered agent</returns>
+		[HttpGet]
+		[Route("/api/v1/tools/{id}/deployments")]
+		public async Task<ActionResult> FindDeploymentAsync(ToolId id, [FromQuery] double phase = 0.0, [FromQuery] GetToolAction action = GetToolAction.Info)
+		{
+			Tool? tool = await _toolCollection.GetAsync(id);
+			if (tool == null)
+			{
+				return NotFound(id);
+			}
+			if (!await AuthorizeDownloadAsync(tool))
+			{
+				return Forbid(AclAction.DownloadTool, id);
+			}
+			if (tool.Deployments.Count == 0)
+			{
+				return NotFound(LogEvent.Create(LogLevel.Error, "Tool {ToolId} does not currently have any deployments", id));
+			}
+
+			DateTime utcNow = _clock.UtcNow;
+
+			int idx = tool.Deployments.Count - 1;
+			for (; idx > 0; idx--)
+			{
+				if (phase < tool.Deployments[idx].GetProgressValue(utcNow))
+				{
+					break;
+				}
+			}
+
+			if (action == GetToolAction.Info)
+			{
+				return Ok(tool.Deployments[idx]);
+			}
+			else
+			{
+				return await GetDeploymentPayloadAsync(tool, tool.Deployments[idx]);
+			}
+		}
+
+		async Task<ActionResult> GetDeploymentPayloadAsync(Tool tool, ToolDeployment deployment)
+		{
+			Stream stream = await _toolCollection.GetDeploymentPayloadAsync(tool, deployment);
+			return new FileStreamResult(stream, deployment.MimeType);
+		}
+
+		async ValueTask<bool> AuthorizeDownloadAsync(Tool tool)
+		{
+			if (!tool.Public)
+			{
+				if (User.Identity == null || !User.Identity.IsAuthenticated)
+				{
+					return false;
+				}
+				if (!await tool.AuthorizeAsync(AclAction.DownloadTool, User, _aclService, null))
+				{
+					return false;
+				}
+			}
+			return true;
+		}
+	}
+}
diff --git a/Engine/Source/Programs/Horde/Horde.Build/Tools/Tool.cs b/Engine/Source/Programs/Horde/Horde.Build/Tools/Tool.cs
new file mode 100644
--- /dev/null
+++ b/Engine/Source/Programs/Horde/Horde.Build/Tools/Tool.cs
@@ -0,0 +1,342 @@
+// Copyright Epic Games, Inc. All Rights Reserved.
+
+using EpicGames.Horde.Storage;
+using Horde.Build.Acls;
+using Horde.Build.Utilities;
+using MongoDB.Bson.Serialization.Attributes;
+using System;
+using System.Collections.Generic;
+using System.Security.Claims;
+using System.Text.Json.Serialization;
+using System.Threading.Tasks;
+
+namespace Horde.Build.Tools
+{
+	using ToolId = StringId<Tool>;
+	using ToolDeploymentId = ObjectId<ToolDeployment>;
+
+	/// <summary>
+	/// Options for configuring a tool
+	/// </summary>
+	public class ToolOptions
+	{
+		/// <inheritdoc cref="VersionedDocument{ToolId, Tool}.Id"/>
+		public ToolId Id { get; set; }
+
+		/// <inheritdoc cref="Tool.Name"/>
+		public string Name { get; set; }
+
+		/// <inheritdoc cref="Tool.Description"/>
+		public string Description { get; set; }
+
+		/// <inheritdoc cref="Tool.Public"/>
+		public bool Public { get; set; }
+
+		/// <inheritdoc cref="Tool.Acl"/>
+		public AclV2? Acl { get; set; }
+
+		/// <summary>
+		/// Constructor
+		/// </summary>
+		public ToolOptions(ToolId id)
+		{
+			Id = id;
+			Name = id.ToString();
+			Description = String.Empty;
+		}
+	}
+
+	/// <summary>
+	/// Describes a standalone, external tool hosted and deployed by Horde. Provides basic functionality for performing
+	/// gradual roll-out, versioning, etc...
+	/// </summary>
+	public class Tool : VersionedDocument<ToolId, Tool>
+	{
+		/// <summary>
+		/// Name of the tool
+		/// </summary>
+		[BsonElement("nam")]
+		[JsonPropertyOrder(2)]
+		public string Name { get; set; }
+
+		/// <summary>
+		/// Description for the tool
+		/// </summary>
+		[BsonElement("dsc")]
+		[JsonPropertyOrder(3)]
+		public string Description { get; set; }
+
+		/// <summary>
+		/// Current deployments of this tool, sorted by time.
+		/// </summary>
+		[BsonElement("dep")]
+		[JsonPropertyOrder(4)]
+		public List<ToolDeployment> Deployments { get; set; } = new List<ToolDeployment>();
+
+		/// <summary>
+		/// Whether this tool should be exposed for download on a public endpoint without authentication
+		/// </summary>
+		[BsonElement("pub")]
+		[JsonPropertyOrder(5)]
+		public bool Public { get; set; }
+
+		/// <summary>
+		/// Access list for this tool
+		/// </summary>
+		[BsonElement("acl")]
+		[JsonPropertyOrder(3), JsonIgnore(Condition = JsonIgnoreCondition.WhenWritingNull)]
+		public AclV2? Acl
+		{
+			get => _acl;
+			set => _acl = (value is AclV2 && !value.IsDefault()) ? value : null;
+		}
+
+		private AclV2? _acl;
+
+		/// <summary>
+		/// Constructor
+		/// </summary>
+		[JsonConstructor]
+		[BsonConstructor]
+		public Tool(ToolId id)
+			: base(id)
+		{
+			Id = id;
+			Name = id.ToString();
+			Description = String.Empty;
+		}
+
+		/// <summary>
+		/// Constructor
+		/// </summary>
+		public Tool(ToolOptions config)
+			: base(config.Id)
+		{
+			Name = config.Name;
+			Description = config.Description;
+			Public = config.Public;
+			Acl = config.Acl;
+		}
+
+		/// <inheritdoc/>
+		public override Tool UpgradeToLatest() => this;
+
+		/// <summary>
+		/// Updates the state values derived from the current time
+		/// </summary>
+		/// <param name="utcNow"></param>
+		public void UpdateTemporalState(DateTime utcNow)
+		{
+			foreach (ToolDeployment deployment in Deployments)
+			{
+				deployment.UpdateTemporalState(utcNow);
+			}
+		}
+
+		/// <summary>
+		/// Test whether a user can perform an action on this tool 
+		/// </summary>
+		public Task<bool> AuthorizeAsync(AclAction action, ClaimsPrincipal user, AclService aclService, GlobalPermissionsCache? cache)
+		{
+			bool? result = Acl?.Authorize(action, user);
+			if (result == null)
+			{
+				return aclService.AuthorizeAsync(action, user, cache);
+			}
+			else
+			{
+				return Task.FromResult(result.Value);
+			}
+		}
+	}
+
+	/// <summary>
+	/// Current state of a tool's deployment
+	/// </summary>
+	public enum ToolDeploymentState
+	{
+		/// <summary>
+		/// The deployment is ongoing
+		/// </summary>
+		Active,
+
+		/// <summary>
+		/// The deployment should be paused at its current state
+		/// </summary>
+		Paused,
+
+		/// <summary>
+		/// Deployment of this version is complete
+		/// </summary>
+		Complete,
+
+		/// <summary>
+		/// The deployment has been cancelled.
+		/// </summary>
+		Cancelled,
+	}
+
+	/// <summary>
+	/// Options for a new deployment
+	/// </summary>
+	public class ToolDeploymentOptions
+	{
+		/// <summary>
+		/// Default mime types for deployment data
+		/// </summary>
+		public const string DefaultMimeType = "application/zip";
+
+		/// <inheritdoc cref="ToolDeployment.Version"/>
+		public string Version { get; set; } = "Unknown";
+
+		/// <inheritdoc cref="ToolDeployment.Duration"/>
+		public TimeSpan Duration { get; set; }
+
+		/// <inheritdoc cref="ToolDeployment.MimeType"/>
+		public string MimeType { get; set; } = DefaultMimeType;
+
+		/// <summary>
+		/// Whether to create the deployment in a paused state
+		/// </summary>
+		public bool CreatePaused { get; set; }
+	}
+
+	/// <summary>
+	/// Deployment of a tool
+	/// </summary>
+	public class ToolDeployment
+	{
+		/// <summary>
+		/// Identifier for this deployment. A new identifier will be assigned to each created instance, so an identifier corresponds to a unique deployment.
+		/// </summary>
+		[JsonPropertyOrder(0)]
+		public ToolDeploymentId Id { get; set; }
+
+		/// <summary>
+		/// Descriptive version string for this tool revision
+		/// </summary>
+		[BsonElement("ver")]
+		[JsonPropertyOrder(1)]
+		public string Version { get; set; }
+
+		/// <summary>
+		/// Current state of this deployment
+		/// </summary>
+		[BsonIgnore]
+		[JsonPropertyOrder(2)]
+		public ToolDeploymentState State { get; set; }
+
+		/// <summary>
+		/// Current progress of the deployment
+		/// </summary>
+		[BsonIgnore]
+		[JsonPropertyOrder(3)]
+		public double Progress { get; set; }
+
+		/// <summary>
+		/// Progress in the deployment at <see cref="StartedAt"/>, from 0-1.
+		/// </summary>
+		[JsonIgnore]
+		[BsonElement("bpr")]
+		public double BaseProgress { get; set; }
+
+		/// <summary>
+		/// Last time at which the progress started. Set to null if the deployment was paused.
+		/// </summary>
+		[JsonPropertyOrder(4)]
+		[BsonElement("stm")]
+		public DateTime? StartedAt { get; set; }
+
+		/// <summary>
+		/// Length of time over which to make the deployment
+		/// </summary>
+		[BsonElement("dur")]
+		[JsonPropertyOrder(5)]
+		public TimeSpan Duration { get; set; }
+
+		/// <summary>
+		/// Mime-type for the data
+		/// </summary>
+		[BsonElement("mtp")]
+		[JsonPropertyOrder(6)]
+		public string MimeType { get; set; }
+
+		/// <summary>
+		/// Reference to this tool in Horde Storage.
+		/// </summary>
+		[BsonElement("ref")]
+		[JsonPropertyOrder(7)]
+		public RefId RefId { get; set; }
+
+		/// <summary>
+		/// Constructor
+		/// </summary>
+		public ToolDeployment(ToolDeploymentId id)
+		{
+			Id = id;
+			Version = String.Empty;
+			MimeType = ToolDeploymentOptions.DefaultMimeType;
+		}
+
+		/// <summary>
+		/// Constructor
+		/// </summary>
+		public ToolDeployment(ToolDeploymentId id, ToolDeploymentOptions options, RefId refId)
+		{
+			Id = id;
+			Version = options.Version;
+			Duration = options.Duration;
+			MimeType = options.MimeType;
+			RefId = refId;
+		}
+
+		/// <summary>
+		/// Update the computed deployment state
+		/// </summary>
+		/// <param name="utcNow"></param>
+		public void UpdateTemporalState(DateTime utcNow)
+		{
+			if (BaseProgress >= 1.0)
+			{
+				State = ToolDeploymentState.Complete;
+				Progress = 1.0;
+			}
+			else if (StartedAt == null)
+			{
+				State = ToolDeploymentState.Paused;
+				Progress = BaseProgress;
+			}
+			else if (Duration > TimeSpan.Zero)
+			{
+				State = ToolDeploymentState.Active;
+				Progress = Math.Clamp((utcNow - StartedAt.Value) / Duration, 0.0, 1.0);
+			}
+			else
+			{
+				State = ToolDeploymentState.Complete;
+				Progress = 1.0;
+			}
+		}
+
+		/// <summary>
+		/// Get the updated progress value for a particular time
+		/// </summary>
+		/// <param name="utcNow">The current time</param>
+		/// <returns></returns>
+		public double GetProgressValue(DateTime utcNow)
+		{
+			if (StartedAt == null)
+			{
+				return Progress;
+			}
+			else if (Duration > TimeSpan.Zero)
+			{
+				return Math.Clamp((utcNow - StartedAt.Value) / Duration, 0.0, 1.0);
+			}
+			else
+			{
+				return 1.0;
+			}
+		}
+	}
+}
diff --git a/Engine/Source/Programs/Horde/Horde.Build/Tools/ToolCollection.cs b/Engine/Source/Programs/Horde/Horde.Build/Tools/ToolCollection.cs
new file mode 100644
--- /dev/null
+++ b/Engine/Source/Programs/Horde/Horde.Build/Tools/ToolCollection.cs
@@ -0,0 +1,410 @@
+// Copyright Epic Games, Inc. All Rights Reserved.
+
+using EpicGames.Core;
+using EpicGames.Horde.Storage;
+using EpicGames.Serialization;
+using Horde.Build.Acls;
+using Horde.Build.Server;
+using Horde.Build.Utilities;
+using HordeCommon;
+using Microsoft.Extensions.Logging;
+using Microsoft.Extensions.Options;
+using MongoDB.Bson;
+using MongoDB.Bson.Serialization;
+using MongoDB.Driver;
+using StackExchange.Redis;
+using System;
+using System.Collections.Generic;
+using System.IO;
+using System.Linq;
+using System.Threading.Tasks;
+
+namespace Horde.Build.Tools
+{
+	using ToolId = StringId<Tool>;
+	using ToolDeploymentId = ObjectId<ToolDeployment>;
+
+	/// <summary>
+	/// Collection of tool documents
+	/// </summary>
+	public class ToolCollection
+	{
+		private class ToolDeploymentData
+		{
+			[CbField]
+			public ToolDeploymentId Id { get; set; }
+
+			[CbField("version")]
+			public string Version { get; set; } = String.Empty;
+
+			[CbField("data")]
+			public CbBinaryAttachment Data { get; set; }
+		}
+
+		private class CachedIndex
+		{
+			[CbField("rev")]
+			public string Rev { get; set; } = String.Empty;
+
+			[CbField("empty")]
+			public bool Empty { get; set; }
+
+			[CbField("ids")]
+			public List<ToolId> Ids { get; set; } = new List<ToolId>();
+		}
+
+		private readonly RedisService _redisService;
+		private readonly VersionedCollection<ToolId, Tool> _tools;
+		private readonly IStorageClient _storage;
+		private readonly IClock _clock;
+		private readonly NamespaceId _namespaceId;
+		private readonly ILogger _logger;
+
+		private static readonly RedisKey s_baseKey = "tools/v1/";
+		private static readonly RedisKey s_indexKey = s_baseKey.Append("index");
+
+		private static readonly IReadOnlyDictionary<int, Type> s_types = RegisterTypes();
+
+		/// <summary>
+		/// Constructor
+		/// </summary>
+		public ToolCollection(MongoService mongoService, RedisService redisService, IStorageClient storage, IClock clock, IOptions<ServerSettings> options, ILogger<ToolCollection> logger)
+		{
+			_redisService = redisService;
+			_tools = new VersionedCollection<ToolId, Tool>(mongoService, "Tools", redisService, s_baseKey, s_types);
+			_storage = storage;
+			_clock = clock;
+			_namespaceId = options.Value.ToolNamespaceId;
+			_logger = logger;
+		}
+
+		/// <summary>
+		/// Registers types required for this collection
+		/// </summary>
+		/// <returns></returns>
+		static IReadOnlyDictionary<int, Type> RegisterTypes()
+		{
+			Dictionary<int, Type> versionToType = new Dictionary<int, Type>();
+
+			versionToType[1] = typeof(Tool);
+			BsonClassMap.RegisterClassMap<Tool>(cm =>
+			{
+				cm.AutoMap();
+				cm.MapCreator(t => new Tool(t.Id));
+				cm.MapMember(x => x.Acl).SetIgnoreIfDefault(true).SetDefaultValue(new AclV2());
+			});
+
+			return versionToType;
+		}
+
+		/// <summary>
+		/// Configures all the available tools
+		/// </summary>
+		/// <param name="tools">The list of configured tools</param>
+		public async Task ConfigureAsync(List<ToolOptions> tools)
+		{
+			List<ToolId> ids = await FindAllIdsAsync();
+
+			List<ToolId> removeIds = ids.Except(tools.Select(x => x.Id)).ToList();
+			foreach (ToolId removeId in removeIds)
+			{
+				_logger.LogInformation("Removing tool {ToolId}", removeId);
+				await _tools.DeleteAsync(removeId);
+			}
+
+			foreach (ToolOptions tool in tools)
+			{
+				if (!ids.Contains(tool.Id))
+				{
+					_logger.LogInformation("Creating tool {ToolId}", tool.Id);
+				}
+				await AddOrUpdateAsync(tool);
+			}
+		}
+
+		/// <summary>
+		/// Finds or adds a tool with the given identifier
+		/// </summary>
+		/// <param name="options">Options for the tool</param>
+		/// <returns>Document describing the tool</returns>
+		private async Task<Tool> AddOrUpdateAsync(ToolOptions options)
+		{
+			for (; ; )
+			{
+				// Check for the common case where the tool already exists and is set to the right values already
+				Tool tool = await _tools.FindOrAddAsync(options.Id, () => new Tool(options));
+				if (tool.Name.Equals(options.Name, StringComparison.Ordinal) && tool.Description.Equals(options.Description, StringComparison.Ordinal) && tool.Public == options.Public && AclV2.Equals(tool.Acl, options.Acl))
+				{
+					return tool;
+				}
+
+				// Otherwise attempt to do an in-place upgrade
+				UpdateDefinition<Tool> update = Builders<Tool>.Update.Set(x => x.Name, options.Name).Set(x => x.Description, options.Description).Set(x => x.Public, options.Public);
+				if (AclV2.IsNullOrDefault(options.Acl))
+				{
+					update = update.Unset(x => x.Acl);
+				}
+				else
+				{
+					update = update.Set(x => x.Acl, options.Acl);
+				}
+
+				// Update the existing tool
+				Tool? updatedTool = await _tools.UpdateAsync(tool, update);
+				if (updatedTool != null)
+				{
+					await ClearCachedIndexAsync();
+					return updatedTool;
+				}
+			}
+		}
+
+		/// <summary>
+		/// Removes a tool with the given identifier
+		/// </summary>
+		/// <param name="id">The tool to remove</param>
+		/// <returns>True if a tool with the given identifier was deleted</returns>
+		private async Task<bool> DeleteAsync(ToolId id)
+		{
+			bool deleted = await _tools.DeleteAsync(id);
+			if (deleted)
+			{
+				await ClearCachedIndexAsync();
+			}
+			return deleted;
+		}
+
+		private async Task ClearCachedIndexAsync()
+		{
+			CachedIndex newIndexData = new CachedIndex { Rev = ObjectId.GenerateNewId().ToString() };
+			await _redisService.Database.StringSetAsync(s_indexKey, CbSerializer.SerializeToByteArray(newIndexData));
+		}
+
+		/// <summary>
+		/// Finds all tools matching a set of criteria
+		/// </summary>
+		/// <returns>Sequence of tool documents</returns>
+		async Task<List<ToolId>> FindAllIdsAsync()
+		{
+			CachedIndex? index;
+			for (; ; )
+			{
+				// Get the cached index, and check if it has valid data
+				RedisValue value = await _redisService.Database.StringGetAsync(s_indexKey);
+				if (!value.IsNullOrEmpty)
+				{
+					index = CbSerializer.Deserialize<CachedIndex>((byte[])value);
+					if (index.Empty || index.Ids.Count > 0)
+					{
+						break;
+					}
+				}
+
+				// Create a new index object
+				index = new CachedIndex();
+				index.Rev = ObjectId.GenerateNewId().ToString();
+				index.Ids = await _tools.BaseCollection.Find(FilterDefinition<VersionedDocument<ToolId, Tool>>.Empty).Project(x => x.Id).ToListAsync();
+				index.Empty = index.Ids.Count == 0;
+
+				// Try to replace the existing one
+				RedisValue newValue = CbSerializer.SerializeToByteArray(index);
+				if (value.IsNull)
+				{
+					if (await _redisService.Database.StringSetAsync(s_indexKey, newValue, when: When.NotExists))
+					{
+						break;
+					}
+				}
+				else
+				{
+					ITransaction transaction = _redisService.Database.CreateTransaction();
+					transaction.AddCondition(Condition.StringEqual(s_indexKey, value));
+					_ = transaction.StringSetAsync(s_indexKey, newValue);
+
+					if (await transaction.ExecuteAsync())
+					{
+						break;
+					}
+				}
+			}
+			return index.Ids;
+		}
+
+		/// <summary>
+		/// Finds all tools matching a set of criteria
+		/// </summary>
+		/// <returns>Sequence of tool documents</returns>
+		public async Task<List<Tool>> FindAllAsync()
+		{
+			// Fetch all the tools in parallel
+			IEnumerable<ToolId> toolIds =await FindAllIdsAsync();
+			List<Task<Tool?>> tasks = toolIds.Select(x => Task.Run(() => GetAsync(x))).ToList();
+
+			// Return all the successful results
+			List<Tool> results = new List<Tool>();
+			foreach (Task<Tool?> task in tasks)
+			{
+				Tool? tool = await task;
+				if (tool != null)
+				{
+					results.Add(tool);
+				}
+			}
+			return results;
+		}
+
+		/// <summary>
+		/// Gets a tool with the given identifier
+		/// </summary>
+		/// <param name="id">The tool identifier</param>
+		/// <returns></returns>
+		public async Task<Tool?> GetAsync(ToolId id)
+		{
+			Tool? tool = await _tools.GetAsync(id);
+			if (tool != null)
+			{
+				tool.UpdateTemporalState(_clock.UtcNow);
+			}
+			return tool;
+		}
+
+		/// <summary>
+		/// Adds a new deployment to the given tool. The new deployment will replace the current active deployment.
+		/// </summary>
+		/// <param name="tool">The tool to update</param>
+		/// <param name="options">Options for the new deployment</param>
+		/// <param name="stream">Stream containing the tool data</param>
+		/// <returns>Updated tool document, or null if it does not exist</returns>
+		public async Task<Tool?> CreateDeploymentAsync(Tool tool, ToolDeploymentOptions options, Stream stream)
+		{
+			// Upload the tool data first
+			IoHash hash = await _storage.WriteBlobAsync(_namespaceId, stream);
+
+			ToolDeploymentData deploymentData = new ToolDeploymentData();
+			deploymentData.Id = ToolDeploymentId.GenerateNewId();
+			deploymentData.Version = options.Version;
+			deploymentData.Data = new CbBinaryAttachment(hash);
+
+			BucketId bucketId = GetBucket(tool.Id);
+			RefId refId = new RefId(deploymentData.Id.ToString());
+			await _storage.SetRefAsync(_namespaceId, bucketId, refId, deploymentData);
+
+			// Create the new deployment object
+			ToolDeployment deployment = new ToolDeployment(deploymentData.Id, options, refId);
+
+			// Start the deployment
+			DateTime utcNow = _clock.UtcNow;
+			if (!options.CreatePaused)
+			{
+				deployment.StartedAt = utcNow;
+			}
+
+			// Create the deployment
+			Tool? newTool = tool;
+			for (; ; )
+			{
+				newTool = await TryAddDeploymentAsync(newTool, deployment);
+				if (newTool != null)
+				{
+					break;
+				}
+
+				newTool = await GetAsync(tool.Id);
+				if (newTool == null)
+				{
+					return null;
+				}
+			}
+
+			// Return the new tool with updated deployment states
+			newTool.UpdateTemporalState(utcNow);
+			return newTool;
+		}
+
+		async ValueTask<Tool?> TryAddDeploymentAsync(Tool tool, ToolDeployment deployment)
+		{
+			Tool? newTool = tool;
+
+			// If there are already a maximum number of deployments, remove the oldest one
+			const int MaxDeploymentCount = 5;
+			while (tool.Deployments.Count >= MaxDeploymentCount)
+			{
+				newTool = await _tools.UpdateAsync(newTool, Builders<Tool>.Update.PopFirst(x => x.Deployments));
+				if (newTool == null)
+				{
+					return null;
+				}
+				await _storage.DeleteRefAsync(_namespaceId, GetBucket(newTool.Id), tool.Deployments[0].RefId);
+			}
+
+			// Add the new deployment
+			return await _tools.UpdateAsync(newTool, Builders<Tool>.Update.Push(x => x.Deployments, deployment));
+		}
+
+		/// <summary>
+		/// Updates the state of the current deployment
+		/// </summary>
+		/// <param name="tool">Tool to be updated</param>
+		/// <param name="deploymentId">Identifier for the deployment to modify</param>
+		/// <param name="action">New state of the deployment</param>
+		/// <returns></returns>
+		public async Task<Tool?> UpdateDeploymentAsync(Tool tool, ToolDeploymentId deploymentId, ToolDeploymentState action)
+		{
+			int idx = tool.Deployments.FindIndex(x => x.Id == deploymentId);
+			if (idx == -1)
+			{
+				return null;
+			}
+
+			BucketId bucketId = GetBucket(tool.Id);
+
+			ToolDeployment deployment = tool.Deployments[idx];
+			switch (action)
+			{
+				case ToolDeploymentState.Complete:
+					return await _tools.UpdateAsync(tool, Builders<Tool>.Update.Set(x => x.Deployments[idx].BaseProgress, 1.0).Unset(x => x.Deployments[idx].StartedAt));
+
+				case ToolDeploymentState.Cancelled:
+					List<ToolDeployment> newDeployments = tool.Deployments.Where(x => x != deployment).ToList();
+					return await _tools.UpdateAsync(tool, Builders<Tool>.Update.Set(x => x.Deployments, newDeployments));
+
+				case ToolDeploymentState.Paused:
+					if (deployment.StartedAt == null)
+					{
+						return tool;
+					}
+					else
+					{
+						return await _tools.UpdateAsync(tool, Builders<Tool>.Update.Set(x => x.Deployments[idx].Progress, deployment.GetProgressValue(_clock.UtcNow)).Set(x => x.Deployments[idx].StartedAt, null));
+					}
+
+				case ToolDeploymentState.Active:
+					if (deployment.StartedAt != null)
+					{
+						return tool;
+					}
+					else
+					{
+						return await _tools.UpdateAsync(tool, Builders<Tool>.Update.Set(x => x.Deployments[idx].StartedAt, _clock.UtcNow));
+					}
+
+				default:
+					throw new ArgumentException("Invalid action for deployment", nameof(action));
+			}
+		}
+
+		/// <summary>
+		/// Opens a stream to the data for a particular deployment
+		/// </summary>
+		/// <param name="tool">Identifier for the tool</param>
+		/// <param name="deployment">The deployment</param>
+		/// <returns>Stream for the data</returns>
+		public async Task<Stream> GetDeploymentPayloadAsync(Tool tool, ToolDeployment deployment)
+		{
+			ToolDeploymentData data = await _storage.GetRefAsync<ToolDeploymentData>(_namespaceId, GetBucket(tool.Id), deployment.RefId);
+			return await _storage.ReadBlobAsync(_namespaceId, data.Data);
+		}
+
+		private static BucketId GetBucket(ToolId toolId) => new BucketId(toolId.ToString());
+	}
+}
diff --git a/Engine/Source/Programs/Horde/Horde.Build/Utilities/HordeControllerBase.cs b/Engine/Source/Programs/Horde/Horde.Build/Utilities/HordeControllerBase.cs
--- a/Engine/Source/Programs/Horde/Horde.Build/Utilities/HordeControllerBase.cs
+++ b/Engine/Source/Programs/Horde/Horde.Build/Utilities/HordeControllerBase.cs
@@ -3,6 +3,7 @@
 using EpicGames.Core;
 using Horde.Build.Acls;
 using Horde.Build.Models;
+using Horde.Build.Tools;
 using Microsoft.AspNetCore.Http;
 using Microsoft.AspNetCore.Mvc;
 using Microsoft.Extensions.Logging;
@@ -13,6 +14,8 @@ namespace Horde.Build.Utilities
 	using LeaseId = ObjectId<ILease>;
 	using ProjectId = StringId<IProject>;
 	using StreamId = StringId<IStream>;
+	using ToolId = StringId<Tool>;
+	using ToolDeploymentId = ObjectId<ToolDeployment>;
 	using TemplateRefId = StringId<TemplateRef>;
 
 	/// <summary>
@@ -88,6 +91,15 @@ namespace Horde.Build.Utilities
 			return Forbid(action, "stream {StreamId}", streamId);
 		}
 
+		/// <summary>
+		/// Returns a 403 (forbidden) response with the given action and object
+		/// </summary>
+		[NonAction]
+		protected ActionResult Forbid(AclAction action, ToolId toolId)
+		{
+			return Forbid(action, "tool {ToolId}", toolId);
+		}
+
 		/// <summary>
 		/// Returns a 403 (forbidden) response with the given action and object
 		/// </summary>
@@ -196,6 +208,24 @@ namespace Horde.Build.Utilities
 			return NotFound("Template {TemplateId} not found on stream {StreamId}", templateId, streamId);
 		}
 
+		/// <summary>
+		/// Returns a 404 response for the given object
+		/// </summary>
+		[NonAction]
+		protected ActionResult NotFound(ToolId toolId)
+		{
+			return NotFound("Tool {ToolId} not found", toolId);
+		}
+
+		/// <summary>
+		/// Returns a 404 response for the given object
+		/// </summary>
+		[NonAction]
+		protected ActionResult NotFound(ToolId toolId, ToolDeploymentId deploymentId)
+		{
+			return NotFound("Deployment {DeploymentId} not found on tool {ToolId}", deploymentId, toolId);
+		}
+
 		/// <summary>
 		/// Returns a 404 response for the given object
 		/// </summary>
diff --git a/Engine/Source/Programs/Horde/Horde.Build/Utilities/MongoExtensions.cs b/Engine/Source/Programs/Horde/Horde.Build/Utilities/MongoExtensions.cs
--- a/Engine/Source/Programs/Horde/Horde.Build/Utilities/MongoExtensions.cs
+++ b/Engine/Source/Programs/Horde/Horde.Build/Utilities/MongoExtensions.cs
@@ -16,6 +16,16 @@ namespace Horde.Build.Utilities
 	/// </summary>
 	public static class MongoExtensions
 	{
+		/// <summary>
+		/// Rounds a time value to its BSON equivalent (ie. milliseconds since Unix Epoch).
+		/// </summary>
+		/// <param name="time"></param>
+		/// <returns></returns>
+		public static DateTime RoundToBsonDateTime(DateTime time)
+		{
+			return BsonUtils.ToDateTimeFromMillisecondsSinceEpoch(BsonUtils.ToMillisecondsSinceEpoch(time));
+		}
+
 		/// <summary>
 		/// Executes a query with a particular index hint
 		/// </summary>
diff --git a/Engine/Source/Programs/Horde/Horde.Build/Utilities/ObjectId.cs b/Engine/Source/Programs/Horde/Horde.Build/Utilities/ObjectId.cs
--- a/Engine/Source/Programs/Horde/Horde.Build/Utilities/ObjectId.cs
+++ b/Engine/Source/Programs/Horde/Horde.Build/Utilities/ObjectId.cs
@@ -7,6 +7,7 @@ using System.Globalization;
 using System.Text.Json;
 using System.Text.Json.Serialization;
 using EpicGames.Core;
+using EpicGames.Serialization;
 using MongoDB.Bson;
 using MongoDB.Bson.Serialization;
 using MongoDB.Bson.Serialization.Serializers;
@@ -17,6 +18,7 @@ namespace Horde.Build.Utilities
 	/// Normalized string identifier for a resource
 	/// </summary>
 	[JsonSchemaString]
+	[JsonConverter(typeof(JsonObjectIdConverterFactory))]
 	[TypeConverter(typeof(ObjectIdTypeConverter))]
 	public struct ObjectId<T> : IEquatable<ObjectId<T>>, IComparable<ObjectId<T>>
 	{
diff --git a/Engine/Source/Programs/Horde/Horde.Build/Utilities/VersionedCollection.cs b/Engine/Source/Programs/Horde/Horde.Build/Utilities/VersionedCollection.cs
--- a/Engine/Source/Programs/Horde/Horde.Build/Utilities/VersionedCollection.cs
+++ b/Engine/Source/Programs/Horde/Horde.Build/Utilities/VersionedCollection.cs
@@ -2,17 +2,41 @@
 
 using System;
 using System.Collections.Generic;
+using System.Linq;
+using System.Text.Json.Serialization;
 using System.Threading.Tasks;
 using Horde.Build.Server;
-using Horde.Build.Utilities;
 using MongoDB.Bson;
+using MongoDB.Bson.IO;
 using MongoDB.Bson.Serialization;
 using MongoDB.Bson.Serialization.Attributes;
+using MongoDB.Bson.Serialization.Conventions;
 using MongoDB.Driver;
 using StackExchange.Redis;
 
-namespace Horde.Build.Tools.Impl
+namespace Horde.Build.Utilities
 {
+	/// <summary>
+	/// Exception indicating that a particular document couldn't be found
+	/// </summary>
+	public sealed class DocumentNotFoundException<TId> : Exception
+	{
+		/// <summary>
+		/// The document identifier
+		/// </summary>
+		public TId Id { get; }
+
+		/// <summary>
+		/// Constructor
+		/// </summary>
+		/// <param name="id"></param>
+		public DocumentNotFoundException(TId id)
+			: base($"Unable to find '{id}'")
+		{
+			Id = id;
+		}
+	}
+
 	/// <summary>
 	/// Base class for a versioned MongoDB document
 	/// </summary>
@@ -25,12 +49,14 @@ namespace Horde.Build.Tools.Impl
 		/// <summary>
 		/// Unique id for this document
 		/// </summary>
+		[JsonPropertyOrder(0)]
 		public TId Id { get; set; }
 
 		/// <summary>
 		/// Last time that the document was updated. This field is checked and updated as part of updates to ensure atomicity.
 		/// </summary>
 		[BsonElement("_u")]
+		[JsonPropertyOrder(1000)]
 		public DateTime LastUpdateTime { get; set; }
 
 		/// <summary>
@@ -50,6 +76,43 @@ namespace Horde.Build.Tools.Impl
 		public abstract TLatest UpgradeToLatest();
 	}
 
+	[BsonDiscriminator()]
+	class VersionedDocumentDiscriminator : IDiscriminatorConvention
+	{
+		public string ElementName => "_v";
+
+		readonly IReadOnlyDictionary<int, Type> _versionToType;
+		readonly IReadOnlyDictionary<Type, int> _typeToVersion;
+
+		public VersionedDocumentDiscriminator(IReadOnlyDictionary<int, Type> types)
+		{
+			_versionToType = types;
+			_typeToVersion = _versionToType.ToDictionary(x => x.Value, x => x.Key);
+		}
+
+		public Type GetActualType(IBsonReader bsonReader, Type nominalType)
+		{
+			BsonReaderBookmark bookmark = bsonReader.GetBookmark();
+			bsonReader.ReadStartDocument();
+			for (; ; )
+			{
+				string name = bsonReader.ReadName();
+				if (name.Equals("_v", StringComparison.Ordinal))
+				{
+					int version = bsonReader.ReadInt32();
+					bsonReader.ReturnToBookmark(bookmark);
+					return _versionToType[version];
+				}
+				bsonReader.SkipValue();
+			}
+		}
+
+		public BsonValue GetDiscriminator(Type nominalType, Type actualType)
+		{
+			return _typeToVersion[actualType];
+		}
+	}
+
 	/// <summary>
 	/// Collection of types derived from <see cref="VersionedDocument{TId, TLatest}"/>
 	/// </summary>
@@ -61,35 +124,46 @@ namespace Horde.Build.Tools.Impl
 
 		private static FilterDefinitionBuilder<VersionedDocument<TId, TLatest>> FilterBuilder { get; } = Builders<VersionedDocument<TId, TLatest>>.Filter;
 
-		private readonly IMongoCollection<VersionedDocument<TId, TLatest>> _baseCollection;
-		private readonly IMongoCollection<TLatest> _latestCollection;
 		private readonly IDatabase _redis;
 		private readonly RedisKey _baseKey;
 
+		private static readonly HashSet<Type> s_registeredTypes = new HashSet<Type>();
+
 		/// <summary>
-		/// Constructor
+		/// Collection of versioned documents
 		/// </summary>
-		/// <param name="baseCollection">The collection of documents to manage</param>
-		/// <param name="redis">Instance of the redis service</param>
-		/// <param name="baseKey">Prefix for key types</param>
-		public VersionedCollection(IMongoCollection<VersionedDocument<TId, TLatest>> baseCollection, IDatabase redis, RedisKey baseKey)
-		{
-			_baseCollection = baseCollection;
-			_latestCollection = baseCollection.OfType<TLatest>();
-			_redis = redis;
-			_baseKey = baseKey;
-		}
+		public IMongoCollection<VersionedDocument<TId, TLatest>> BaseCollection { get; }
+
+		/// <summary>
+		/// Collection of versioned documents, filtered to the latest revision
+		/// </summary>
+		public IMongoCollection<TLatest> LatestCollection { get; }
 
 		/// <summary>
 		/// Constructor
 		/// </summary>
 		/// <param name="mongoService">The database service</param>
 		/// <param name="collectionName">Name of the collection of documents to manage</param>
-		/// <param name="redis">Instance of the redis service</param>
+		/// <param name="redisService">Instance of the redis service</param>
 		/// <param name="baseKey">Prefix for key types</param>
-		public VersionedCollection(MongoService mongoService, string collectionName, IDatabase redis, RedisKey baseKey)
-			: this(mongoService.Database.GetCollection<VersionedDocument<TId, TLatest>>(collectionName), redis, baseKey)
+		/// <param name="types">Types to serialize from the collection</param>
+		public VersionedCollection(MongoService mongoService, string collectionName, RedisService redisService, RedisKey baseKey, IReadOnlyDictionary<int, Type> types)
 		{
+			_redis = redisService.Database;
+			_baseKey = baseKey;
+
+			lock (s_registeredTypes)
+			{
+				Type type = typeof(VersionedDocument<TId, TLatest>);
+				if (s_registeredTypes.Add(type))
+				{
+					VersionedDocumentDiscriminator discriminator = new VersionedDocumentDiscriminator(types);
+					BsonSerializer.RegisterDiscriminatorConvention(type, discriminator);
+				}
+			}
+
+			BaseCollection = mongoService.Database.GetCollection<VersionedDocument<TId, TLatest>>(collectionName);
+			LatestCollection = BaseCollection.OfType<TLatest>();
 		}
 
 		private RedisKey GetDocKey(TId id) => _baseKey.Append(id.ToString());
@@ -118,9 +192,9 @@ namespace Horde.Build.Tools.Impl
 		/// <returns>True if the document was added, false if it already exists</returns>
 		public async Task<bool> AddAsync(TLatest doc)
 		{
-			doc.LastUpdateTime = DateTime.UtcNow;
+			doc.LastUpdateTime = MongoExtensions.RoundToBsonDateTime(DateTime.UtcNow);
 
-			if (!await _latestCollection.InsertOneIgnoreDuplicatesAsync(doc))
+			if (!await BaseCollection.InsertOneIgnoreDuplicatesAsync(doc))
 			{
 				return false;
 			}
@@ -132,10 +206,9 @@ namespace Horde.Build.Tools.Impl
 		private void AddCachedValue(RedisKey docKey, TLatest doc)
 		{
 			KeyValuePair<RedisKey, RedisValue>[] pairs = new KeyValuePair<RedisKey, RedisValue>[2];
-			pairs[0] = new KeyValuePair<RedisKey, RedisValue>(docKey, doc.ToBson());
+			pairs[0] = new KeyValuePair<RedisKey, RedisValue>(docKey, doc.ToBson(typeof(VersionedDocument<TId, TLatest>)));
 			pairs[1] = new KeyValuePair<RedisKey, RedisValue>(docKey.Append(s_timeSuffix), doc.LastUpdateTime.Ticks);
-
-			_ = _redis.StringSetAsync(pairs, When.NotExists, CommandFlags.FireAndForget);
+			_redis.StringSetAsync(pairs, When.NotExists, flags: CommandFlags.FireAndForget);
 		}
 
 		private async ValueTask<VersionedDocument<TId, TLatest>?> GetCachedValueAsync(RedisKey docKey)
@@ -161,7 +234,7 @@ namespace Horde.Build.Tools.Impl
 			transaction.AddCondition(Condition.StringEqual(docKey.Append(s_timeSuffix), prevDoc.LastUpdateTime.Ticks));
 
 			KeyValuePair<RedisKey, RedisValue>[] pairs = new KeyValuePair<RedisKey, RedisValue>[2];
-			pairs[0] = new KeyValuePair<RedisKey, RedisValue>(docKey, doc.ToBson());
+			pairs[0] = new KeyValuePair<RedisKey, RedisValue>(docKey, doc.ToBson(typeof(VersionedDocument<TId, TLatest>)));
 			pairs[1] = new KeyValuePair<RedisKey, RedisValue>(docKey.Append(s_timeSuffix), doc.LastUpdateTime.Ticks);
 
 			_ = transaction.StringSetAsync(pairs);
@@ -195,7 +268,7 @@ namespace Horde.Build.Tools.Impl
 				if (cachedDoc == null)
 				{
 					// Read the value from the database
-					VersionedDocument<TId, TLatest>? doc = await _baseCollection.Find(FilterBuilder.Eq(x => x.Id, id)).FirstOrDefaultAsync();
+					VersionedDocument<TId, TLatest>? doc = await BaseCollection.Find(FilterBuilder.Eq(x => x.Id, id)).FirstOrDefaultAsync();
 					if (doc == null)
 					{
 						return null;
@@ -267,7 +340,7 @@ namespace Horde.Build.Tools.Impl
 		{
 			update = update.Set(x => x.LastUpdateTime, new DateTime(Math.Max(doc.LastUpdateTime.Ticks + 1, DateTime.UtcNow.Ticks)));
 
-			TLatest? newDoc = await _latestCollection.FindOneAndUpdateAsync(GetFilter(doc), update, new FindOneAndUpdateOptions<TLatest> { ReturnDocument = ReturnDocument.After });
+			TLatest? newDoc = await LatestCollection.FindOneAndUpdateAsync(GetFilter(doc), update, new FindOneAndUpdateOptions<TLatest> { ReturnDocument = ReturnDocument.After });
 			if (newDoc != null)
 			{
 				await UpdateCachedValueAsync(GetDocKey(doc.Id), doc, newDoc);
@@ -290,7 +363,7 @@ namespace Horde.Build.Tools.Impl
 
 			newDoc.LastUpdateTime = new DateTime(Math.Max(oldDoc.LastUpdateTime.Ticks + 1, DateTime.UtcNow.Ticks));
 
-			ReplaceOneResult result = await _baseCollection.ReplaceOneAsync(GetFilter(oldDoc), newDoc);
+			ReplaceOneResult result = await BaseCollection.ReplaceOneAsync(GetFilter(oldDoc), newDoc);
 			if (result.ModifiedCount == 0)
 			{
 				return false;
@@ -309,7 +382,7 @@ namespace Horde.Build.Tools.Impl
 		{
 			RedisKey docKey = GetDocKey(doc.Id);
 
-			DeleteResult result = await _baseCollection.DeleteOneAsync(GetFilter((VersionedDocument<TId, TLatest>)doc));
+			DeleteResult result = await BaseCollection.DeleteOneAsync(GetFilter((VersionedDocument<TId, TLatest>)doc));
 			if (result.DeletedCount > 0)
 			{
 				await DeleteCachedValueAsync(docKey);
@@ -328,7 +401,7 @@ namespace Horde.Build.Tools.Impl
 		{
 			RedisKey docKey = GetDocKey(id);
 
-			DeleteResult result = await _baseCollection.DeleteOneAsync(GetFilter(id));
+			DeleteResult result = await BaseCollection.DeleteOneAsync(GetFilter(id));
 			await DeleteCachedValueAsync(docKey);
 
 			return result.DeletedCount > 0;
diff --git a/Engine/Source/Programs/Shared/EpicGames.Redis/RedisString.cs b/Engine/Source/Programs/Shared/EpicGames.Redis/RedisString.cs
--- a/Engine/Source/Programs/Shared/EpicGames.Redis/RedisString.cs
+++ b/Engine/Source/Programs/Shared/EpicGames.Redis/RedisString.cs
@@ -37,7 +37,7 @@ namespace EpicGames.Redis
 		}
 
 		/// <inheritdoc cref="IDatabaseAsync.StringSetAsync(RedisKey, RedisValue, TimeSpan?, When, CommandFlags)"/>
-		public Task SetAsync(TElement value, TimeSpan? expiry = null, When when = When.Always, CommandFlags flags = CommandFlags.None)
+		public Task<bool> SetAsync(TElement value, TimeSpan? expiry = null, When when = When.Always, CommandFlags flags = CommandFlags.None)
 		{
 			return Database.StringSetAsync(Key, RedisSerializer.Serialize(value), expiry, when, flags);
 		}
@@ -48,10 +48,16 @@ namespace EpicGames.Redis
 	/// </summary>
 	public static class RedisStringExtensions
 	{
+		/// <inheritdoc cref="IDatabaseAsync.KeyDeleteAsync(RedisKey, CommandFlags)"/>
+		public static Task<bool> DeleteAsync<TElement>(this RedisString<TElement> str, CommandFlags flags = CommandFlags.None)
+		{
+			return str.Database.KeyDeleteAsync(str.Key, flags);
+		}
+
 		/// <inheritdoc cref="IDatabaseAsync.StringGetAsync(RedisKey, CommandFlags)"/>
-		public static async Task<TElement?> GetAsync<TElement>(this RedisString<TElement> @string, CommandFlags flags = CommandFlags.None) where TElement : class
+		public static async Task<TElement?> GetAsync<TElement>(this RedisString<TElement> str, CommandFlags flags = CommandFlags.None) where TElement : class
 		{
-			RedisValue value = await @string.Database.StringGetAsync(@string.Key, flags);
+			RedisValue value = await str.Database.StringGetAsync(str.Key, flags);
 			if (value.IsNullOrEmpty)
 			{
 				return null;
@@ -60,9 +66,9 @@ namespace EpicGames.Redis
 		}
 
 		/// <inheritdoc cref="IDatabaseAsync.StringGetAsync(RedisKey, CommandFlags)"/>
-		public static async Task<TElement?> GetValueAsync<TElement>(this RedisString<TElement> @string, CommandFlags flags = CommandFlags.None) where TElement : struct
+		public static async Task<TElement?> GetValueAsync<TElement>(this RedisString<TElement> str, CommandFlags flags = CommandFlags.None) where TElement : struct
 		{
-			RedisValue value = await @string.Database.StringGetAsync(@string.Key, flags);
+			RedisValue value = await str.Database.StringGetAsync(str.Key, flags);
 			if (value.IsNullOrEmpty)
 			{
 				return default(TElement);
@@ -71,35 +77,35 @@ namespace EpicGames.Redis
 		}
 
 		/// <inheritdoc cref="IDatabaseAsync.StringDecrementAsync(RedisKey, Int64, CommandFlags)"/>
-		public static Task<long> DecrementAsync(this RedisString<long> @string, long value = 1L, CommandFlags flags = CommandFlags.None)
+		public static Task<long> DecrementAsync(this RedisString<long> str, long value = 1L, CommandFlags flags = CommandFlags.None)
 		{
-			return @string.Database.StringDecrementAsync(@string.Key, value, flags);
+			return str.Database.StringDecrementAsync(str.Key, value, flags);
 		}
 
 		/// <inheritdoc cref="IDatabaseAsync.StringDecrementAsync(RedisKey, Double, CommandFlags)"/>
-		public static Task<double> DecrementAsync(this RedisString<double> @string, double value = 1.0, CommandFlags flags = CommandFlags.None)
+		public static Task<double> DecrementAsync(this RedisString<double> str, double value = 1.0, CommandFlags flags = CommandFlags.None)
 		{
-			return @string.Database.StringDecrementAsync(@string.Key, value, flags);
+			return str.Database.StringDecrementAsync(str.Key, value, flags);
 		}
 
 		/// <inheritdoc cref="IDatabaseAsync.StringDecrementAsync(RedisKey, Double, CommandFlags)"/>
-		public static Task<long> IncrementAsync(this RedisString<long> @string, long value = 1L, CommandFlags flags = CommandFlags.None)
+		public static Task<long> IncrementAsync(this RedisString<long> str, long value = 1L, CommandFlags flags = CommandFlags.None)
 		{
-			return @string.Database.StringIncrementAsync(@string.Key, value, flags);
+			return str.Database.StringIncrementAsync(str.Key, value, flags);
 		}
 
 		/// <inheritdoc cref="IDatabaseAsync.StringDecrementAsync(RedisKey, Double, CommandFlags)"/>
-		public static Task<double> IncrementAsync(this RedisString<double> @string, double value = 1.0, CommandFlags flags = CommandFlags.None)
+		public static Task<double> IncrementAsync(this RedisString<double> str, double value = 1.0, CommandFlags flags = CommandFlags.None)
 		{
-			return @string.Database.StringIncrementAsync(@string.Key, value, flags);
+			return str.Database.StringIncrementAsync(str.Key, value, flags);
 		}
 
 		/// <summary>
 		/// Creates a version of this string which modifies a transaction rather than the direct DB
 		/// </summary>
-		public static RedisString<TElement> With<TElement>(this ITransaction transaction, RedisString<TElement> set)
+		public static RedisString<TElement> With<TElement>(this ITransaction transaction, RedisString<TElement> str)
 		{
-			return new RedisString<TElement>(transaction, set.Key);
+			return new RedisString<TElement>(transaction, str.Key);
 		}
 	}
 }
diff --git a/Engine/Source/Programs/Shared/EpicGames.Serialization/CbConverter.cs b/Engine/Source/Programs/Shared/EpicGames.Serialization/CbConverter.cs
--- a/Engine/Source/Programs/Shared/EpicGames.Serialization/CbConverter.cs
+++ b/Engine/Source/Programs/Shared/EpicGames.Serialization/CbConverter.cs
@@ -261,11 +261,12 @@ namespace EpicGames.Serialization
 		public static List<CbConverterFactory> ConverterFactories = new List<CbConverterFactory>
 		{
 			new CbClassConverterFactory(),
+			new CbStringConverterFactory(),
 			new CbEnumConverterFactory(),
 			new CbListConverterFactory(),
 			new CbArrayConverterFactory(),
 			new CbDictionaryConverterFactory(),
-			new CbNullableConverterFactory()
+			new CbNullableConverterFactory(),
 		};
 
 		/// <summary>
diff --git a/Engine/Source/Programs/Shared/EpicGames.Serialization/CbSerializer.cs b/Engine/Source/Programs/Shared/EpicGames.Serialization/CbSerializer.cs
--- a/Engine/Source/Programs/Shared/EpicGames.Serialization/CbSerializer.cs
+++ b/Engine/Source/Programs/Shared/EpicGames.Serialization/CbSerializer.cs
@@ -133,6 +133,19 @@ namespace EpicGames.Serialization
 			return writer.ToObject();
 		}
 
+		/// <summary>
+		/// Serialize an object
+		/// </summary>
+		/// <typeparam name="T"></typeparam>
+		/// <param name="value"></param>
+		/// <returns></returns>
+		public static byte[] SerializeToByteArray<T>(T value)
+		{
+			CbWriter writer = new CbWriter();
+			CbConverter.GetConverter<T>().Write(writer, value);
+			return writer.ToByteArray();
+		}
+
 		/// <summary>
 		/// Serialize a property to a given writer
 		/// </summary>
