Horde: Fix tests.

diff --git a/Engine/Source/Programs/Horde/Horde.Build.Tests/ToolTests.cs b/Engine/Source/Programs/Horde/Horde.Build.Tests/ToolTests.cs
--- a/Engine/Source/Programs/Horde/Horde.Build.Tests/ToolTests.cs
+++ b/Engine/Source/Programs/Horde/Horde.Build.Tests/ToolTests.cs
@@ -1,6 +1,7 @@
 // Copyright Epic Games, Inc. All Rights Reserved.
 
 using System;
+using System.Collections.Generic;
 using System.IO;
 using System.Text;
 using System.Threading.Tasks;
@@ -22,7 +23,11 @@ namespace Horde.Build.Tests
         public async Task AddTool()
         {
 			ToolCollection collection = ServiceProvider.GetRequiredService<ToolCollection>();
-			Tool tool = await collection.AddOrUpdateAsync(new ToolOptions(new ToolId("ugs")) { Name = "UnrealGameSync", Description = "Tool for syncing content from source control" });
+
+			ToolId toolId = new ToolId("ugs");
+			await collection.ConfigureAsync(new List<ToolOptions> { new ToolOptions(toolId) { Name = "UnrealGameSync", Description = "Tool for syncing content from source control" } });
+
+			Tool tool = Deref(await collection.GetAsync(toolId));
 			Assert.AreEqual(tool.Id, new ToolId("ugs"));
 			Assert.AreEqual(tool.Name, "UnrealGameSync");
 			Assert.AreEqual(tool.Description, "Tool for syncing content from source control");
@@ -39,7 +44,10 @@ namespace Horde.Build.Tests
 		{
 			ToolCollection collection = ServiceProvider.GetRequiredService<ToolCollection>();
 
-			Tool tool = await collection.AddOrUpdateAsync(new ToolOptions(new ToolId("ugs")) { Name = "UnrealGameSync", Description = "Tool for syncing content from source control" });
+			ToolId toolId = new ToolId("ugs");
+			await collection.ConfigureAsync(new List<ToolOptions> { new ToolOptions(toolId) { Name = "UnrealGameSync", Description = "Tool for syncing content from source control" } });
+
+			Tool tool = Deref(await collection.GetAsync(toolId));
 			Assert.AreEqual(new ToolId("ugs"), tool.Id);
 			Assert.AreEqual("UnrealGameSync", tool.Name);
 			Assert.AreEqual("Tool for syncing content from source control", tool.Description);
@@ -65,14 +73,14 @@ namespace Horde.Build.Tests
 			Assert.IsNull(tool.Deployments[0].StartedAt);
 
 			// Start the deployment
-			tool = Deref(await collection.UpdateDeploymentAsync(tool, deploymentId, ToolDeploymentState.Resume));
+			tool = Deref(await collection.UpdateDeploymentAsync(tool, deploymentId, ToolDeploymentState.Active));
 			Assert.AreEqual(1, tool.Deployments.Count);
 			Assert.IsNotNull(tool.Deployments[0].StartedAt);
 			Assert.IsTrue(Math.Abs((tool.Deployments[0].StartedAt!.Value - clock.UtcNow).TotalSeconds) < 1.0);
 
 			// Check it updates
 			await clock.AdvanceAsync(TimeSpan.FromMinutes(2.5));
-			tool = Deref(await collection.UpdateDeploymentAsync(tool, deploymentId, ToolDeploymentState.Pause));
+			tool = Deref(await collection.UpdateDeploymentAsync(tool, deploymentId, ToolDeploymentState.Paused));
 			Assert.AreEqual(1, tool.Deployments.Count);
 			Assert.IsNull(tool.Deployments[0].StartedAt);
 			Assert.IsTrue((tool.Deployments[0].Progress - 0.5) < 0.1);
diff --git a/Engine/Source/Programs/Horde/Horde.Build/Tools/ToolCollection.cs b/Engine/Source/Programs/Horde/Horde.Build/Tools/ToolCollection.cs
--- a/Engine/Source/Programs/Horde/Horde.Build/Tools/ToolCollection.cs
+++ b/Engine/Source/Programs/Horde/Horde.Build/Tools/ToolCollection.cs
@@ -375,7 +375,7 @@ namespace Horde.Build.Tools
 					}
 					else
 					{
-						return await _tools.UpdateAsync(tool, Builders<Tool>.Update.Set(x => x.Deployments[idx].Progress, deployment.GetProgressValue(_clock.UtcNow)).Set(x => x.Deployments[idx].StartedAt, null));
+						return await _tools.UpdateAsync(tool, Builders<Tool>.Update.Set(x => x.Deployments[idx].BaseProgress, deployment.GetProgressValue(_clock.UtcNow)).Set(x => x.Deployments[idx].StartedAt, null));
 					}
 
 				case ToolDeploymentState.Active:
