Horde Storage - Blob Index Consistency checks can now read ahead (do multiple checks in parallell). Also some tweaks to make the logging / tracing better.

diff --git a/Engine/Source/Programs/Horde/HordeStorage/Horde.Storage/Implementation/Blob/BlobIndexConsistencyCheckService.cs b/Engine/Source/Programs/Horde/HordeStorage/Horde.Storage/Implementation/Blob/BlobIndexConsistencyCheckService.cs
--- a/Engine/Source/Programs/Horde/HordeStorage/Horde.Storage/Implementation/Blob/BlobIndexConsistencyCheckService.cs
+++ b/Engine/Source/Programs/Horde/HordeStorage/Horde.Storage/Implementation/Blob/BlobIndexConsistencyCheckService.cs
@@ -86,117 +86,121 @@ namespace Horde.Storage.Implementation
             ulong countOfBlobsChecked = 0;
             ulong countOfIncorrectBlobsFound = 0;
             string currentRegion = _jupiterSettings.CurrentValue.CurrentSite;
-            await foreach (IBlobIndex.BlobInfo blobInfo in _blobIndex.GetAllBlobs())
-            {
-                Interlocked.Increment(ref countOfBlobsChecked);
-
-                if (countOfBlobsChecked % 100 == 0)
-                    _logger.Information("Consistency check running on blob index, count of blobs processed so far: {CountOfBlobs}", countOfBlobsChecked);
+            await Parallel.ForEachAsync(_blobIndex.GetAllBlobs(), new ParallelOptions
+                {
+                    MaxDegreeOfParallelism = _settings.CurrentValue.BlobIndexMaxParallelOperations,
+                },
+                async (blobInfo, token) =>
+                {
+                    Interlocked.Increment(ref countOfBlobsChecked);
 
+                    if (countOfBlobsChecked % 100 == 0)
+                        _logger.Information("Consistency check running on blob index, count of blobs processed so far: {CountOfBlobs}", countOfBlobsChecked);
 
-                // skip namespace that we shouldn't consistency check, e.g. the legacy namespaces
-                if (!NamespaceShouldBeCheckedForConsistency(blobInfo.Namespace))
-                    continue;
-                
-                using IScope scope = Tracer.Instance.StartActive("consistency_check.blob_index");
-                scope.Span.ResourceName = $"{blobInfo.Namespace}.{blobInfo.BlobIdentifier}";
+                    // skip namespace that we shouldn't consistency check, e.g. the legacy namespaces
+                    if (!NamespaceShouldBeCheckedForConsistency(blobInfo.Namespace))
+                        return;
+                    
+                    using IScope scope = Tracer.Instance.StartActive("consistency_check.blob_index");
+                    scope.Span.ResourceName = $"{blobInfo.Namespace}.{blobInfo.BlobIdentifier}";
 
-                bool issueFound = false;
-                bool deleted = false;
+                    bool issueFound = false;
+                    bool deleted = false;
 
-                try
-                {
-                    if (!blobInfo.Regions.Any())
+                    try
                     {
-                        Interlocked.Increment(ref countOfIncorrectBlobsFound);
-                        _logger.Warning("Blob {Blob} in namespace {Namespace} is not tracked to exist in any region", blobInfo.BlobIdentifier, blobInfo.Namespace);
-                        issueFound = true;
-
-                        if (await _blobService.ExistsInRootStore(blobInfo.Namespace, blobInfo.BlobIdentifier))
+                        if (!blobInfo.Regions.Any())
                         {
-                            _logger.Warning("Blob {Blob} in namespace {Namespace} was found to exist in our blob store so re-adding it to the index", blobInfo.BlobIdentifier, blobInfo.Namespace);
+                            Interlocked.Increment(ref countOfIncorrectBlobsFound);
+                            _logger.Warning("Blob {Blob} in namespace {Namespace} is not tracked to exist in any region", blobInfo.BlobIdentifier, blobInfo.Namespace);
+                            issueFound = true;
 
-                            // we did have it in our blob store so we adjust the index
-                            await _blobIndex.AddBlobToIndex(blobInfo.Namespace, blobInfo.BlobIdentifier);
-                            continue;
-                        }
-                        else
-                        {
-                            if (_settings.CurrentValue.AllowDeletesInBlobIndex)
+                            if (await _blobService.ExistsInRootStore(blobInfo.Namespace, blobInfo.BlobIdentifier))
+                            {
+                                _logger.Warning("Blob {Blob} in namespace {Namespace} was found to exist in our blob store so re-adding it to the index", blobInfo.BlobIdentifier, blobInfo.Namespace);
+
+                                // we did have it in our blob store so we adjust the index
+                                await _blobIndex.AddBlobToIndex(blobInfo.Namespace, blobInfo.BlobIdentifier);
+                            }
+                            else
                             {
-                                // this blob doesn't exist anywhere so we just cleanup the blob index
-                                _logger.Warning("Blob {Blob} in namespace {Namespace} was removed from the blob index as it didnt exist anywhere", blobInfo.BlobIdentifier, blobInfo.Namespace);
+                                if (_settings.CurrentValue.AllowDeletesInBlobIndex)
+                                {
+                                    // this blob doesn't exist anywhere so we just cleanup the blob index
+                                    _logger.Warning("Blob {Blob} in namespace {Namespace} was removed from the blob index as it didnt exist anywhere", blobInfo.BlobIdentifier, blobInfo.Namespace);
 
-                                await _blobIndex.RemoveBlobFromIndex(blobInfo.Namespace, blobInfo.BlobIdentifier);
-                                deleted = true;
+                                    await _blobIndex.RemoveBlobFromIndex(blobInfo.Namespace, blobInfo.BlobIdentifier);
+                                    deleted = true;
+                                }
                             }
                         }
-                    }
-                    
-                    if (blobInfo.Regions.Contains(currentRegion))
-                    {
-                        if (!await _blobService.ExistsInRootStore(blobInfo.Namespace, blobInfo.BlobIdentifier))
+                        else if (blobInfo.Regions.Contains(currentRegion))
                         {
-                            Interlocked.Increment(ref countOfIncorrectBlobsFound);
-                            issueFound = true;
-                            
-                            if (blobInfo.Regions.Count > 1)
+                            if (!await _blobService.ExistsInRootStore(blobInfo.Namespace, blobInfo.BlobIdentifier))
                             {
-                                _logger.Warning("Blob {Blob} in namespace {Namespace} did not exist in root store but is tracked as doing so in the blob index. Attempting to replicate it.", blobInfo.BlobIdentifier, blobInfo.Namespace);
-
-                                try
+                                Interlocked.Increment(ref countOfIncorrectBlobsFound);
+                                issueFound = true;
+                                
+                                if (blobInfo.Regions.Count > 1)
                                 {
-                                    BlobContents _ = await _blobService.ReplicateObject(blobInfo.Namespace, blobInfo.BlobIdentifier, force: true);
+                                    _logger.Warning("Blob {Blob} in namespace {Namespace} did not exist in root store but is tracked as doing so in the blob index. Attempting to replicate it.", blobInfo.BlobIdentifier, blobInfo.Namespace);
+
+                                    try
+                                    {
+                                        BlobContents _ = await _blobService.ReplicateObject(blobInfo.Namespace, blobInfo.BlobIdentifier, force: true);
+                                    }
+                                    catch (BlobReplicationException e)
+                                    {
+                                        // we update the blob index to accurately reflect that we do not have the blob, this is not good though as it means a upload that we thought happened now lacks content
+                                        if (_settings.CurrentValue.AllowDeletesInBlobIndex)
+                                        {
+                                            _logger.Warning("Updating blob index to remove Blob {Blob} in namespace {Namespace} as we failed to repair it.", blobInfo.BlobIdentifier, blobInfo.Namespace);
+                                            await _blobIndex.RemoveBlobFromRegion(blobInfo.Namespace, blobInfo.BlobIdentifier);
+                                            deleted = true;
+                                        }
+                                        else
+                                        {
+                                            _logger.Error(e, "Failed to replicate Blob {Blob} in namespace {Namespace}. Unable to repair the blob index", blobInfo.BlobIdentifier, blobInfo.Namespace);
+                                        }
+                                    }
                                 }
-                                catch (BlobReplicationException e)
+                                else
                                 {
-                                    _logger.Error(e, "Failed to replicate Blob {Blob} in namespace {Namespace}. Unable to repair the blob index", blobInfo.BlobIdentifier, blobInfo.Namespace);
+                                    // if the blob only exists in the current region there is no point in attempting to replicate it
+                                    _logger.Warning("Blob {Blob} in namespace {Namespace} did not exist in root store but is tracked as doing so in the blob index. Does not exist anywhere else so unable to replicate it.", blobInfo.BlobIdentifier, blobInfo.Namespace);
 
-                                    // we update the blob index to accurately reflect that we do not have the blob, this is not good though as it means a upload that we thought happened now lacks content
                                     if (_settings.CurrentValue.AllowDeletesInBlobIndex)
                                     {
-                                        _logger.Warning("Updating blob index to remove Blob {Blob} in namespace {Namespace} as we failed to repair it.", blobInfo.BlobIdentifier, blobInfo.Namespace);
+                                        // this blob can not be repaired so we just delete it from the blob index
+                                        _logger.Warning("Blob {Blob} in namespace {Namespace} can not be repaired so removing existence from current region.", blobInfo.BlobIdentifier, blobInfo.Namespace);
+
                                         await _blobIndex.RemoveBlobFromRegion(blobInfo.Namespace, blobInfo.BlobIdentifier);
                                         deleted = true;
                                     }
                                 }
                             }
-                            else
-                            {
-                                // if the blob only exists in the current region there is no point in attempting to replicate it
-                                _logger.Warning("Blob {Blob} in namespace {Namespace} did not exist in root store but is tracked as doing so in the blob index. Does not exist anywhere else so unable to replicate it.", blobInfo.BlobIdentifier, blobInfo.Namespace);
-
-                                if (_settings.CurrentValue.AllowDeletesInBlobIndex)
-                                {
-                                    // this blob can not be repaired so we just delete it from the blob index
-                                    _logger.Warning("Blob {Blob} in namespace {Namespace} can not be repaired so removing existence from current region.", blobInfo.BlobIdentifier, blobInfo.Namespace);
-
-                                    await _blobIndex.RemoveBlobFromRegion(blobInfo.Namespace, blobInfo.BlobIdentifier);
-                                    deleted = true;
-                                }
-                            }
                         }
                     }
-                }
-                catch (UnknownNamespaceException)
-                {
-                    if (_settings.CurrentValue.AllowDeletesInBlobIndex)
+                    catch (UnknownNamespaceException)
                     {
-                        _logger.Warning("Blob {Blob} in namespace {Namespace} is of a unknown namespace, removing.", blobInfo.BlobIdentifier, blobInfo.Namespace);
+                        if (_settings.CurrentValue.AllowDeletesInBlobIndex)
+                        {
+                            _logger.Warning("Blob {Blob} in namespace {Namespace} is of a unknown namespace, removing.", blobInfo.BlobIdentifier, blobInfo.Namespace);
 
-                        // for entries that are of a unknown namespace we simply remove them
-                        await _blobIndex.RemoveBlobFromIndex(blobInfo.Namespace, blobInfo.BlobIdentifier);
-                        deleted = true;
+                            // for entries that are of a unknown namespace we simply remove them
+                            await _blobIndex.RemoveBlobFromIndex(blobInfo.Namespace, blobInfo.BlobIdentifier);
+                            deleted = true;
+                        }
+                    }
+                    catch (Exception e)
+                    {
+                        _logger.Error(e, "Exception when doing blob index consistency check for {Blob} in namespace {Namespace}", blobInfo.BlobIdentifier, blobInfo.Namespace);
                     }
-                }
-                catch (Exception e)
-                {
-                    _logger.Error(e, "Exception when doing blob index consistency check for {Blob} in namespace {Namespace}", blobInfo.BlobIdentifier, blobInfo.Namespace);
-                }
 
-                scope.Span.SetTag("issueFound", issueFound.ToString());
-                scope.Span.SetTag("deleted", deleted.ToString());
-            }
+                    scope.Span.SetTag("issueFound", issueFound.ToString());
+                    scope.Span.SetTag("deleted", deleted.ToString());
+                }
+            );
 
             _logger.Information("Blob Index Consistency check finished, found {CountOfIncorrectBlobs} incorrect blobs. Processed {CountOfBlobs} blobs.", countOfIncorrectBlobsFound, countOfBlobsChecked);
         }
diff --git a/Engine/Source/Programs/Horde/HordeStorage/Horde.Storage/Settings.cs b/Engine/Source/Programs/Horde/HordeStorage/Horde.Storage/Settings.cs
--- a/Engine/Source/Programs/Horde/HordeStorage/Horde.Storage/Settings.cs
+++ b/Engine/Source/Programs/Horde/HordeStorage/Horde.Storage/Settings.cs
@@ -292,6 +292,7 @@ namespace Horde.Storage
         public bool EnableBlobStoreChecks { get; set; } = false;
         public bool EnableBlobIndexChecks { get; set; } = false;
         public double ConsistencyCheckPollFrequencySeconds { get; set; } = TimeSpan.FromHours(2).TotalSeconds;
+        public int BlobIndexMaxParallelOperations { get; set; } = 4;
         public bool AllowDeletesInBlobIndex { get; set; } = false;
     }
 }
