Horde Storage - Fixed issue were attachments were not found if not part of the root object (attachment scanning didnt happen for embedded objects)

diff --git a/Engine/Source/Programs/Horde/HordeStorage/Horde.Storage/Controllers/ReferencesController.cs b/Engine/Source/Programs/Horde/HordeStorage/Horde.Storage/Controllers/ReferencesController.cs
--- a/Engine/Source/Programs/Horde/HordeStorage/Horde.Storage/Controllers/ReferencesController.cs
+++ b/Engine/Source/Programs/Horde/HordeStorage/Horde.Storage/Controllers/ReferencesController.cs
@@ -156,8 +156,24 @@ namespace Horde.Storage.Controllers
                     {
                         byte[] blobMemory = await blob.Stream.ToByteArray();
                         CbObject cb = new CbObject(blobMemory);
-                        int countOfAttachmentFields = cb.Count(field => field.IsAttachment());
-                        CbField? binaryAttachmentField = cb.FirstOrDefault(field => field.IsBinaryAttachment());
+
+                        (int,CbField?) CountFields(CbObject o)
+                        {
+                            int count = 0;
+                            CbField? foundField = null;
+                            cb.IterateAttachments(field =>
+                            {
+                                ++count;
+                                if (field.IsBinaryAttachment())
+                                    foundField = field;
+                            });
+
+                            return (count, foundField);
+                        }
+
+                        // breaking lambda call into private method to workaround incorrect triggering of CA1508 - [REDACTED_PROFILE_URL]/roslyn-analyzers/issues/5254
+                        (int countOfAttachmentFields,CbField? binaryAttachmentField) = CountFields(cb);
+
                         if (countOfAttachmentFields == 1 && binaryAttachmentField != null)
                         {
                             // there is a single attachment field and that is of the binary attachment type, fetch that attachment and return it instead of the compact binary
@@ -189,12 +205,32 @@ namespace Horde.Storage.Controllers
                     }
                     case CustomMediaTypeNames.JupiterInlinedPayload:
                     {
-
                         byte[] blobMemory = await blob.Stream.ToByteArray();
                         CbObject cb = new CbObject(blobMemory);
-                        List<CbField> compactBinaryFields = cb.ToList();
-                        int countOfBinaryAttachmentFields = compactBinaryFields.Count(field => field.IsBinaryAttachment());
-                        int countOfAttachmentFields = compactBinaryFields.Count(field => field.IsAttachment());
+
+                        static (int, int) CountFields(CbObject o)
+                        {
+                            int countOfBinaryAttachmentFields = 0;
+                            int countOfAttachmentFields = 0;
+
+                            o.IterateAttachments(field =>
+                            {
+                                if (field.IsBinaryAttachment())
+                                {
+                                    ++countOfBinaryAttachmentFields;
+                                }
+
+                                if (field.IsAttachment())
+                                {
+                                    ++countOfAttachmentFields;
+                                }
+                            });
+
+                            return (countOfBinaryAttachmentFields, countOfAttachmentFields);
+                        }
+                        // breaking lambda call into private method to workaround incorrect triggering of CA1508 - [REDACTED_PROFILE_URL]/roslyn-analyzers/issues/5254
+                        (int countOfAttachmentFields, int countOfBinaryAttachmentFields) = CountFields(cb);
+
                         // if the object consists of a single attachment field we return this attachment field instead
                         if (countOfBinaryAttachmentFields == 1 && countOfAttachmentFields == 1)
                         {
diff --git a/Engine/Source/Programs/Horde/HordeStorage/Horde.Storage/Implementation/IReferenceResolver.cs b/Engine/Source/Programs/Horde/HordeStorage/Horde.Storage/Implementation/IReferenceResolver.cs
--- a/Engine/Source/Programs/Horde/HordeStorage/Horde.Storage/Implementation/IReferenceResolver.cs
+++ b/Engine/Source/Programs/Horde/HordeStorage/Horde.Storage/Implementation/IReferenceResolver.cs
@@ -44,32 +44,29 @@ namespace Horde.Storage.Implementation
             {
                 if (objectsToVisit.TryDequeue(out CbObject? parent))
                 {
-                    // enumerate all fields in the compact binary and start to resolve their dependencies
-                    foreach (CbField field in parent)
-                    {
-                        if (!field.IsAttachment())
-                            continue;
+                    List<BlobIdentifier> blobIdentifiers = new List<BlobIdentifier>();
 
+                    parent.IterateAttachments(field =>
+                    {
                         IoHash attachmentHash = field.AsAttachment();
-                        if (field.HasError())
-                            continue;
 
                         BlobIdentifier blobIdentifier = BlobIdentifier.FromIoHash(attachmentHash);
                         if (field.IsBinaryAttachment())
                         {
                             pendingContentIdResolves.Add(ResolveContentId(ns, ContentId.FromIoHash(attachmentHash)));
                         }
-                        else
+                        else if (field.IsObjectAttachment())
                         {
-                            // the identifier as is points to a blob
-                            yield return blobIdentifier;
+                            blobIdentifiers.Add(blobIdentifier);
+                            pendingCompactBinaryAttachments.Add(ParseCompactBinaryAttachment(ns, blobIdentifier));
                         }
+                        else
+                            throw new NotImplementedException($"Unknown attachment type for field {field}");
+                    });
 
-                        // if its a reference to another compact binary we need to fetch it and resolve it
-                        if (!field.IsObjectAttachment())
-                            continue;
-
-                        pendingCompactBinaryAttachments.Add(ParseCompactBinaryAttachment(ns, blobIdentifier));
+                    foreach (BlobIdentifier blobIdentifier in blobIdentifiers)
+                    {
+                        yield return blobIdentifier;
                     }
                 }
 
diff --git a/Engine/Source/Programs/Horde/HordeStorage/Horde.Storage/Implementation/ObjectService.cs b/Engine/Source/Programs/Horde/HordeStorage/Horde.Storage/Implementation/ObjectService.cs
--- a/Engine/Source/Programs/Horde/HordeStorage/Horde.Storage/Implementation/ObjectService.cs
+++ b/Engine/Source/Programs/Horde/HordeStorage/Horde.Storage/Implementation/ObjectService.cs
@@ -94,7 +94,7 @@ namespace Horde.Storage.Implementation
 
         public async Task<(ContentId[], BlobIdentifier[])> Put(NamespaceId ns, BucketId bucket, IoHashKey key, BlobIdentifier blobHash, CbObject payload)
         {
-            bool hasReferences = payload.Any(field => field.IsAttachment());
+            bool hasReferences = HasAttachments(payload);
 
             // if we have no references we are always finalized, e.g. there are no referenced blobs to upload
             bool isFinalized = !hasReferences;
@@ -108,6 +108,33 @@ namespace Horde.Storage.Implementation
             return await DoFinalize(ns, bucket, key, blobHash, payload);
         }
 
+        private bool HasAttachments(CbObject payload)
+        {
+            bool FieldHasAttachments(CbField field)
+            {
+                if (field.IsObject())
+                {
+                    bool hasAttachment = HasAttachments(field.AsObject());
+                    if (hasAttachment)
+                        return true;
+                }
+
+                if (field.IsArray())
+                {
+                    foreach (CbField subField in field.AsArray())
+                    {
+                        bool hasAttachment = FieldHasAttachments(subField);
+                        if (hasAttachment)
+                            return true;
+                    }
+                }
+
+                return field.IsAttachment();
+            }
+
+            return payload.Any(FieldHasAttachments);
+        }
+
         public async Task<(ContentId[], BlobIdentifier[])> Finalize(NamespaceId ns, BucketId bucket, IoHashKey key, BlobIdentifier blobHash)
         {
             (ObjectRecord o, BlobContents? blob) = await Get(ns, bucket, key);
@@ -129,7 +156,7 @@ namespace Horde.Storage.Implementation
 
             ContentId[] missingReferences = Array.Empty<ContentId>();
             BlobIdentifier[] missingBlobs = Array.Empty<BlobIdentifier>();
-            bool hasReferences = payload.Any(field => field.IsAttachment());
+            bool hasReferences = HasAttachments(payload);
             if (hasReferences)
             {
                 using IScope _ = Tracer.Instance.StartActive("ObjectService.ResolveReferences");
diff --git a/Engine/Source/Programs/Horde/HordeStorage/Horde.Storage/Tests/Functional/Objects/ReferencesTests.cs b/Engine/Source/Programs/Horde/HordeStorage/Horde.Storage/Tests/Functional/Objects/ReferencesTests.cs
--- a/Engine/Source/Programs/Horde/HordeStorage/Horde.Storage/Tests/Functional/Objects/ReferencesTests.cs
+++ b/Engine/Source/Programs/Horde/HordeStorage/Horde.Storage/Tests/Functional/Objects/ReferencesTests.cs
@@ -994,6 +994,54 @@ namespace Horde.Storage.FunctionalTests.References
             }
         }
 
+        
+        [TestMethod]
+        public async Task PutMissingAttachmentComplex()
+        {
+            string blobContents = "This is a string that is referenced as a blob but will not be uploaded";
+            byte[] blobData = Encoding.ASCII.GetBytes(blobContents);
+            BlobIdentifier blobHash = BlobIdentifier.FromBlob(blobData);
+
+            CbWriter writer = new CbWriter();
+            writer.BeginObject();
+            writer.WriteString("name", "randomStringContent");
+            writer.BeginArray("values");
+            writer.BeginObject();
+            writer.WriteInteger("rawSize", 200);
+            writer.WriteBinaryAttachment("rawHash", blobHash.AsIoHash());
+            writer.EndObject();
+            writer.EndArray();
+            writer.EndObject();
+
+            byte[] objectData = writer.ToByteArray();
+            BlobIdentifier objectHash = BlobIdentifier.FromBlob(objectData);
+
+            IoHashKey key = IoHashKey.FromName("putContentIdMissingBlobComplex");
+
+            {
+                HttpContent requestContent = new ByteArrayContent(objectData);
+                requestContent.Headers.ContentType = new MediaTypeHeaderValue(CustomMediaTypeNames.UnrealCompactBinary);
+                requestContent.Headers.Add(CommonHeaders.HashHeaderName, objectHash.ToString());
+
+                HttpResponseMessage result = await _httpClient!.PutAsync(requestUri: $"api/v1/refs/{TestNamespace}/bucket/{key}.uecb", requestContent);
+                result.EnsureSuccessStatusCode();
+
+                {
+                    Assert.AreEqual(result!.Content.Headers.ContentType!.MediaType, CustomMediaTypeNames.UnrealCompactBinary);
+                    await using MemoryStream ms = new MemoryStream();
+                    await result.Content.CopyToAsync(ms);
+                    byte[] roundTrippedBuffer = ms.ToArray();
+                    CbObject cb = new CbObject(roundTrippedBuffer);
+                    CbField needsField = cb["needs"];
+                    Assert.AreNotEqual(CbField.Empty, needsField);
+                    List<IoHash> missingBlobs = needsField.AsArray().Select(field => field.AsHash()).ToList();
+                    Assert.AreEqual(1, missingBlobs.Count);
+
+                    Assert.AreEqual(blobHash.AsIoHash(), missingBlobs[0], "Expected to find missing hash even for attachments not directly in the root object");
+                }
+            }
+        }
+
         [TestMethod]
         public async Task PutAndFinalize()
         {
