Horde Storage - Added option to control if blob store consistency checks runs on the root store. Enalbed this on the worker (which was the old behavior) but disabled this on this on http serving pods which are now running consistency checks. Results in checks being run on the http serving pods local caches (which cant be checked from the workers as they are seperate pods)

diff --git a/Engine/Source/Programs/Horde/HordeStorage/Horde.Storage/Implementation/Blob/BlobStoreConsistencyCheckService.cs b/Engine/Source/Programs/Horde/HordeStorage/Horde.Storage/Implementation/Blob/BlobStoreConsistencyCheckService.cs
--- a/Engine/Source/Programs/Horde/HordeStorage/Horde.Storage/Implementation/Blob/BlobStoreConsistencyCheckService.cs
+++ b/Engine/Source/Programs/Horde/HordeStorage/Horde.Storage/Implementation/Blob/BlobStoreConsistencyCheckService.cs
@@ -78,6 +78,9 @@ namespace Horde.Storage.Implementation
 
                 bool isRootStore = blobStore is AmazonS3Store or AzureBlobStore;
 
+                if (!_settings.CurrentValue.RunBlobStoreConsistencyCheckOnRootStore && isRootStore)
+                    return;
+
                 List<NamespaceId> namespaces = await _refsStore.GetNamespaces().ToListAsync();
                 namespaces.AddRange(await _referencesStore.GetNamespaces().ToListAsync());
 
diff --git a/Engine/Source/Programs/Horde/HordeStorage/Horde.Storage/Settings.cs b/Engine/Source/Programs/Horde/HordeStorage/Horde.Storage/Settings.cs
--- a/Engine/Source/Programs/Horde/HordeStorage/Horde.Storage/Settings.cs
+++ b/Engine/Source/Programs/Horde/HordeStorage/Horde.Storage/Settings.cs
@@ -294,5 +294,6 @@ namespace Horde.Storage
         public double ConsistencyCheckPollFrequencySeconds { get; set; } = TimeSpan.FromHours(2).TotalSeconds;
         public int BlobIndexMaxParallelOperations { get; set; } = 4;
         public bool AllowDeletesInBlobIndex { get; set; } = false;
+        public bool RunBlobStoreConsistencyCheckOnRootStore { get; set; } = false;
     }
 }
