Horde.Storage: Fix race conditions on writing to blobs/updating timestamps on blobs.

- Do not replace a file if it already exists.
- When writing a new file, write it to a temporary location first, then move it into place.

diff --git a/Engine/Source/Programs/Horde/HordeStorage/Horde.Storage/Implementation/Blob/FileSystemStore.cs b/Engine/Source/Programs/Horde/HordeStorage/Horde.Storage/Implementation/Blob/FileSystemStore.cs
--- a/Engine/Source/Programs/Horde/HordeStorage/Horde.Storage/Implementation/Blob/FileSystemStore.cs
+++ b/Engine/Source/Programs/Horde/HordeStorage/Horde.Storage/Implementation/Blob/FileSystemStore.cs
@@ -55,62 +55,44 @@ namespace Horde.Storage.Implementation
             return new DirectoryInfo(Path.Combine(GetRootDir(), ns.ToString()));
         }
 
+        static string s_processSuffix = Guid.NewGuid().ToString();
+        static int s_uniqueId = 0;
+
         public async Task<BlobIdentifier> PutObject(NamespaceId ns, ReadOnlyMemory<byte> content, BlobIdentifier blobIdentifier)
         {
-            FileInfo filePath = GetFilesystemPath(ns, blobIdentifier);
-            filePath.Directory?.Create();
-
-            try
-            {
-                await using FileStream fs = new FileStream(filePath.FullName, FileMode.Create, FileAccess.Write, FileShare.Read, DefaultBufferSize, FileOptions.Asynchronous | FileOptions.SequentialScan);
-                CancellationToken cancellationToken = default(CancellationToken);
-                await fs.WriteAsync(content, cancellationToken);
-                await fs.FlushAsync(cancellationToken);
-            }
-            catch (IOException e)
-            {
-                // ignore errors about accessing the file, it could be used by another call to write it at the same time, which is completely safe as the identifier is the content of the file
-                if (e.HResult == SharingViolationErrorCode)
-                {
-                    _logger.Warning("Sharing violation attempting to write {Blob} in {Namespace}", blobIdentifier, ns);
-                }
-                else
-                    throw;
-            }
-
-            if (content.Length == 0)
-                _logger.Warning("0 byte file written as {Blob} {Namespace} {Method}", blobIdentifier, ns, "ReadOnlyMemory");
-
-            UpdateLastWriteTime(filePath.FullName, DateTime.UnixEpoch);
-            return blobIdentifier;
+			using EpicGames.Core.ReadOnlyMemoryStream stream = new EpicGames.Core.ReadOnlyMemoryStream(content);
+			return await PutObject(ns, stream, blobIdentifier);
         }
 
-
         public async Task<BlobIdentifier> PutObject(NamespaceId ns, Stream content, BlobIdentifier blobIdentifier)
         {
             FileInfo filePath = GetFilesystemPath(ns, blobIdentifier);
             filePath.Directory?.Create();
 
-            try
-            {
-                await using FileStream fs = new FileStream(filePath.FullName, FileMode.Create, FileAccess.Write, FileShare.Read, DefaultBufferSize, FileOptions.Asynchronous | FileOptions.SequentialScan);
-                CancellationToken cancellationToken = default(CancellationToken);
-                await content.CopyToAsync(fs, cancellationToken);
-                await fs.FlushAsync(cancellationToken);
-            }
-            catch (IOException e)
+            if (!filePath.Exists)
             {
-                // ignore errors about accessing the file, it could be used by another call to write it at the same time, which is completely safe as the identifier is the content of the file
-                if (e.HResult == SharingViolationErrorCode)
+                int uniqueId = Interlocked.Increment(ref s_uniqueId);
+
+                string tempFilePath = $"{filePath.FullName}.{s_processSuffix}.{uniqueId}";
+                await using (FileStream fs = new FileStream(tempFilePath, FileMode.Create, FileAccess.Write, FileShare.Read, DefaultBufferSize, FileOptions.Asynchronous | FileOptions.SequentialScan))
                 {
-                    _logger.Warning("Sharing violation attempting to write {Blob} in {Namespace}", blobIdentifier, ns);
+                    CancellationToken cancellationToken = default(CancellationToken);
+                    await content.CopyToAsync(fs, cancellationToken);
                 }
-                else
-                    throw;
+
+                try
+                {
+                    File.Move(tempFilePath, filePath.FullName, true);
+                }
+                catch (IOException) when (File.Exists(filePath.FullName))
+                {
+                }
+
+                filePath.Refresh();
+
+                if (filePath.Length == 0)
+                    _logger.Warning("0 byte file written as {Blob} {Namespace} {Method}", blobIdentifier, ns, "Stream");
             }
-            filePath.Refresh();
-            if (filePath.Length == 0)
-                _logger.Warning("0 byte file written as {Blob} {Namespace} {Method}", blobIdentifier, ns, "Stream");
 
             UpdateLastWriteTime(filePath.FullName, DateTime.UnixEpoch);
             return blobIdentifier;
@@ -118,32 +100,8 @@ namespace Horde.Storage.Implementation
 
         public async Task<BlobIdentifier> PutObject(NamespaceId ns, byte[] content, BlobIdentifier blobIdentifier)
         {
-            FileInfo filePath = GetFilesystemPath(ns, blobIdentifier);
-            filePath.Directory?.Create();
-
-            try
-            {
-                await using FileStream fs = new FileStream(filePath.FullName, FileMode.Create, FileAccess.Write, FileShare.Read, DefaultBufferSize, FileOptions.Asynchronous | FileOptions.SequentialScan);
-                CancellationToken cancellationToken = default(CancellationToken);
-                await fs.WriteAsync(content, cancellationToken);
-                await fs.FlushAsync(cancellationToken);
-            }
-            catch (IOException e)
-            {
-                // ignore errors about accessing the file, it could be used by another call to write it at the same time, which is completely safe as the identifier is the content of the file
-                if (e.HResult == SharingViolationErrorCode)
-                {
-                    _logger.Warning("Sharing violation attempting to write {Blob} in {Namespace}", blobIdentifier, ns);
-                }
-                else
-                    throw;
-            }
-
-            if (content.Length == 0)
-                _logger.Warning("0 byte file written as {Blob} {Namespace} {Method}", blobIdentifier, ns, "Array");
-
-            UpdateLastWriteTime(filePath.FullName, DateTime.UnixEpoch);
-            return blobIdentifier;
+			using MemoryStream stream = new MemoryStream(content);
+			return await PutObject(ns, stream, blobIdentifier);
         }
 
         public Task<BlobContents> GetObject(NamespaceId ns, BlobIdentifier blob)
