Horde.Storage: Add CbField to ResolvedReferencesResult.References

#rnx

diff --git a/Engine/Source/Programs/Horde/HordeStorage/Horde.Storage/Controllers/ObjectController.cs b/Engine/Source/Programs/Horde/HordeStorage/Horde.Storage/Controllers/ObjectController.cs
--- a/Engine/Source/Programs/Horde/HordeStorage/Horde.Storage/Controllers/ObjectController.cs
+++ b/Engine/Source/Programs/Horde/HordeStorage/Horde.Storage/Controllers/ObjectController.cs
@@ -296,6 +296,7 @@ namespace Horde.Storage.Controllers
             References = references;
         }
 
+        [CbField("references")]
         public BlobIdentifier[] References { get; set; }
     }
 }
