Horde.Storage: Add default constructor for ResolvedReferencesResult

#rnx

diff --git a/Engine/Source/Programs/Horde/HordeStorage/Horde.Storage/Controllers/ObjectController.cs b/Engine/Source/Programs/Horde/HordeStorage/Horde.Storage/Controllers/ObjectController.cs
--- a/Engine/Source/Programs/Horde/HordeStorage/Horde.Storage/Controllers/ObjectController.cs
+++ b/Engine/Source/Programs/Horde/HordeStorage/Horde.Storage/Controllers/ObjectController.cs
@@ -291,6 +291,11 @@ namespace Horde.Storage.Controllers
 
     public class ResolvedReferencesResult
     {
+        public ResolvedReferencesResult()
+        {
+            References = null!;
+        }
+
         public ResolvedReferencesResult(BlobIdentifier[] references)
         {
             References = references;
