Horde Storage - Made partial scylla object a error and moved scylla settings to a seperate file

diff --git a/Engine/Source/Programs/Horde/HordeStorage/Horde.Storage/Implementation/Objects/ScyllaReferencesStore.cs b/Engine/Source/Programs/Horde/HordeStorage/Horde.Storage/Implementation/Objects/ScyllaReferencesStore.cs
--- a/Engine/Source/Programs/Horde/HordeStorage/Horde.Storage/Implementation/Objects/ScyllaReferencesStore.cs
+++ b/Engine/Source/Programs/Horde/HordeStorage/Horde.Storage/Implementation/Objects/ScyllaReferencesStore.cs
@@ -2,7 +2,6 @@
 
 using System;
 using System.Collections.Generic;
-using System.ComponentModel.DataAnnotations;
 using System.Linq;
 using System.Threading.Tasks;
 using Cassandra;
@@ -75,9 +74,13 @@ namespace Horde.Storage.Implementation
             if (o == null)
                 throw new ObjectNotFoundException(ns, bucket, name);
 
-            if (o.RequiredFieldsAreMissing())
+            try
+            {
+                o.ThrowIfRequiredFieldIsMissing();
+            }
+            catch (Exception e)
             {
-                _logger.Warning("Partial object found {Namespace} {Bucket} {Name} ignoring object", ns, bucket, name);
+                _logger.Error(e, "Partial object found {Namespace} {Bucket} {Name} ignoring object", ns, bucket, name);
                 throw new ObjectNotFoundException(ns, bucket, name);
             }
 
@@ -275,53 +278,6 @@ namespace Horde.Storage.Implementation
         }
     }
 
-    public class ScyllaSettings
-    {
-        public string ConnectionString { get; set; } = "Contact Points=localhost,scylla";
-        public long InlineBlobMaxSize { get; set; } = 32 * 1024; // default to 32 kb blobs max
-        public string[] ContactPoints { get; set; } = new string[] {"localhost", "scylla"};
-        public int MaxSnapshotsPerNamespace { get; set; } = 10;
-
-        /// <summary>
-        /// Set to override the replication strategy used to create keyspace
-        /// Note that this only applies when creating the keyspace
-        /// To modify a existing keyspace see https://docs.datastax.com/en/dse/6.7/dse-admin/datastax_enterprise/operations/opsChangeKSStrategy.html
-        /// </summary>
-        public Dictionary<string, string>? KeyspaceReplicationStrategy { get; set; }
-
-        /// <summary>
-        /// Set to override the replication strategy used to create the local keyspace
-        /// Note that this only applies when creating the keyspace
-        /// To modify a existing keyspace see https://docs.datastax.com/en/dse/6.7/dse-admin/datastax_enterprise/operations/opsChangeKSStrategy.html
-        /// </summary>
-        public Dictionary<string, string>? LocalKeyspaceReplicationStrategy { get; set; }
-
-        /// <summary>
-        /// Used to configure the load balancing policy to stick to this specified datacenter
-        /// </summary>
-        [Required]
-        public string LocalDatacenterName { get; set; } = null!;
-        
-        [Required] 
-        public string LocalKeyspaceSuffix { get; set; } = null!;
-
-        /// <summary>
-        /// Max number of connections for each scylla host before switching to another host
-        /// See https://docs.datastax.com/en/developer/nodejs-driver/4.6/features/connection-pooling/
-        /// </summary>
-        public int MaxConnectionForLocalHost { get; set; } = 8192;
-
-        /// <summary>
-        /// The time for a replication log event to live in the incremental state before being deleted, assumption is that the snapshot will have processed the event within this time
-        /// </summary>
-        public TimeSpan ReplicationLogTimeToLive { get; set; } = TimeSpan.FromDays(7);
-
-        /// <summary>
-        /// Whether a Cosmos DB Cassandra layer from Azure should be used.
-        /// </summary>
-        public bool UseAzureCosmosDB { get; set; } = false;
-    }
-
     [Cassandra.Mapping.Attributes.Table("objects")]
     public class ScyllaObject
     {
@@ -363,25 +319,18 @@ namespace Horde.Storage.Implementation
         [Cassandra.Mapping.Attributes.Column("last_access_time")]
         public DateTime LastAccessTime { get; set; }
 
-        public bool RequiredFieldsAreMissing()
+        public void ThrowIfRequiredFieldIsMissing()
         {
             if (string.IsNullOrEmpty(Namespace))
-                //throw new ArgumentException("Namespace was not valid", nameof(Namespace));
-                return true;
+                throw new ArgumentException("Namespace was not valid", nameof(Namespace));
             if (string.IsNullOrEmpty(Bucket))
-                //throw new ArgumentException("Bucket was not valid", nameof(Bucket));
-                return true;
+                throw new ArgumentException("Bucket was not valid", nameof(Bucket));
             if (string.IsNullOrEmpty(Name))
-                //throw new ArgumentException("Name was not valid", nameof(Name));
-                return true;
+                throw new ArgumentException("Name was not valid", nameof(Name));
             if (PayloadHash == null)
-                //throw new ArgumentException("PayloadHash was not valid", nameof(PayloadHash));
-                return true;
+                throw new ArgumentException("PayloadHash was not valid", nameof(PayloadHash));
             if (!IsFinalized.HasValue)
-                //throw new ArgumentException("IsFinalized was not valid", nameof(IsFinalized));
-                return true;
-
-            return false;
+                throw new ArgumentException("IsFinalized was not valid", nameof(IsFinalized));
         }
     }
 
diff --git a/Engine/Source/Programs/Horde/HordeStorage/Horde.Storage/Implementation/Objects/ScyllaSettings.cs b/Engine/Source/Programs/Horde/HordeStorage/Horde.Storage/Implementation/Objects/ScyllaSettings.cs
new file mode 100644
--- /dev/null
+++ b/Engine/Source/Programs/Horde/HordeStorage/Horde.Storage/Implementation/Objects/ScyllaSettings.cs
@@ -0,0 +1,53 @@
+// Copyright Epic Games, Inc. All Rights Reserved.
+
+using System;
+using System.Collections.Generic;
+using System.ComponentModel.DataAnnotations;
+
+namespace Horde.Storage.Implementation;
+
+public class ScyllaSettings
+{
+    public string ConnectionString { get; set; } = "Contact Points=localhost,scylla";
+    public long InlineBlobMaxSize { get; set; } = 32 * 1024; // default to 32 kb blobs max
+    public int MaxSnapshotsPerNamespace { get; set; } = 10;
+
+    /// <summary>
+    /// Set to override the replication strategy used to create keyspace
+    /// Note that this only applies when creating the keyspace
+    /// To modify a existing keyspace see https://docs.datastax.com/en/dse/6.7/dse-admin/datastax_enterprise/operations/opsChangeKSStrategy.html
+    /// </summary>
+    public Dictionary<string, string>? KeyspaceReplicationStrategy { get; set; }
+
+    /// <summary>
+    /// Set to override the replication strategy used to create the local keyspace
+    /// Note that this only applies when creating the keyspace
+    /// To modify a existing keyspace see https://docs.datastax.com/en/dse/6.7/dse-admin/datastax_enterprise/operations/opsChangeKSStrategy.html
+    /// </summary>
+    public Dictionary<string, string>? LocalKeyspaceReplicationStrategy { get; set; }
+
+    /// <summary>
+    /// Used to configure the load balancing policy to stick to this specified datacenter
+    /// </summary>
+    [Required]
+    public string LocalDatacenterName { get; set; } = null!;
+        
+    [Required] 
+    public string LocalKeyspaceSuffix { get; set; } = null!;
+
+    /// <summary>
+    /// Max number of connections for each scylla host before switching to another host
+    /// See https://docs.datastax.com/en/developer/nodejs-driver/4.6/features/connection-pooling/
+    /// </summary>
+    public int MaxConnectionForLocalHost { get; set; } = 8192;
+
+    /// <summary>
+    /// The time for a replication log event to live in the incremental state before being deleted, assumption is that the snapshot will have processed the event within this time
+    /// </summary>
+    public TimeSpan ReplicationLogTimeToLive { get; set; } = TimeSpan.FromDays(7);
+
+    /// <summary>
+    /// Whether a Cosmos DB Cassandra layer from Azure should be used.
+    /// </summary>
+    public bool UseAzureCosmosDB { get; set; } = false;
+}
