Horde Storage - Fixed issue were ref was thought to not exist when not fetching its payload

diff --git a/Engine/Source/Programs/Horde/HordeStorage/Horde.Storage/Implementation/Objects/ScyllaReferencesStore.cs b/Engine/Source/Programs/Horde/HordeStorage/Horde.Storage/Implementation/Objects/ScyllaReferencesStore.cs
--- a/Engine/Source/Programs/Horde/HordeStorage/Horde.Storage/Implementation/Objects/ScyllaReferencesStore.cs
+++ b/Engine/Source/Programs/Horde/HordeStorage/Horde.Storage/Implementation/Objects/ScyllaReferencesStore.cs
@@ -76,7 +76,7 @@ namespace Horde.Storage.Implementation
 
             try
             {
-                o.ThrowIfRequiredFieldIsMissing();
+                o.ThrowIfRequiredFieldIsMissing(includePayload);
             }
             catch (Exception e)
             {
@@ -319,7 +319,7 @@ namespace Horde.Storage.Implementation
         [Cassandra.Mapping.Attributes.Column("last_access_time")]
         public DateTime LastAccessTime { get; set; }
 
-        public void ThrowIfRequiredFieldIsMissing()
+        public void ThrowIfRequiredFieldIsMissing(bool includePayload)
         {
             if (string.IsNullOrEmpty(Namespace))
                 throw new ArgumentException("Namespace was not valid", nameof(Namespace));
@@ -327,7 +327,7 @@ namespace Horde.Storage.Implementation
                 throw new ArgumentException("Bucket was not valid", nameof(Bucket));
             if (string.IsNullOrEmpty(Name))
                 throw new ArgumentException("Name was not valid", nameof(Name));
-            if (PayloadHash == null)
+            if (PayloadHash == null && includePayload)
                 throw new ArgumentException("PayloadHash was not valid", nameof(PayloadHash));
             if (!IsFinalized.HasValue)
                 throw new ArgumentException("IsFinalized was not valid", nameof(IsFinalized));
