Horde Storage - Added tests for field filtering (excluding the inline payload from scylla). Removed the passing of field filters to scylla from ref get though, as we need the inline blob to be able to execute that endpoint. Filtering on the json blob will still happen.

diff --git a/Engine/Source/Programs/Horde/HordeStorage/Horde.Storage/Controllers/ReferencesController.cs b/Engine/Source/Programs/Horde/HordeStorage/Horde.Storage/Controllers/ReferencesController.cs
--- a/Engine/Source/Programs/Horde/HordeStorage/Horde.Storage/Controllers/ReferencesController.cs
+++ b/Engine/Source/Programs/Horde/HordeStorage/Horde.Storage/Controllers/ReferencesController.cs
@@ -114,7 +114,7 @@ namespace Horde.Storage.Controllers
 
             try
             {
-                (ObjectRecord objectRecord, BlobContents? blob) = await _objectService.Get(ns, bucket, key, fields);
+                (ObjectRecord objectRecord, BlobContents? blob) = await _objectService.Get(ns, bucket, key, Array.Empty<string>());
 
                 if (blob == null)
                     throw new InvalidOperationException($"Blob was null when attempting to fetch {ns} {bucket} {key}");
diff --git a/Engine/Source/Programs/Horde/HordeStorage/Horde.Storage/Tests/Functional/Objects/ReferencesTests.cs b/Engine/Source/Programs/Horde/HordeStorage/Horde.Storage/Tests/Functional/Objects/ReferencesTests.cs
--- a/Engine/Source/Programs/Horde/HordeStorage/Horde.Storage/Tests/Functional/Objects/ReferencesTests.cs
+++ b/Engine/Source/Programs/Horde/HordeStorage/Horde.Storage/Tests/Functional/Objects/ReferencesTests.cs
@@ -445,6 +445,76 @@ namespace Horde.Storage.FunctionalTests.References
         }
 
         
+        [TestMethod]
+        public async Task PutGetCompactBinaryFiltering()
+        {
+            CbWriter writer = new CbWriter();
+            writer.BeginObject();
+            writer.WriteString("stringField", "thisIsAField");
+            writer.EndObject();
+
+            byte[] objectData = writer.ToByteArray();
+            BlobIdentifier objectHash = BlobIdentifier.FromBlob(objectData);
+            IoHashKey key = IoHashKey.FromName("newReferenceObject");
+
+            HttpContent requestContent = new ByteArrayContent(objectData);
+            requestContent.Headers.ContentType = new MediaTypeHeaderValue(CustomMediaTypeNames.UnrealCompactBinary);
+            requestContent.Headers.Add(CommonHeaders.HashHeaderName, objectHash.ToString());
+
+            HttpResponseMessage result = await _httpClient!.PutAsync(requestUri: $"api/v1/refs/{TestNamespace}/bucket/{key}.uecb", requestContent);
+            result.EnsureSuccessStatusCode();
+
+            {
+                Assert.AreEqual(CustomMediaTypeNames.UnrealCompactBinary, result!.Content.Headers.ContentType!.MediaType);
+                // check that no blobs are missing
+                await using MemoryStream ms = new MemoryStream();
+                await result.Content.CopyToAsync(ms);
+                byte[] roundTrippedBuffer = ms.ToArray();
+                CbObject cb = new CbObject(roundTrippedBuffer);
+                CbField needsField = cb["needs"];
+                Assert.AreNotEqual(CbField.Empty, needsField);
+                List<BlobIdentifier> missingBlobs = needsField.AsArray().Select(field => BlobIdentifier.FromIoHash(field.AsHash())).ToList();
+                Assert.AreEqual(0, missingBlobs.Count);
+            }
+
+            {
+                BucketId bucket = new BucketId("bucket");
+
+                ObjectRecord objectRecord = await ReferencesStore.Get(TestNamespace, bucket, key, IReferencesStore.FieldFlags.None);
+
+                Assert.IsTrue(objectRecord.IsFinalized);
+                Assert.AreEqual(key, objectRecord.Name);
+                Assert.AreEqual(objectHash, objectRecord.BlobIdentifier);
+                Assert.IsNull(objectRecord.InlinePayload);
+            }
+            
+            {
+                HttpResponseMessage getResponse = await _httpClient.GetAsync($"api/v1/refs/{TestNamespace}/bucket/{key}.json?fields=name");
+                getResponse.EnsureSuccessStatusCode();
+                await using MemoryStream ms = new MemoryStream();
+                await getResponse.Content.CopyToAsync(ms);
+
+                byte[] roundTrippedBuffer = ms.ToArray();
+
+                string s = Encoding.ASCII.GetString(roundTrippedBuffer);
+                JObject jObject = JObject.Parse(s);
+                Assert.AreEqual(1, jObject.Children().Count());
+
+                JToken? childToken = jObject.Children().First();
+                Assert.IsNotNull(childToken);
+                Assert.AreEqual(JTokenType.Property, childToken.Type);
+                JProperty property = jObject.Children<JProperty>().First();
+                Assert.IsNotNull(property);
+
+                Assert.AreEqual("stringField", property!.Name);
+
+                string? value = property.Value.Value<string>();
+                Assert.IsNotNull(value);
+                Assert.AreEqual("thisIsAField", value);
+            }
+        }
+
+        
         [TestMethod]
         public async Task PutGetCompactBinaryHierarchy()
         {
