Horde Storage - Fixed failing mongo tests by not outputting the blob if we have flagged to not include it.

diff --git a/Engine/Source/Programs/Horde/HordeStorage/Horde.Storage/Implementation/Objects/MongoReferencesStore.cs b/Engine/Source/Programs/Horde/HordeStorage/Horde.Storage/Implementation/Objects/MongoReferencesStore.cs
--- a/Engine/Source/Programs/Horde/HordeStorage/Horde.Storage/Implementation/Objects/MongoReferencesStore.cs
+++ b/Engine/Source/Programs/Horde/HordeStorage/Horde.Storage/Implementation/Objects/MongoReferencesStore.cs
@@ -44,12 +44,19 @@ namespace Horde.Storage.Implementation
 
         public async Task<ObjectRecord> Get(NamespaceId ns, BucketId bucket, IoHashKey key, IReferencesStore.FieldFlags flags)
         {
+            bool includePayload = (flags & IReferencesStore.FieldFlags.IncludePayload) != 0;
             IMongoCollection<MongoReferencesModelV0> collection = GetCollection<MongoReferencesModelV0>();
             IAsyncCursor<MongoReferencesModelV0>? cursor = await collection.FindAsync(m => m.Ns == ns.ToString() && m.Bucket == bucket.ToString() && m.Key == key.ToString());
             MongoReferencesModelV0? model = await cursor.FirstOrDefaultAsync();
             if (model == null)
                 throw new ObjectNotFoundException(ns, bucket, key);
 
+            if (!includePayload)
+            {
+                // TODO: We should actually use this information to use a projection and not fetch the actual blob
+                model.InlineBlob = null;
+            }
+
             return model.ToObjectRecord();
         }
 
