Horde Storage - No longer doing a lookup before attempting to delete a record in scylla, unnesscary work and that early out would result in partial records not getting deleted (partial records can happen if a last access update runs after the object has been deleted).
Updated tests to verify its safe to delete a object that is already deleted.

diff --git a/Engine/Source/Programs/Horde/HordeStorage/Horde.Storage/Implementation/Objects/ScyllaReferencesStore.cs b/Engine/Source/Programs/Horde/HordeStorage/Horde.Storage/Implementation/Objects/ScyllaReferencesStore.cs
--- a/Engine/Source/Programs/Horde/HordeStorage/Horde.Storage/Implementation/Objects/ScyllaReferencesStore.cs
+++ b/Engine/Source/Programs/Horde/HordeStorage/Horde.Storage/Implementation/Objects/ScyllaReferencesStore.cs
@@ -161,17 +161,6 @@ namespace Horde.Storage.Implementation
             using IScope scope = Tracer.Instance.StartActive("scylla.delete_record");
             scope.Span.ResourceName = $"{ns}.{bucket}.{key}";
 
-            ObjectRecord record;
-            try
-            {
-                record = await Get(ns, bucket, key, IReferencesStore.FieldFlags.None);
-            }
-            catch (ObjectNotFoundException)
-            {
-                // if the record does not exist we do not need to do anything
-                return false;
-            }
-
             AppliedInfo<ScyllaObject> info = await _mapper.DeleteIfAsync<ScyllaObject>("WHERE namespace=? AND bucket=? AND name=?", ns.ToString(), bucket.ToString(), key.ToString());
 
             if (info.Applied)
diff --git a/Engine/Source/Programs/Horde/HordeStorage/Horde.Storage/Tests/Functional/Objects/ReferencesTests.cs b/Engine/Source/Programs/Horde/HordeStorage/Horde.Storage/Tests/Functional/Objects/ReferencesTests.cs
--- a/Engine/Source/Programs/Horde/HordeStorage/Horde.Storage/Tests/Functional/Objects/ReferencesTests.cs
+++ b/Engine/Source/Programs/Horde/HordeStorage/Horde.Storage/Tests/Functional/Objects/ReferencesTests.cs
@@ -1319,6 +1319,12 @@ namespace Horde.Storage.FunctionalTests.References
                 HttpResponseMessage result = await _httpClient!.GetAsync(requestUri: $"api/v1/refs/{TestNamespace}/bucket/{key}.raw");
                 Assert.AreEqual(HttpStatusCode.NotFound, result.StatusCode);
             }
+
+            // delete the object again which doesn't exist anymore
+            {
+                HttpResponseMessage result = await _httpClient!.DeleteAsync(requestUri: $"api/v1/refs/{TestNamespace}/bucket/{key}");
+                result.EnsureSuccessStatusCode();
+            }
         }
 
         
