Horde Storage - Ref GC should avoid inserting a delete event if the ref was in fact not deleted. Otherwise we risk flooding the replication log with delete events if something prevents us from deleting a record.

diff --git a/Engine/Source/Programs/Horde/HordeStorage/Horde.Storage/Implementation/GC/CacheCleanup.cs b/Engine/Source/Programs/Horde/HordeStorage/Horde.Storage/Implementation/GC/CacheCleanup.cs
--- a/Engine/Source/Programs/Horde/HordeStorage/Horde.Storage/Implementation/GC/CacheCleanup.cs
+++ b/Engine/Source/Programs/Horde/HordeStorage/Horde.Storage/Implementation/GC/CacheCleanup.cs
@@ -93,20 +93,23 @@ namespace Horde.Storage.Implementation
                 using IScope scope = Tracer.Instance.StartActive("gc.ref");
                 scope.Span.ResourceName = $"{ns}:{bucket}.{name}";
                 // delete the old record from the ref refs
-                Task<bool> storeDelete = _referencesStore.Delete(ns, bucket, name);
-                // insert a delete event into the transaction log
-                Task<(string, Guid)> transactionLogDelete = _replicationLog.InsertDeleteEvent(ns, bucket, name, null);
 
+                bool storeDelete = false;
                 try
                 {
-                    Task.WaitAll(storeDelete, transactionLogDelete);
+                    storeDelete = await _referencesStore.Delete(ns, bucket, name);
+                    if (storeDelete)
+                    {
+                        // insert a delete event into the transaction log
+                        await _replicationLog.InsertDeleteEvent(ns, bucket, name, null);
+                    }
                 }
                 catch (Exception e)
                 {
                     _logger.Warning(e, "Exception when attempting to delete record {Bucket} {Name} in {Namespace}", bucket, name, ns);
                 }
 
-                if (await storeDelete)
+                if (storeDelete)
                 {
                     Interlocked.Increment(ref countOfDeletedRecords);
                 }
