Horde Storage - Fixed consistency between reference store backends when deleting objects, now always returning a 404 if the object you are deleting didnt exist.
Also some driveby fixes to the swagger documentation

diff --git a/Engine/Source/Programs/Horde/HordeStorage/Horde.Storage/Controllers/ReferencesController.cs b/Engine/Source/Programs/Horde/HordeStorage/Horde.Storage/Controllers/ReferencesController.cs
--- a/Engine/Source/Programs/Horde/HordeStorage/Horde.Storage/Controllers/ReferencesController.cs
+++ b/Engine/Source/Programs/Horde/HordeStorage/Horde.Storage/Controllers/ReferencesController.cs
@@ -839,7 +839,7 @@ namespace Horde.Storage.Controllers
         /// <param name="ns">Namespace. Each namespace is completely separated from each other. Use for different types of data that is never expected to be similar (between two different games for instance)</param>
         /// <param name="bucket">The category/type of record you are caching. Is a clustered key together with the actual key, but all records in the same bucket can be dropped easily.</param>
         [HttpDelete("{ns}/{bucket}", Order = 500)]
-        [ProducesResponseType(204)]
+        [ProducesResponseType(200)]
         [Authorize("Object.delete")]
         public async Task<IActionResult> DeleteBucket(
             [FromRoute] [Required] NamespaceId ns,
@@ -873,8 +873,8 @@ namespace Horde.Storage.Controllers
         /// <param name="bucket">The category/type of record you are caching. Is a clustered key together with the actual key, but all records in the same bucket can be dropped easily.</param>
         /// <param name="key">The unique name of this particular key</param>
         [HttpDelete("{ns}/{bucket}/{key}", Order = 500)]
-        [ProducesResponseType(204)]
-        [ProducesResponseType(400)]
+        [ProducesResponseType(200)]
+        [ProducesResponseType(404)]
         [Authorize("Object.delete")]
         public async Task<IActionResult> Delete(
             [FromRoute] [Required] NamespaceId ns,
@@ -891,11 +891,19 @@ namespace Horde.Storage.Controllers
             try
             {
                 bool deleted = await _objectService.Delete(ns, bucket, key);
-                return Ok(new RefDeletedResponse(deleted ? 1: 0));
+                if (!deleted)
+                {
+                    return NotFound(new ProblemDetails { Title = $"Object {key} in bucket {bucket} and namespace {ns} did not exist" });
+                }
+                return Ok(new RefDeletedResponse(deleted ? 1 : 0));
             }
             catch (NamespaceNotFoundException e)
             {
-                return NotFound(new ProblemDetails {Title = $"Namespace {e.Namespace} did not exist"});
+                return NotFound(new ProblemDetails { Title = $"Namespace {e.Namespace} did not exist" });
+            }
+            catch (ObjectNotFoundException)
+            {
+                return NotFound(new ProblemDetails { Title = $"Object {key} in bucket {bucket} and namespace {ns} did not exist" });
             }
         }
     }
diff --git a/Engine/Source/Programs/Horde/HordeStorage/Horde.Storage/Tests/Functional/Objects/ReferencesTests.cs b/Engine/Source/Programs/Horde/HordeStorage/Horde.Storage/Tests/Functional/Objects/ReferencesTests.cs
--- a/Engine/Source/Programs/Horde/HordeStorage/Horde.Storage/Tests/Functional/Objects/ReferencesTests.cs
+++ b/Engine/Source/Programs/Horde/HordeStorage/Horde.Storage/Tests/Functional/Objects/ReferencesTests.cs
@@ -1323,7 +1323,7 @@ namespace Horde.Storage.FunctionalTests.References
             // delete the object again which doesn't exist anymore
             {
                 HttpResponseMessage result = await _httpClient!.DeleteAsync(requestUri: $"api/v1/refs/{TestNamespace}/bucket/{key}");
-                result.EnsureSuccessStatusCode();
+                Assert.AreEqual(HttpStatusCode.NotFound, result.StatusCode);
             }
         }
 
