Horde Storage - Fixed issue in Mongo and Scylla backends were deletes of a already deleted row was not reported as a 404.

diff --git a/Engine/Source/Programs/Horde/HordeStorage/Horde.Storage/Implementation/Objects/MongoReferencesStore.cs b/Engine/Source/Programs/Horde/HordeStorage/Horde.Storage/Implementation/Objects/MongoReferencesStore.cs
--- a/Engine/Source/Programs/Horde/HordeStorage/Horde.Storage/Implementation/Objects/MongoReferencesStore.cs
+++ b/Engine/Source/Programs/Horde/HordeStorage/Horde.Storage/Implementation/Objects/MongoReferencesStore.cs
@@ -151,7 +151,7 @@ namespace Horde.Storage.Implementation
             DeleteResult result = await collection.DeleteOneAsync(model =>
                 model.Ns == ns.ToString() && model.Bucket == bucket.ToString() && model.Key == key.ToString());
 
-            return result.IsAcknowledged;
+            return result.DeletedCount != 0;
         }
 
         public async Task<long> DropNamespace(NamespaceId ns)
diff --git a/Engine/Source/Programs/Horde/HordeStorage/Horde.Storage/Implementation/Objects/ScyllaReferencesStore.cs b/Engine/Source/Programs/Horde/HordeStorage/Horde.Storage/Implementation/Objects/ScyllaReferencesStore.cs
--- a/Engine/Source/Programs/Horde/HordeStorage/Horde.Storage/Implementation/Objects/ScyllaReferencesStore.cs
+++ b/Engine/Source/Programs/Horde/HordeStorage/Horde.Storage/Implementation/Objects/ScyllaReferencesStore.cs
@@ -161,7 +161,7 @@ namespace Horde.Storage.Implementation
             using IScope scope = Tracer.Instance.StartActive("scylla.delete_record");
             scope.Span.ResourceName = $"{ns}.{bucket}.{key}";
 
-            AppliedInfo<ScyllaObject> info = await _mapper.DeleteIfAsync<ScyllaObject>("WHERE namespace=? AND bucket=? AND name=?", ns.ToString(), bucket.ToString(), key.ToString());
+            AppliedInfo<ScyllaObject> info = await _mapper.DeleteIfAsync<ScyllaObject>("WHERE namespace=? AND bucket=? AND name=? IF EXISTS", ns.ToString(), bucket.ToString(), key.ToString());
 
             if (info.Applied)
                 return true;
