Horde: Add timing stats to remote execution responses.

Two sets of data are returned - one detailing the timeline of the task on the server (being queued, through dispatched, to completed), and one detailing the timeline of the task on the remote agent (downloading data, executing the task, and uploading the result).

The standalone compute command implementation in Horde.Agent also logs its own timeline of events.

diff --git a/Engine/Source/Programs/Horde/Horde.Agent/Commands/Compute/ComputeCommand.cs b/Engine/Source/Programs/Horde/Horde.Agent/Commands/Compute/ComputeCommand.cs
--- a/Engine/Source/Programs/Horde/Horde.Agent/Commands/Compute/ComputeCommand.cs
+++ b/Engine/Source/Programs/Horde/Horde.Agent/Commands/Compute/ComputeCommand.cs
@@ -2,6 +2,7 @@
 
 using System;
 using System.Collections.Generic;
+using System.Diagnostics;
 using System.IO;
 using System.Linq;
 using System.Text;
@@ -88,6 +89,8 @@ namespace Horde.Agent.Commands
 		[CommandLine("-LogLevel")]
 		public string LogLevelStr { get; set; } = "debug";
 
+		readonly Stopwatch _timer = Stopwatch.StartNew();
+
 		static (DirectoryTree, IoHash) CreateSandbox(DirectoryInfo baseDirInfo, Dictionary<IoHash, byte[]> uploadList)
 		{
 			DirectoryTree tree = new DirectoryTree();
@@ -111,9 +114,20 @@ namespace Horde.Agent.Commands
 			return (tree, AddCbObject(uploadList, tree));
 		}
 
+		readonly List<string> _localScopeNames = new List<string>();
+		readonly List<int> _localScopeTimesMs = new List<int>();
+
+		void AddLocalScope(string name)
+		{
+			_localScopeTimesMs.Add((int)_timer.Elapsed.TotalMilliseconds);
+			_localScopeNames.Add(name);
+			_timer.Restart();
+		}
+
 		/// <inheritdoc/>
 		public override async Task<int> ExecuteAsync(ILogger logger)
 		{
+			DateTime startTime = DateTime.UtcNow - _timer.Elapsed;
 			InputDir ??= Input.Directory;
 
 			if (Enum.TryParse(LogLevelStr, true, out LogEventLevel logEventLevel))
@@ -167,9 +181,13 @@ namespace Horde.Agent.Commands
 					logger.LogInformation("Using salt: {SaltBytes}", StringUtils.FormatHexString(saltBytes.ToByteArray()));
 				}
 
+				AddLocalScope("setup");
+
 				Dictionary<IoHash, byte[]> blobs = new Dictionary<IoHash, byte[]>();
 				(_, IoHash sandboxHash) = CreateSandbox(InputDir.ToDirectoryInfo(), blobs);
 
+				AddLocalScope("scan");
+
 				JsonComputeTask jsonComputeTask = new JsonComputeTask();
 				configuration.Bind(jsonComputeTask);
 
@@ -187,7 +205,25 @@ namespace Horde.Agent.Commands
 				task.OutputPaths.AddRange(jsonComputeTask.OutputPaths.Select(x => (Utf8String)x));
 				task.RequirementsHash = AddCbObject(blobs, requirements);
 
-				await ExecuteAction(logger, computeClient, storageClient, jsonComputeTask.ClusterId, task, blobs);	
+				AddLocalScope("task-setup");
+
+				ComputeTaskStatus result = await ExecuteAction(logger, computeClient, storageClient, jsonComputeTask.ClusterId, task, blobs);
+				if (result.Outcome != ComputeTaskOutcome.Success)
+				{
+					logger.LogError("{OperationName}: Outcome: {Outcome}, Detail: {Detail}", result.TaskRefId, result.Outcome.ToString(), result.Detail ?? "(none)");
+				}
+
+				AddLocalScope($"complete");
+
+				LogScopes("local-timing", startTime, _localScopeNames.ToArray(), _localScopeTimesMs.ToArray(), logger);
+				if (result.QueueStats != null)
+				{
+					LogScopes("server-timing", result.QueueStats.StartTime, ComputeTaskQueueStats.ScopeNames, result.QueueStats.Scopes, logger);
+				}
+				if (result.ExecutionStats != null)
+				{
+					LogScopes("execution-timing", result.ExecutionStats.StartTime, ComputeTaskExecutionStats.ScopeNames, result.ExecutionStats.Scopes, logger);
+				}
 			}
 			return 0;
 		}
@@ -200,37 +236,54 @@ namespace Horde.Agent.Commands
 			return hash;
 		}
 
-		private async Task ExecuteAction(ILogger logger, IComputeClient computeClient, IStorageClient storageClient, ClusterId clusterId, ComputeTask task, Dictionary<IoHash, byte[]> uploadList)
+		static async Task UploadSandbox(IComputeClusterInfo cluster, IStorageClient storageClient, Dictionary<IoHash, byte[]> uploadList, ILogger logger)
 		{
-			IComputeClusterInfo cluster = await computeClient.GetClusterInfoAsync(clusterId);
+			int totalUploaded = 0;
+			int totalSkipped = 0;
+			int uploadedBytes = 0;
+			int skippedBytes = 0;
 
+			List<Task> tasks = new List<Task>();
+			foreach ((IoHash hash, byte[] data) in uploadList)
 			{
-				int index = 0, totalUploaded = 0, totalSkipped = 0, uploadedBytes = 0, skippedBytes = 0;
-				IEnumerable<Task> tasks = uploadList.Select(async (pair) =>
+				int index = tasks.Count;
+				async Task UploadItemAsync()
 				{
-					if (await storageClient.HasBlobAsync(cluster.NamespaceId, pair.Key))
+					if (await storageClient.HasBlobAsync(cluster.NamespaceId, hash))
 					{
-						logger.LogInformation("Skipped blob {Idx}/{Count}: {Hash} ({Length} bytes)", Interlocked.Increment(ref index), uploadList.Count, pair.Key, pair.Value.Length);
+						logger.LogInformation("Skipped blob {Idx}/{Count}: {Hash} ({Length} bytes)", index, uploadList.Count, hash, data.Length);
 						Interlocked.Increment(ref totalSkipped);
-						Interlocked.Add(ref skippedBytes, pair.Value.Length);
-						return;
+						Interlocked.Add(ref skippedBytes, data.Length);
 					}
+					else
+					{
+						await storageClient.WriteBlobFromMemoryAsync(cluster.NamespaceId, hash, data);
+						logger.LogInformation("Uploaded blob {Idx}/{Count}: {Hash} ({Bytes} bytes)", index, uploadList.Count, hash, data.Length);
+						Interlocked.Increment(ref totalUploaded);
+						Interlocked.Add(ref uploadedBytes, data.Length);
+					}
+				}
+				tasks.Add(UploadItemAsync());
+			}
+			await Task.WhenAll(tasks);
 
-					await storageClient.WriteBlobFromMemoryAsync(cluster.NamespaceId, pair.Key, pair.Value);
-					logger.LogInformation("Uploaded blob {Idx}/{Count}: {Hash} ({Bytes} bytes)", Interlocked.Increment(ref index), uploadList.Count, pair.Key, pair.Value.Length);
-					Interlocked.Increment(ref totalUploaded);
-					Interlocked.Add(ref uploadedBytes, pair.Value.Length);
-				});
-				await Task.WhenAll(tasks);
+			logger.LogInformation("Uploaded {UploadedCount} blobs ({UploadedBytes} bytes), Skipped {SkippedCount} blobs ({SkippedBytes} bytes)", totalUploaded, uploadedBytes, totalSkipped, skippedBytes);
+		}
 
-				logger.LogInformation("Uploaded {UploadedCount} blobs ({UploadedBytes} bytes), Skipped {SkippedCount} blobs ({SkippedBytes} bytes)", totalUploaded, uploadedBytes, totalSkipped, skippedBytes);
-			}
+		private async Task<ComputeTaskStatus> ExecuteAction(ILogger logger, IComputeClient computeClient, IStorageClient storageClient, ClusterId clusterId, ComputeTask task, Dictionary<IoHash, byte[]> uploadList)
+		{
+			IComputeClusterInfo cluster = await computeClient.GetClusterInfoAsync(clusterId);
+
+			await UploadSandbox(cluster, storageClient, uploadList, logger);
+			AddLocalScope("uploaded-sandbox");
 
 			CbObject taskObject = CbSerializer.Serialize(task);
 			IoHash taskHash = IoHash.Compute(taskObject.GetView().Span);
 			RefId taskRefId = new RefId(taskHash);
 			await storageClient.SetRefAsync(cluster.NamespaceId, cluster.RequestBucketId, taskRefId, taskObject);
 
+			AddLocalScope("uploaded-task");
+
 			ChannelId channelId = new ChannelId(Guid.NewGuid().ToString());
 			logger.LogInformation("cluster: {ClusterId}", clusterId);
 			logger.LogInformation("channel: {ChannelId}", channelId);
@@ -239,27 +292,40 @@ namespace Horde.Agent.Commands
 
 			// Execute the action
 			await computeClient.AddTaskAsync(clusterId, channelId, taskRefId, task.RequirementsHash, SkipCacheLookup);
+			AddLocalScope("queued-task");
 
-			await foreach(IComputeTaskInfo response in computeClient.GetTaskUpdatesAsync(clusterId, channelId))
+			await foreach (ComputeTaskStatus response in computeClient.GetTaskUpdatesAsync(clusterId, channelId))
 			{
 				logger.LogInformation("{OperationName}: Execution state: {State}", response.TaskRefId, response.State.ToString());
+
 				if (!String.IsNullOrEmpty(response.AgentId) || !String.IsNullOrEmpty(response.LeaseId))
 				{
 					logger.LogInformation("{OperationName}: Running on agent {AgentId} under lease {LeaseId}", response.TaskRefId, response.AgentId, response.LeaseId);
 				}
+
 				if (response.ResultRefId != null)
 				{
 					await HandleCompleteTask(storageClient, cluster.NamespaceId, cluster.ResponseBucketId, response.ResultRefId.Value, logger);
 				}
 				if (response.State == ComputeTaskState.Complete)
 				{
-					if (response.Outcome != ComputeTaskOutcome.Success)
-					{
-						logger.LogError("{OperationName}: Outcome: {Outcome}, Detail: {Detail}", response.TaskRefId, response.Outcome.ToString(), response.Detail ?? "(none)");
-					}
-					break;
+					return response;
 				}
 			}
+
+			throw new InvalidOperationException("Execution finished without completed message");
+		}
+
+		static void LogScopes(string prefix, DateTime startTime, string[] scopeNames, int[] scopes, ILogger logger)
+		{
+			int totalMs = 0;
+
+			logger.LogInformation("{Prefix}: [{At,6:n0}] start (at {Time})", prefix, totalMs, startTime);
+			for (int idx = 0; idx < scopes.Length; idx++)
+			{
+				totalMs += scopes[idx];
+				logger.LogInformation("{Prefix}: [{At,6:n0}] {Name} ({Time:n0})", prefix, totalMs, scopeNames[idx], scopes[idx]);
+			}
 		}
 
 		async Task HandleCompleteTask(IStorageClient storageClient, NamespaceId namespaceId, BucketId outputBucketId, RefId outputRefId, ILogger logger)
diff --git a/Engine/Source/Programs/Horde/Horde.Agent/Execution/ComputeTaskExecutor.cs b/Engine/Source/Programs/Horde/Horde.Agent/Execution/ComputeTaskExecutor.cs
--- a/Engine/Source/Programs/Horde/Horde.Agent/Execution/ComputeTaskExecutor.cs
+++ b/Engine/Source/Programs/Horde/Horde.Agent/Execution/ComputeTaskExecutor.cs
@@ -75,6 +75,13 @@ namespace Horde.Agent.Execution
 			return (stdOutStream.ToArray(), stdErrStream.ToArray());
 		}
 
+		static int GetMarkerTime(Stopwatch timer)
+		{
+			int elapsedMs = (int)timer.ElapsedMilliseconds;
+			timer.Restart();
+			return elapsedMs;
+		}
+
 		/// <summary>
 		/// Execute an action
 		/// </summary>
@@ -85,20 +92,24 @@ namespace Horde.Agent.Execution
 		/// <returns>The action result</returns>
 		public async Task<ComputeTaskResultMessage> ExecuteAsync(string leaseId, ComputeTaskMessage computeTaskMessage, DirectoryReference sandboxDir, CancellationToken cancellationToken)
 		{
+			Stopwatch timer = Stopwatch.StartNew();
+			ComputeTaskExecutionStatsMessage stats = new ComputeTaskExecutionStatsMessage();
+			stats.StartTime = Google.Protobuf.WellKnownTypes.Timestamp.FromDateTime(DateTime.UtcNow);
+
 			NamespaceId namespaceId = new NamespaceId(computeTaskMessage.NamespaceId);
 			BucketId inputBucketId = new BucketId(computeTaskMessage.InputBucketId);
 			BucketId outputBucketId = new BucketId(computeTaskMessage.OutputBucketId);
 
 			ComputeTask task = await _storageClient.GetRefAsync<ComputeTask>(namespaceId, inputBucketId, computeTaskMessage.TaskRefId.AsRefId());
 			_logger.LogInformation("Executing task {Hash} for lease ID {LeaseId}", computeTaskMessage.TaskRefId, leaseId);
+			stats.DownloadRefMs = GetMarkerTime(timer);
 
 			DirectoryReference.CreateDirectory(sandboxDir);
 			FileUtils.ForceDeleteDirectoryContents(sandboxDir);
 
-			DateTimeOffset inputFetchStart = DateTimeOffset.UtcNow;
 			DirectoryTree inputDirectory = await _storageClient.ReadBlobAsync<DirectoryTree>(namespaceId, task.SandboxHash);
 			await SetupSandboxAsync(namespaceId, inputDirectory, sandboxDir);
-			DateTimeOffset inputFetchCompleted = DateTimeOffset.UtcNow;
+			stats.DownloadInputMs = GetMarkerTime(timer);
 
 			using (ManagedProcessGroup processGroup = new ManagedProcessGroup())
 			{
@@ -119,10 +130,10 @@ namespace Horde.Agent.Execution
 				_logger.LogInformation("Executing {FileName} with arguments {Arguments}", fileName, arguments);
 
 				TimeSpan timeout = TimeSpan.FromMinutes(5.0);// Action.Timeout == null ? TimeSpan.FromMinutes(5) : Action.Timeout.ToTimeSpan();
-				DateTimeOffset executionStartTime = DateTimeOffset.UtcNow;
 				using (ManagedProcess process = new ManagedProcess(processGroup, fileName, arguments, workingDirectory, newEnvironment, ProcessPriorityClass.Normal, ManagedProcessFlags.None))
 				{
 					(byte[] stdOutData, byte[] stdErrData) = await ReadProcessStreams(process, timeout, cancellationToken);
+					stats.ExecMs = GetMarkerTime(timer);
 
 					foreach (string line in Encoding.UTF8.GetString(stdOutData).Split('\n'))
 					{
@@ -138,6 +149,7 @@ namespace Horde.Agent.Execution
 					ComputeTaskResult result = new ComputeTaskResult(process.ExitCode);
 					result.StdOutHash = await _storageClient.WriteBlobFromMemoryAsync(namespaceId, stdOutData);
 					result.StdErrHash = await _storageClient.WriteBlobFromMemoryAsync(namespaceId, stdErrData);
+					stats.UploadLogMs = GetMarkerTime(timer);
 
 					FileReference[] outputFiles = ResolveOutputPaths(sandboxDir, task.OutputPaths.Select(x => x.ToString())).OrderBy(x => x.FullName, StringComparer.Ordinal).ToArray();
 					if (outputFiles.Length > 0)
@@ -148,11 +160,13 @@ namespace Horde.Agent.Execution
 						}
 						result.OutputHash = await PutOutput(namespaceId, sandboxDir, outputFiles);
 					}
+					stats.UploadOutputMs = GetMarkerTime(timer);
 
 					CbObject resultObject = CbSerializer.Serialize(result);
-
 					await _storageClient.SetRefAsync(namespaceId, outputBucketId, computeTaskMessage.TaskRefId, resultObject);
-					return new ComputeTaskResultMessage(computeTaskMessage.TaskRefId);
+					stats.UploadRefMs = GetMarkerTime(timer);
+
+					return new ComputeTaskResultMessage(computeTaskMessage.TaskRefId, stats);
 				}
 			}
 		}
diff --git a/Engine/Source/Programs/Horde/Horde.Agent/Utilities/ProtoMethods.cs b/Engine/Source/Programs/Horde/Horde.Agent/Utilities/ProtoMethods.cs
deleted file mode 100644
diff --git a/Engine/Source/Programs/Horde/Horde.Build/Compute/IComputeService.cs b/Engine/Source/Programs/Horde/Horde.Build/Compute/IComputeService.cs
--- a/Engine/Source/Programs/Horde/Horde.Build/Compute/IComputeService.cs
+++ b/Engine/Source/Programs/Horde/Horde.Build/Compute/IComputeService.cs
@@ -1,6 +1,5 @@
 // Copyright Epic Games, Inc. All Rights Reserved.
 
-using System;
 using System.Collections.Generic;
 using System.Threading;
 using System.Threading.Tasks;
@@ -13,54 +12,6 @@ using Horde.Build.Utilities;
 
 namespace Horde.Build.Compute
 {
-	using LeaseId = ObjectId<ILease>;
-
-	/// <summary>
-	/// Status of a compute task
-	/// </summary>
-	public interface IComputeTaskStatus
-	{
-		/// <summary>
-		/// The input task ref id
-		/// </summary>
-		public RefId TaskRefId { get; }
-
-		/// <summary>
-		/// Timestamp for the event
-		/// </summary>
-		public DateTime Time { get; }
-
-		/// <summary>
-		/// State of the task
-		/// </summary>
-		public ComputeTaskState State { get; }
-
-		/// <summary>
-		/// Indicates the outcome of a task
-		/// </summary>
-		public ComputeTaskOutcome Outcome { get; }
-
-		/// <summary>
-		/// When transitioning to the executing state, includes the agent name
-		/// </summary>
-		public AgentId? AgentId { get; }
-
-		/// <summary>
-		/// The lease id
-		/// </summary>
-		public LeaseId? LeaseId { get; }
-
-		/// <summary>
-		/// Hash of the result, if complete
-		/// </summary>
-		public RefId? ResultRefId { get; }
-
-		/// <summary>
-		/// Additional information for the given outcome
-		/// </summary>
-		public string? Detail { get; }
-	}
-
 	/// <summary>
 	/// Interface for the compute service
 	/// </summary>
@@ -88,7 +39,7 @@ namespace Horde.Build.Compute
 		/// <param name="clusterId">Cluster containing the channel</param>
 		/// <param name="channelId">Queue to remove items from</param>
 		/// <returns>List of status updates</returns>
-		Task<List<IComputeTaskStatus>> GetTaskUpdatesAsync(ClusterId clusterId, ChannelId channelId);
+		Task<List<ComputeTaskStatus>> GetTaskUpdatesAsync(ClusterId clusterId, ChannelId channelId);
 
 		/// <summary>
 		/// Dequeue completed items from a queue
@@ -97,6 +48,6 @@ namespace Horde.Build.Compute
 		/// <param name="channelId">Queue to remove items from</param>
 		/// <param name="cancellationToken">Cancellation token to stop waiting for items</param>
 		/// <returns>List of status updates</returns>
-		Task<List<IComputeTaskStatus>> WaitForTaskUpdatesAsync(ClusterId clusterId, ChannelId channelId, CancellationToken cancellationToken);
+		Task<List<ComputeTaskStatus>> WaitForTaskUpdatesAsync(ClusterId clusterId, ChannelId channelId, CancellationToken cancellationToken);
 	}
 }
diff --git a/Engine/Source/Programs/Horde/Horde.Build/Compute/Impl/ComputeController.cs b/Engine/Source/Programs/Horde/Horde.Build/Compute/Impl/ComputeController.cs
--- a/Engine/Source/Programs/Horde/Horde.Build/Compute/Impl/ComputeController.cs
+++ b/Engine/Source/Programs/Horde/Horde.Build/Compute/Impl/ComputeController.cs
@@ -109,31 +109,20 @@ namespace Horde.Build.Compute.Impl
 				return Forbid(AclAction.ViewComputeTasks);
 			}
 
-			List<IComputeTaskStatus> results;
+			List<ComputeTaskStatus> updates;
 			if (wait == 0)
 			{
-				results = await _computeService.GetTaskUpdatesAsync(clusterId, channelId);
+				updates = await _computeService.GetTaskUpdatesAsync(clusterId, channelId);
 			}
 			else
 			{
 				using CancellationTokenSource delaySource = new CancellationTokenSource(wait * 1000);
-				results = await _computeService.WaitForTaskUpdatesAsync(clusterId, channelId, delaySource.Token);
+				updates = await _computeService.WaitForTaskUpdatesAsync(clusterId, channelId, delaySource.Token);
 			}
 
 			GetTaskUpdatesResponse response = new GetTaskUpdatesResponse();
-			foreach (IComputeTaskStatus result in results)
+			foreach (ComputeTaskStatus update in updates)
 			{
-				GetTaskUpdateResponse update = new GetTaskUpdateResponse();
-
-				update.TaskRefId = result.TaskRefId;
-				update.Time = result.Time;
-				update.State = result.State;
-				update.Outcome = result.Outcome;
-				update.ResultRefId = result.ResultRefId;
-				update.AgentId = result.AgentId?.ToString();
-				update.LeaseId = result.LeaseId?.ToString();
-				update.Detail = result.Detail;
-
 				response.Updates.Add(update);
 			}
 
diff --git a/Engine/Source/Programs/Horde/Horde.Build/Compute/Impl/ComputeService.cs b/Engine/Source/Programs/Horde/Horde.Build/Compute/Impl/ComputeService.cs
--- a/Engine/Source/Programs/Horde/Horde.Build/Compute/Impl/ComputeService.cs
+++ b/Engine/Source/Programs/Horde/Horde.Build/Compute/Impl/ComputeService.cs
@@ -42,77 +42,19 @@ namespace Horde.Build.Compute.Impl
 		[CbField("ch")]
 		public ChannelId ChannelId { get; set; }
 
+		[CbField("q")]
+		public DateTime QueuedAt { get; set; }
+
 		private ComputeTaskInfo()
 		{
 		}
 
-		public ComputeTaskInfo(ClusterId clusterId, RefId taskRefId, ChannelId channelId)
+		public ComputeTaskInfo(ClusterId clusterId, RefId taskRefId, ChannelId channelId, DateTime queuedAt)
 		{
 			ClusterId = clusterId;
 			TaskRefId = taskRefId;
 			ChannelId = channelId;
-		}
-	}
-
-	/// <summary>
-	/// Current status of a task
-	/// </summary>
-	[RedisConverter(typeof(RedisCbConverter<>))]
-	class ComputeTaskStatus : IComputeTaskStatus
-	{
-		/// <inheritdoc/>
-		[CbField("h")]
-		public RefId TaskRefId { get; set; }
-
-		/// <inheritdoc/>
-		[CbField("t")]
-		public DateTime Time { get; set; }
-
-		/// <inheritdoc/>
-		[CbField("s")]
-		public ComputeTaskState State { get; set; }
-
-		/// <inheritdoc/>
-		[CbField("o")]
-		public ComputeTaskOutcome Outcome { get; set; }
-
-		/// <inheritdoc cref="IComputeTaskStatus.AgentId"/>
-		[CbField("a")]
-		public Utf8String AgentId { get; set; }
-
-		/// <inheritdoc cref="IComputeTaskStatus.LeaseId"/>
-		[CbField("l")]
-		byte[]? LeaseIdBytes { get; set; }
-
-		/// <inheritdoc/>
-		[CbField("r")]
-		public RefId? ResultRefId { get; set; }
-
-		/// <inheritdoc/>
-		[CbField("d")]
-		public string? Detail { get; set; }
-
-		/// <inheritdoc/>
-		AgentId? IComputeTaskStatus.AgentId => AgentId.IsEmpty ? (AgentId?)null : new AgentId(AgentId.ToString());
-
-		/// <inheritdoc/>
-		public LeaseId? LeaseId
-		{
-			get => (LeaseIdBytes == null || LeaseIdBytes.Length == 0) ? (LeaseId?)null : new LeaseId(LeaseIdBytes);
-			set => LeaseIdBytes = value?.Value.ToByteArray();
-		}
-
-		private ComputeTaskStatus()
-		{
-		}
-
-		public ComputeTaskStatus(RefId taskRefId, ComputeTaskState state, AgentId? agentId, LeaseId? leaseId)
-		{
-			TaskRefId = taskRefId;
-			Time = DateTime.UtcNow;
-			State = state;
-			AgentId = (agentId == null) ? Utf8String.Empty : agentId.Value.ToString();
-			LeaseIdBytes = leaseId?.Value.ToByteArray();
+			QueuedAt = queuedAt;
 		}
 	}
 
@@ -180,6 +122,11 @@ namespace Horde.Build.Compute.Impl
 		readonly LazyCachedValue<Task<Globals>> _globals;
 		readonly ILogger _logger;
 
+		static ComputeService()
+		{
+			RedisSerializer.RegisterConverter<ComputeTaskStatus, RedisCbConverter<ComputeTaskStatus>>();
+		}
+
 		/// <summary>
 		/// Constructor
 		/// </summary>
@@ -196,6 +143,15 @@ namespace Horde.Build.Compute.Impl
 			OnLeaseStartedProperties.Add(x => x.TaskRefId);
 		}
 
+		static ComputeTaskStatus CreateStatus(RefId taskRefId, ComputeTaskState state)
+		{
+			ComputeTaskStatus status = new ComputeTaskStatus();
+			status.TaskRefId = taskRefId;
+			status.Time = DateTime.UtcNow;
+			status.State = state;
+			return status;
+		}
+
 		/// <inheritdoc/>
 		public Task StartAsync(CancellationToken token) => _expireTasksTicker.StartAsync();
 
@@ -221,9 +177,10 @@ namespace Horde.Build.Compute.Impl
 						break;
 					}
 
-					ComputeTaskStatus status = new ComputeTaskStatus(computeTask.TaskRefId, ComputeTaskState.Complete, null, null);
+					ComputeTaskStatus status = CreateStatus(computeTask.TaskRefId, ComputeTaskState.Complete);
 					status.Outcome = ComputeTaskOutcome.Expired;
 					status.Detail = $"No agents monitoring queue {queueKey}";
+
 					_logger.LogInformation("Compute task expired (queue: {RequirementsHash}, task: {TaskHash}, channel: {ChannelId})", queueKey, computeTask.TaskRefId, computeTask.ChannelId);
 					await PostStatusMessageAsync(computeTask, status);
 				}
@@ -258,7 +215,7 @@ namespace Horde.Build.Compute.Impl
 			List<Task> tasks = new List<Task>();
 			foreach (RefId taskRefId in taskRefIds)
 			{
-				ComputeTaskInfo taskInfo = new ComputeTaskInfo(clusterId, taskRefId, channelId);
+				ComputeTaskInfo taskInfo = new ComputeTaskInfo(clusterId, taskRefId, channelId, DateTime.UtcNow);
 				_logger.LogDebug("Adding task {TaskHash} from channel {ChannelId} to queue {ClusterId}{RequirementsHash}", taskRefId.Hash, channelId, clusterId, requirementsHash);
 				tasks.Add(_taskScheduler.EnqueueAsync(new QueueKey(clusterId, requirementsHash), taskInfo, false));
 			}
@@ -306,8 +263,12 @@ namespace Horde.Build.Compute.Impl
 			if (cluster == null)
 			{
 				_logger.LogWarning("Invalid cluster '{ClusterId}'; failing task {TaskRefId}", taskInfo.ClusterId, taskInfo.TaskRefId);
-				ComputeTaskStatus status = new ComputeTaskStatus(taskInfo.TaskRefId, ComputeTaskState.Complete, agent.Id, null) { Detail = $"Invalid cluster '{taskInfo.ClusterId}'" };
+
+				ComputeTaskStatus status = CreateStatus(taskInfo.TaskRefId, ComputeTaskState.Complete);
+				status.AgentId = agent.Id.ToString();
+				status.Detail = $"Invalid cluster '{taskInfo.ClusterId}'";
 				await PostStatusMessageAsync(taskInfo, status);
+
 				return null;
 			}
 
@@ -315,8 +276,12 @@ namespace Horde.Build.Compute.Impl
 			if (requirements == null)
 			{
 				_logger.LogWarning("Unable to fetch requirements {RequirementsHash}", queueKey);
-				ComputeTaskStatus status = new ComputeTaskStatus(taskInfo.TaskRefId, ComputeTaskState.Complete, agent.Id, null) { Detail = $"Unable to retrieve requirements '{queueKey}'" };
+
+				ComputeTaskStatus status = CreateStatus(taskInfo.TaskRefId, ComputeTaskState.Complete);
+				status.AgentId = agent.Id.ToString();
+				status.Detail = $"Unable to retrieve requirements '{queueKey}'";
 				await PostStatusMessageAsync(taskInfo, status);
+
 				return null;
 			}
 
@@ -328,7 +293,9 @@ namespace Horde.Build.Compute.Impl
 			computeTask.OutputBucketId = cluster.ResponseBucketId.ToString();
 			computeTask.RequirementsHash = queueKey.RequirementsHash;
 			computeTask.TaskRefId = taskInfo.TaskRefId;
-
+			computeTask.QueuedAt = Timestamp.FromDateTime(taskInfo.QueuedAt);
+			computeTask.DispatchedMs = (int)(DateTime.UtcNow - taskInfo.QueuedAt).TotalMilliseconds;
+				
 			string leaseName = $"Remote action ({taskInfo.TaskRefId})";
 			byte[] payload = Any.Pack(computeTask).ToByteArray();
 
@@ -341,28 +308,28 @@ namespace Horde.Build.Compute.Impl
 		public override Task CancelLeaseAsync(IAgent agent, LeaseId leaseId, ComputeTaskMessage message)
 		{
 			ClusterId clusterId = new ClusterId(message.ClusterId);
-			ComputeTaskInfo taskInfo = new ComputeTaskInfo(clusterId, new RefId(new IoHash(message.TaskRefId.ToByteArray())), new ChannelId(message.ChannelId));
+			ComputeTaskInfo taskInfo = new ComputeTaskInfo(clusterId, new RefId(new IoHash(message.TaskRefId.ToByteArray())), new ChannelId(message.ChannelId), message.QueuedAt.ToDateTime());
 			return _taskScheduler.EnqueueAsync(new QueueKey(clusterId, new IoHash(message.RequirementsHash.ToByteArray())), taskInfo, true);
 		}
 
 		/// <inheritdoc/>
-		public async Task<List<IComputeTaskStatus>> GetTaskUpdatesAsync(ClusterId clusterId, ChannelId channelId)
+		public async Task<List<ComputeTaskStatus>> GetTaskUpdatesAsync(ClusterId clusterId, ChannelId channelId)
 		{
-			List<ComputeTaskStatus> messages = await _messageQueue.ReadMessagesAsync(GetMessageQueueId(clusterId, channelId));
-			return messages.ConvertAll<IComputeTaskStatus>(x => x);
+			return await _messageQueue.ReadMessagesAsync(GetMessageQueueId(clusterId, channelId));
 		}
 
-		public async Task<List<IComputeTaskStatus>> WaitForTaskUpdatesAsync(ClusterId clusterId, ChannelId channelId, CancellationToken cancellationToken)
+		public async Task<List<ComputeTaskStatus>> WaitForTaskUpdatesAsync(ClusterId clusterId, ChannelId channelId, CancellationToken cancellationToken)
 		{
-			List<ComputeTaskStatus> messages = await _messageQueue.WaitForMessagesAsync(GetMessageQueueId(clusterId, channelId), cancellationToken);
-			return messages.ConvertAll<IComputeTaskStatus>(x => x);
+			return await _messageQueue.WaitForMessagesAsync(GetMessageQueueId(clusterId, channelId), cancellationToken);
 		}
 
 		public override async Task OnLeaseStartedAsync(IAgent agent, LeaseId leaseId, ComputeTaskMessage computeTask, ILogger logger)
 		{
 			await base.OnLeaseStartedAsync(agent, leaseId, computeTask, logger);
 
-			ComputeTaskStatus status = new ComputeTaskStatus(computeTask.TaskRefId, ComputeTaskState.Executing, agent.Id, leaseId);
+			ComputeTaskStatus status = CreateStatus(computeTask.TaskRefId, ComputeTaskState.Executing);
+			status.AgentId = agent.Id.ToString();
+			status.LeaseId = leaseId.ToString();
 			await PostStatusMessageAsync(computeTask, status);
 		}
 
@@ -370,9 +337,16 @@ namespace Horde.Build.Compute.Impl
 		{
 			await base.OnLeaseFinishedAsync(agent, leaseId, computeTask, outcome, output, logger);
 
+			DateTime queuedAt = computeTask.QueuedAt.ToDateTime();
+
 			ComputeTaskResultMessage message = ComputeTaskResultMessage.Parser.ParseFrom(output.ToArray());
 
-			ComputeTaskStatus status = new ComputeTaskStatus(computeTask.TaskRefId, ComputeTaskState.Complete, agent.Id, leaseId);
+			ComputeTaskStatus status = CreateStatus(computeTask.TaskRefId, ComputeTaskState.Complete);
+			status.AgentId = agent.Id.ToString();
+			status.LeaseId = leaseId.ToString();
+			status.QueueStats = new ComputeTaskQueueStats(queuedAt, computeTask.DispatchedMs, (int)(DateTime.UtcNow - queuedAt).TotalMilliseconds);
+			status.ExecutionStats = message.ExecutionStats?.ToNative();
+
 			if (message.ResultRefId != null)
 			{
 				status.ResultRefId = message.ResultRefId;
@@ -474,7 +448,7 @@ namespace Horde.Build.Compute.Impl
 						break;
 					}
 
-					ComputeTaskStatus status = new ComputeTaskStatus(computeTask.TaskRefId, ComputeTaskState.Complete, null, null);
+					ComputeTaskStatus status = CreateStatus(computeTask.TaskRefId, ComputeTaskState.Complete);
 					status.Outcome = ComputeTaskOutcome.BlobNotFound;
 					status.Detail = $"Missing requirements object {queueKey.RequirementsHash}";
 					_logger.LogInformation("Compute task failed due to missing requirements (queue: {QueueKey}, task: {TaskHash}, channel: {ChannelId})", queueKey, computeTask.TaskRefId, computeTask.ChannelId);
diff --git a/Engine/Source/Programs/Horde/Horde.Build/Protos/horde/tasks/compute_task.cs b/Engine/Source/Programs/Horde/Horde.Build/Protos/horde/tasks/compute_task.cs
new file mode 100644
--- /dev/null
+++ b/Engine/Source/Programs/Horde/Horde.Build/Protos/horde/tasks/compute_task.cs
@@ -0,0 +1,46 @@
+// Copyright Epic Games, Inc. All Rights Reserved.
+
+using EpicGames.Horde.Compute;
+using EpicGames.Horde.Storage;
+
+namespace HordeCommon.Rpc.Tasks
+{
+	/// <summary>
+	/// 
+	/// </summary>
+	partial class ComputeTaskResultMessage
+	{
+		/// <summary>
+		/// Constructor
+		/// </summary>
+		public ComputeTaskResultMessage(RefId resultRefId, ComputeTaskExecutionStatsMessage? executionStats)
+		{
+			ResultRefId = resultRefId;
+			ExecutionStats = executionStats;
+		}
+
+		/// <summary>
+		/// Constructor
+		/// </summary>
+		public ComputeTaskResultMessage(ComputeTaskOutcome outcome, string? detail = null)
+		{
+			Outcome = (int)outcome;
+			Detail = detail;
+		}
+	}
+
+	/// <summary>
+	/// 
+	/// </summary>
+	partial class ComputeTaskExecutionStatsMessage
+	{
+		/// <summary>
+		/// 
+		/// </summary>
+		/// <returns></returns>
+		public ComputeTaskExecutionStats ToNative()
+		{
+			return new ComputeTaskExecutionStats(StartTime.ToDateTime(), DownloadRefMs, DownloadInputMs, ExecMs, UploadLogMs, UploadOutputMs, UploadRefMs);
+		}
+	}
+}
diff --git a/Engine/Source/Programs/Shared/EpicGames.Horde/Compute/ComputeTaskResult.cs b/Engine/Source/Programs/Shared/EpicGames.Horde/Compute/ComputeTaskResult.cs
--- a/Engine/Source/Programs/Shared/EpicGames.Horde/Compute/ComputeTaskResult.cs
+++ b/Engine/Source/Programs/Shared/EpicGames.Horde/Compute/ComputeTaskResult.cs
@@ -34,7 +34,7 @@ namespace EpicGames.Horde.Compute
 		public CbObjectAttachment? OutputHash { get; set; }
 
 		/// <summary>
-		/// Constructor
+		/// Constructor for serialization
 		/// </summary>
 		private ComputeTaskResult()
 		{
diff --git a/Engine/Source/Programs/Shared/EpicGames.Horde/Compute/IComputeClient.cs b/Engine/Source/Programs/Shared/EpicGames.Horde/Compute/IComputeClient.cs
--- a/Engine/Source/Programs/Shared/EpicGames.Horde/Compute/IComputeClient.cs
+++ b/Engine/Source/Programs/Shared/EpicGames.Horde/Compute/IComputeClient.cs
@@ -2,6 +2,7 @@
 
 using EpicGames.Core;
 using EpicGames.Horde.Storage;
+using EpicGames.Serialization;
 using System;
 using System.Collections.Generic;
 using System.Threading;
@@ -100,47 +101,249 @@ namespace EpicGames.Horde.Compute
 	/// <summary>
 	/// Supplies information about the current execution state of a task
 	/// </summary>
-	public interface IComputeTaskInfo
+	public class ComputeTaskStatus
 	{
 		/// <summary>
 		/// Reference to the task decriptor that was requested
 		/// </summary>
-		public RefId TaskRefId { get; }
+		[CbField("h")]
+		public RefId TaskRefId { get; set; }
 
 		/// <summary>
 		/// Time that the event happened
 		/// </summary>
-		public DateTime Time { get; }
+		[CbField("t")]
+		public DateTime Time { get; set; }
 
 		/// <summary>
 		/// New state of the task
 		/// </summary>
-		public ComputeTaskState State { get; }
+		[CbField("s")]
+		public ComputeTaskState State { get; set; }
 
 		/// <summary>
 		/// Outcome of the task execution.Note that this reflects the outcome of the execution rather than the outcome of the task.        
 		/// </summary>
-		public ComputeTaskOutcome Outcome { get; }
+		[CbField("o")]
+		public ComputeTaskOutcome Outcome { get; set; }
 
 		/// <summary>
 		/// Additional information about the outcome of the task.Dependent on the value of Outcome.        
 		/// </summary>
-		public string? Detail { get; }
+		[CbField("d")]
+		public string? Detail { get; set; }
 
 		/// <summary>
 		/// Ref containing the task result
 		/// </summary>
-		public RefId? ResultRefId { get; }
+		[CbField("r")]
+		public RefId? ResultRefId { get; set; }
 
 		/// <summary>
 		/// When transitioning to the executing state, includes the name of the agent performing the work
 		/// </summary>
-		public string? AgentId { get; }
+		[CbField("a")]
+		public string? AgentId { get; set; }
 
 		/// <summary>
 		/// When transitioning to the executing state, includes the id of the lease assigned to the agent
 		/// </summary>
-		public string? LeaseId { get; }
+		[CbField("l")]
+		public string? LeaseId { get; set; }
+
+		/// <summary>
+		/// Stats about how the task was queued
+		/// </summary>
+		[CbField("qs")]
+		public ComputeTaskQueueStats? QueueStats { get; set; }
+
+		/// <summary>
+		/// Stats for execution of this job
+		/// </summary>
+		[CbField("es")]
+		public ComputeTaskExecutionStats? ExecutionStats { get; set; }
+	}
+
+	/// <summary>
+	/// Stats for queueing a task
+	/// </summary>
+	public class ComputeTaskQueueStats
+	{
+		/// <summary>
+		/// Names of scopes 
+		/// </summary>
+		public static readonly string[] ScopeNames =
+		{
+			"dispatched",
+			"complete"
+		};
+
+		private readonly int[] _scopes = new int[ScopeNames.Length];
+
+		/// <summary>
+		/// Time that the task was queued
+		/// </summary>
+		[CbField("t")]
+		public DateTime StartTime { get; set; }
+
+		/// <summary>
+		/// Timing values measured sequentially starting from StartTime, in milliseconds.
+		/// If the order or meaning of fields in this array are changed, the field name should be changed.
+		/// </summary>
+		[CbField("s")]
+		public int[] Scopes
+		{
+			get => _scopes;
+			set => value.AsSpan(0, Math.Max(_scopes.Length, value.Length)).CopyTo(_scopes);
+		}
+
+		/// <summary>
+		/// Time taken until the task was dequeued
+		/// </summary>
+		public int QueuedMs
+		{
+			get => _scopes[0];
+			set => _scopes[0] = value;
+		}
+
+		/// <summary>
+		/// Time that the task was executing
+		/// </summary>
+		public int ExecutingMs
+		{
+			get => _scopes[1];
+			set => _scopes[1] = value;
+		}
+
+		/// <summary>
+		/// Constructor
+		/// </summary>
+		public ComputeTaskQueueStats()
+		{
+		}
+
+		/// <summary>
+		/// Constructor
+		/// </summary>
+		public ComputeTaskQueueStats(DateTime startTime, int queuedMs, int executingMs)
+		{
+			StartTime = startTime;
+			QueuedMs = queuedMs;
+			ExecutingMs = executingMs;
+		}
+	}
+
+	/// <summary>
+	/// Stats for executing a task
+	/// </summary>
+	public class ComputeTaskExecutionStats
+	{
+		/// <summary>
+		/// Names of scopes 
+		/// </summary>
+		public static readonly string[] ScopeNames =
+		{
+			"download-ref",
+			"download-input",
+			"execute",
+			"upload-log",
+			"upload-output",
+			"upload-ref",
+		};
+
+		private readonly int[] _scopes = new int[ScopeNames.Length];
+
+		/// <summary>
+		/// Start time for execution
+		/// </summary>
+		[CbField("t")]
+		public DateTime StartTime { get; set; }
+
+		/// <summary>
+		/// Timing values measured sequentially starting from StartTime, in milliseconds. If the order or meaning of fields in this array are changed, the field name should be changed.
+		/// </summary>
+		[CbField("s")]
+		public int[] Scopes
+		{
+			get => _scopes;
+			set => value.AsSpan(0, Math.Max(_scopes.Length, value.Length)).CopyTo(_scopes);
+		}
+
+		/// <summary>
+		/// Time taken to download the job specification
+		/// </summary>
+		public int DownloadRefMs
+		{
+			get => _scopes[0];
+			set => _scopes[0] = value;
+		}
+
+		/// <summary>
+		/// Time taken to download and unpack the sandbox
+		/// </summary>
+		public int DownloadInputMs
+		{
+			get => _scopes[1];
+			set => _scopes[1] = value;
+		}
+
+		/// <summary>
+		/// Time taken to execute the task
+		/// </summary>
+		public int ExecMs
+		{
+			get => _scopes[2];
+			set => _scopes[2] = value;
+		}
+
+		/// <summary>
+		/// Time taken to upload log data
+		/// </summary>
+		public int UploadLogMs
+		{
+			get => _scopes[3];
+			set => _scopes[3] = value;
+		}
+
+		/// <summary>
+		/// Time taken to upload output data
+		/// </summary>
+		public int UploadOutputMs
+		{
+			get => _scopes[4];
+			set => _scopes[4] = value;
+		}
+
+		/// <summary>
+		/// Time taken to upload the output ref
+		/// </summary>
+		public int UploadRefMs
+		{
+			get => _scopes[5];
+			set => _scopes[5] = value;
+		}
+
+		/// <summary>
+		/// Constructor
+		/// </summary>
+		public ComputeTaskExecutionStats()
+		{
+		}
+
+		/// <summary>
+		/// Constructor
+		/// </summary>
+		public ComputeTaskExecutionStats(DateTime startTime, int downloadRefMs, int downloadInputMs, int execMs, int uploadLogMs, int uploadOutputMs, int uploadRefMs)
+		{
+			StartTime = startTime;
+
+			DownloadRefMs = downloadRefMs;
+			DownloadInputMs = downloadInputMs;
+			ExecMs = execMs;
+			UploadLogMs = uploadLogMs;
+			UploadOutputMs = uploadOutputMs;
+			UploadRefMs = uploadRefMs;
+		}
 	}
 
 	/// <summary>
@@ -175,7 +378,7 @@ namespace EpicGames.Horde.Compute
 		/// <param name="channelId">Channel to receive updates on</param>
 		/// <param name="cancellationToken">Cancellation token for the operation</param>
 		/// <returns></returns>
-		IAsyncEnumerable<IComputeTaskInfo> GetTaskUpdatesAsync(ClusterId clusterId, ChannelId channelId, CancellationToken cancellationToken = default);
+		IAsyncEnumerable<ComputeTaskStatus> GetTaskUpdatesAsync(ClusterId clusterId, ChannelId channelId, CancellationToken cancellationToken = default);
 	}
 
 	/// <summary>
diff --git a/Engine/Source/Programs/Shared/EpicGames.Horde/Storage/Impl/HttpStorageClient.cs b/Engine/Source/Programs/Shared/EpicGames.Horde/Storage/Impl/HttpStorageClient.cs
--- a/Engine/Source/Programs/Shared/EpicGames.Horde/Storage/Impl/HttpStorageClient.cs
+++ b/Engine/Source/Programs/Shared/EpicGames.Horde/Storage/Impl/HttpStorageClient.cs
@@ -117,10 +117,11 @@ namespace EpicGames.Horde.Storage.Impl
 		/// <inheritdoc/>
 		public async Task WriteBlobAsync(NamespaceId namespaceId, IoHash hash, Stream stream, CancellationToken cancellationToken)
 		{
-			using HttpRequestMessage request = new HttpRequestMessage(HttpMethod.Put, $"api/v1/blobs/{namespaceId}/{hash}");
-
 			StreamContent content = new StreamContent(stream);
 			content.Headers.ContentType = new MediaTypeHeaderValue("application/octet-stream");
+			content.Headers.ContentLength = stream.Length;
+
+			using HttpRequestMessage request = new HttpRequestMessage(HttpMethod.Put, $"api/v1/blobs/{namespaceId}/{hash}");
 			request.Content = content;
 
 			HttpResponseMessage response = await _httpClient.SendAsync(request, cancellationToken);
@@ -136,10 +137,11 @@ namespace EpicGames.Horde.Storage.Impl
 		/// <inheritdoc/>
 		public async Task<IoHash> WriteBlobAsync(NamespaceId namespaceId, Stream stream, CancellationToken cancellationToken)
 		{
-			using HttpRequestMessage request = new HttpRequestMessage(HttpMethod.Post, $"api/v1/blobs/{namespaceId}");
-
 			StreamContent content = new StreamContent(stream);
 			content.Headers.ContentType = new MediaTypeHeaderValue("application/octet-stream");
+			content.Headers.ContentLength = stream.Length;
+
+			using HttpRequestMessage request = new HttpRequestMessage(HttpMethod.Post, $"api/v1/blobs/{namespaceId}");
 			request.Content = content;
 
 			HttpResponseMessage response = await _httpClient.SendAsync(request, cancellationToken);
diff --git a/Engine/Source/Programs/Shared/EpicGames.Redis/RedisSerializer.cs b/Engine/Source/Programs/Shared/EpicGames.Redis/RedisSerializer.cs
--- a/Engine/Source/Programs/Shared/EpicGames.Redis/RedisSerializer.cs
+++ b/Engine/Source/Programs/Shared/EpicGames.Redis/RedisSerializer.cs
@@ -131,46 +131,74 @@ namespace EpicGames.Redis
 			return new KeyValuePair<Type, object>(typeof(T), new RedisNativeConverter<T>(fromRedisValueFunc, toRedisValueFunc));
 		}
 
+		static readonly Dictionary<Type, Type> s_typeToConverterType = new Dictionary<Type, Type>();
+
 		/// <summary>
-		/// Static class for caching converter lookups
+		/// Register a custom converter for a particular type
 		/// </summary>
-		/// <typeparam name="T"></typeparam>
-		class CachedConverter<T>
+		public static void RegisterConverter<T, TConverter>() where TConverter : IRedisConverter<T>
 		{
-			public static IRedisConverter<T> Converter = CreateConverter();
-
-			static IRedisConverter<T> CreateConverter()
+			lock (s_typeToConverterType)
 			{
-				Type type = typeof(T);
+				if (!s_typeToConverterType.TryGetValue(typeof(T), out Type? converterType) || converterType != typeof(TConverter))
+				{
+					s_typeToConverterType.Add(typeof(T), typeof(TConverter));
+				}
+			}
+		}
+
+		/// <summary>
+		/// Creates a converter for a given type
+		/// </summary>
+		static IRedisConverter<T> CreateConverter<T>()
+		{
+			Type type = typeof(T);
 
-				// Check for a custom converter
-				RedisConverterAttribute? attribute = type.GetCustomAttribute<RedisConverterAttribute>();
-				if (attribute != null)
+			// Check for a registered converter type
+			lock (s_typeToConverterType)
+			{
+				if (s_typeToConverterType.TryGetValue(type, out Type? converterType))
 				{
-					Type converterType = attribute.ConverterType;
-					if (converterType.IsGenericTypeDefinition)
-					{
-						converterType = converterType.MakeGenericType(type);
-					}
 					return (IRedisConverter<T>)Activator.CreateInstance(converterType)!;
 				}
+			}
 
-				// Check for known basic types
-				object? nativeConverter;
-				if (s_nativeConverters.TryGetValue(typeof(T), out nativeConverter))
-                {
-					return (IRedisConverter<T>)nativeConverter;
-                }
-
-				// Check if there's a regular converter we can use to convert to/from a string
-				TypeConverter? converter = TypeDescriptor.GetConverter(type);
-				if (converter != null && converter.CanConvertFrom(typeof(string)) && converter.CanConvertTo(typeof(string)))
+			// Check for a custom converter
+			RedisConverterAttribute? attribute = type.GetCustomAttribute<RedisConverterAttribute>();
+			if (attribute != null)
+			{
+				Type converterType = attribute.ConverterType;
+				if (converterType.IsGenericTypeDefinition)
 				{
-					return new RedisStringConverter<T>(converter);
+					converterType = converterType.MakeGenericType(type);
 				}
+				return (IRedisConverter<T>)Activator.CreateInstance(converterType)!;
+			}
+
+			// Check for known basic types
+			object? nativeConverter;
+			if (s_nativeConverters.TryGetValue(typeof(T), out nativeConverter))
+			{
+				return (IRedisConverter<T>)nativeConverter;
+			}
 
-				throw new Exception($"Unable to find Redis converter for {type.Name}");
+			// Check if there's a regular converter we can use to convert to/from a string
+			TypeConverter? converter = TypeDescriptor.GetConverter(type);
+			if (converter != null && converter.CanConvertFrom(typeof(string)) && converter.CanConvertTo(typeof(string)))
+			{
+				return new RedisStringConverter<T>(converter);
 			}
+
+			throw new Exception($"Unable to find Redis converter for {type.Name}");
+		}
+
+		/// <summary>
+		/// Static class for caching converter lookups
+		/// </summary>
+		/// <typeparam name="T"></typeparam>
+		class CachedConverter<T>
+		{
+			public static IRedisConverter<T> Converter = CreateConverter<T>();
 		}
 
 		/// <summary>
