Horde Storage - Updated blake 3 references to same version as used in EpicGames.IoHash as we still do use blake3 directly.

diff --git a/Engine/Source/Programs/Horde/HordeStorage/Jupiter.Common/Implementation/BlobIdentifier.cs b/Engine/Source/Programs/Horde/HordeStorage/Jupiter.Common/Implementation/BlobIdentifier.cs
--- a/Engine/Source/Programs/Horde/HordeStorage/Jupiter.Common/Implementation/BlobIdentifier.cs
+++ b/Engine/Source/Programs/Horde/HordeStorage/Jupiter.Common/Implementation/BlobIdentifier.cs
@@ -93,7 +93,7 @@ namespace Jupiter.Implementation
             }
             
             // we only keep the first 20 bytes of the Blake3 hash
-            Span<byte> hash = blake3Hash.AsSpan().Slice(0, 20);
+            Span<byte> hash = blake3Hash.AsSpanUnsafe().Slice(0, 20);
             return new BlobIdentifier(hash.ToArray());
         }
 
@@ -114,7 +114,7 @@ namespace Jupiter.Implementation
             }
 
             // we only keep the first 20 bytes of the Blake3 hash
-            Span<byte> hash = blake3Hash.AsSpan().Slice(0, 20);
+            Span<byte> hash = blake3Hash.AsSpanUnsafe().Slice(0, 20);
             return new BlobIdentifier(hash.ToArray());
         }
 
@@ -137,7 +137,7 @@ namespace Jupiter.Implementation
             Hash blake3Hash = hasher.Finalize();
 
             // we only keep the first 20 bytes of the Blake3 hash
-            byte[] hash = blake3Hash.AsSpan().Slice(0, 20).ToArray();
+            byte[] hash = blake3Hash.AsSpanUnsafe().Slice(0, 20).ToArray();
             return new BlobIdentifier(hash);
         }
 
diff --git a/Engine/Source/Programs/Horde/HordeStorage/Jupiter.Common/Implementation/ContentHash.cs b/Engine/Source/Programs/Horde/HordeStorage/Jupiter.Common/Implementation/ContentHash.cs
--- a/Engine/Source/Programs/Horde/HordeStorage/Jupiter.Common/Implementation/ContentHash.cs
+++ b/Engine/Source/Programs/Horde/HordeStorage/Jupiter.Common/Implementation/ContentHash.cs
@@ -96,7 +96,7 @@ namespace Jupiter.Implementation
             Hash blake3Hash = hasher.Finalize();
 
             // we only keep the first 20 bytes of the Blake3 hash
-            Span<byte> hash = blake3Hash.AsSpan().Slice(0, HashLength);
+            Span<byte> hash = blake3Hash.AsSpanUnsafe().Slice(0, HashLength);
 
             return new ContentHash(hash.ToArray());
         }
