Horde Storage - Further fixes required after blake3 upgrade

diff --git a/Engine/Source/Programs/Horde/HordeStorage/Jupiter.Common/Implementation/TreeHash.cs b/Engine/Source/Programs/Horde/HordeStorage/Jupiter.Common/Implementation/TreeHash.cs
--- a/Engine/Source/Programs/Horde/HordeStorage/Jupiter.Common/Implementation/TreeHash.cs
+++ b/Engine/Source/Programs/Horde/HordeStorage/Jupiter.Common/Implementation/TreeHash.cs
@@ -73,7 +73,7 @@ namespace Jupiter.Implementation
             Hash blake3Hash = Hasher.Hash(bytes);
 
             // we only keep the first 20 bytes of the Blake3 hash
-            TreeHash hash = new TreeHash(blake3Hash.AsSpan().Slice(0, HashLength).ToArray());
+            TreeHash hash = new TreeHash(blake3Hash.AsSpanUnsafe().Slice(0, HashLength).ToArray());
             return hash;
         }
     }
