Horde: Improve logging when scaling AWS instances

diff --git a/Engine/Source/Programs/Horde/Horde.Build/Fleet/Impl/Managers/Fleet/AwsFleetManager.cs b/Engine/Source/Programs/Horde/Horde.Build/Fleet/Impl/Managers/Fleet/AwsFleetManager.cs
--- a/Engine/Source/Programs/Horde/Horde.Build/Fleet/Impl/Managers/Fleet/AwsFleetManager.cs
+++ b/Engine/Source/Programs/Horde/Horde.Build/Fleet/Impl/Managers/Fleet/AwsFleetManager.cs
@@ -93,8 +93,8 @@ namespace Horde.Build.Services.Impl
 
 				if (startRequest.InstanceIds.Count < count)
 				{
-					_logger.LogInformation("Unable to expand pool with the requested number of instances. " +
-					                      "Num requested instances to add {RequestedCount}. Actual instances started {InstancesStarted}", count, startRequest.InstanceIds.Count);
+					_logger.LogInformation("Unable to expand pool {PoolName} with the requested number of instances. " +
+					                      "Num requested instances to add {RequestedCount}. Actual instances started {InstancesStarted}", pool.Name, count, startRequest.InstanceIds.Count);
 				}
 			}
 		}
@@ -123,11 +123,11 @@ namespace Horde.Build.Services.Impl
 				IAuditLogChannel<AgentId> agentLogger = _agentCollection.GetLogger(agent.Id);
 				if (await _agentCollection.TryUpdateSettingsAsync(agent, bRequestShutdown: true, shutdownReason: "Autoscaler") != null)
 				{
-					agentLogger.LogInformation("Marked for shutdown due to autoscaling (currently {NumLeases} leases outstanding)", agent.Leases.Count);
+					agentLogger.LogInformation("Marked {AgentId} in pool {PoolName} for shutdown due to autoscaling (currently {NumLeases} leases outstanding)", agent.Id, pool.Name, agent.Leases.Count);
 				}
 				else
 				{
-					agentLogger.LogError("Unable to mark agent for shutdown due to autoscaling");
+					agentLogger.LogError("Unable to mark agent {AgentId} in pool {PoolName} for shutdown due to autoscaling", agent.Id, pool.Name);
 				}
 			}
 		}
