Horde: Do not list suspects as having declined issues if they were never promoted.

diff --git a/Engine/Source/Programs/Horde/Horde.Build/Notifications/Impl/SlackNotificationSink.cs b/Engine/Source/Programs/Horde/Horde.Build/Notifications/Impl/SlackNotificationSink.cs
--- a/Engine/Source/Programs/Horde/Horde.Build/Notifications/Impl/SlackNotificationSink.cs
+++ b/Engine/Source/Programs/Horde/Horde.Build/Notifications/Impl/SlackNotificationSink.cs
@@ -773,7 +773,11 @@ namespace Horde.Build.Notifications.Impl
 				List<string> declinedLines = new List<string>();
 				foreach (IIssueSuspect suspect in details.Suspects)
 				{
-					if (suspect.DeclinedAt == null)
+					if (!details.Issue.Promoted)
+					{
+						declinedLines.Add($"Possibly {await FormatNameAsync(suspect.AuthorId)} (CL {suspect.Change})");
+					}
+					else if (suspect.DeclinedAt == null)
 					{
 						declinedLines.Add($":heavy_minus_sign: Ignored by {await FormatNameAsync(suspect.AuthorId)} (CL {suspect.Change})");
 					}
