Horde: Fix matching of asset log warnings.

diff --git a/Engine/Source/Programs/Horde/Horde.Agent.Tests/LogParserTests.cs b/Engine/Source/Programs/Horde/Horde.Agent.Tests/LogParserTests.cs
--- a/Engine/Source/Programs/Horde/Horde.Agent.Tests/LogParserTests.cs
+++ b/Engine/Source/Programs/Horde/Horde.Agent.Tests/LogParserTests.cs
@@ -409,6 +409,24 @@ namespace Horde.Agent.Tests
 			}
 		}
 
+		[TestMethod]
+		public void AssetLogEventMatcher()
+		{
+			string[] lines =
+			{
+				@"LogBlueprint: Warning: [AssetLog] D:\\build\\++UE5\\Sync\\QAGame\\Plugins\\NiagaraFluids\\Content\\Blueprints\\Phsyarum_BP.uasset: [Compiler] Fill Texture 2D : Usage of 'Fill Texture 2D' has been deprecated. This function has been replaced by object user variables on the emitter to specify render targets to fill with data.",
+			};
+
+			List<LogEvent> events = Parse(String.Join("\n", lines), new DirectoryReference("D:\\build\\++UE5\\Sync"));
+			Assert.AreEqual(1, events.Count);
+			Assert.AreEqual(events[0].Id, KnownLogEvents.Engine_AssetLog);
+
+			LogValue assetProperty = (LogValue)events[0].Properties!["asset"];
+			Assert.AreEqual("Asset", assetProperty.Type);
+			Assert.AreEqual(@"//UE4/Main/QAGame/Plugins/NiagaraFluids/Content/Blueprints/Phsyarum_BP.uasset@12345", assetProperty.Properties!["depotPath"].ToString());
+			Assert.AreEqual(@"QAGame/Plugins/NiagaraFluids/Content/Blueprints/Phsyarum_BP.uasset", assetProperty.Properties!["relativePath"].ToString());
+		}
+
 		[TestMethod]
 		public void ClangEventMatcher()
 		{
diff --git a/Engine/Source/Programs/Horde/Horde.Agent/Parser/LogEventMarkup.cs b/Engine/Source/Programs/Horde/Horde.Agent/Parser/LogEventMarkup.cs
--- a/Engine/Source/Programs/Horde/Horde.Agent/Parser/LogEventMarkup.cs
+++ b/Engine/Source/Programs/Horde/Horde.Agent/Parser/LogEventMarkup.cs
@@ -56,7 +56,7 @@ namespace Horde.Agent.Parser
 				FileReference location = FileReference.Combine(DirectoryReference.Combine(context.WorkspaceDir, "Engine", "Binaries", "Win64"), group.Value);
 				if (location.IsUnderDirectory(context.WorkspaceDir) && !location.ContainsName("Intermediate", context.WorkspaceDir))
 				{
-					string relativePath = location.MakeRelativeTo(context.WorkspaceDir);
+					string relativePath = location.MakeRelativeTo(context.WorkspaceDir).Replace('\\', '/');
 					string depotPath = $"{context.PerforceStream.TrimEnd('/')}/{relativePath.Replace(Path.DirectorySeparatorChar, '/')}@{context.PerforceChange.Value}";
 
 					Dictionary<string, object> properties = new Dictionary<string, object>();
diff --git a/Engine/Source/Programs/Horde/Horde.Agent/Parser/Matchers/ContentEventMatcher.cs b/Engine/Source/Programs/Horde/Horde.Agent/Parser/Matchers/ContentEventMatcher.cs
--- a/Engine/Source/Programs/Horde/Horde.Agent/Parser/Matchers/ContentEventMatcher.cs
+++ b/Engine/Source/Programs/Horde/Horde.Agent/Parser/Matchers/ContentEventMatcher.cs
@@ -3,6 +3,7 @@
 using System;
 using System.Text.RegularExpressions;
 using EpicGames.Core;
+using Horde.Agent.Parser.Interfaces;
 using Microsoft.Extensions.Logging;
 
 namespace Horde.Agent.Parser.Matchers
@@ -20,6 +21,13 @@ namespace Horde.Agent.Parser.Matchers
 			@"(?<asset>(?:[a-zA-Z]:)?[^:]+(?:.uasset|.umap)):\s*" + 
 			@"(?<message>.*)";
 
+		private readonly ILogContext _context;
+
+		public ContentEventMatcher(ILogContext context)
+		{
+			_context = context;
+		}
+
 		/// <inheritdoc/>
 		public LogEventMatch? Match(ILogCursor input)
 		{
@@ -34,7 +42,7 @@ namespace Horde.Agent.Parser.Matchers
 
 					builder.Annotate(match.Groups["channel"], LogEventMarkup.Channel);
 					builder.Annotate(match.Groups["severity"], LogEventMarkup.Severity);
-					//					Builder.Annotate(Match.Groups["asset"]).MarkAsAsset(Context);
+					builder.AnnotateAsset(match.Groups["asset"], _context);
 					builder.Annotate(match.Groups["message"], LogEventMarkup.Message);
 
 					return builder.ToMatch(LogEventPriority.AboveNormal, GetLogLevelFromSeverity(match), KnownLogEvents.Engine_AssetLog);
