Horde: Include a summary of the first few errors and warnings with build health notifications.

diff --git a/Engine/Source/Programs/Horde/Horde.Build.Tests/IssueServiceTests.cs b/Engine/Source/Programs/Horde/Horde.Build.Tests/IssueServiceTests.cs
--- a/Engine/Source/Programs/Horde/Horde.Build.Tests/IssueServiceTests.cs
+++ b/Engine/Source/Programs/Horde/Horde.Build.Tests/IssueServiceTests.cs
@@ -397,7 +397,7 @@ namespace Horde.Build.Tests
 				await UpdateCompleteStep(job, 0, 0, JobStepOutcome.Warnings);
 
 				ILogFile? log = await LogFileService.GetLogFileAsync(job.Batches[0].Steps[0].LogId!.Value);
-				List<ILogEvent> events = await LogFileService.FindLogEventsAsync(log!);
+				List<ILogEvent> events = await LogFileService.FindEventsAsync(log!);
 				Assert.AreEqual(1, events.Count);
 				Assert.AreEqual(2, events[0].LineCount);
 
@@ -1225,7 +1225,7 @@ namespace Horde.Build.Tests
 				await UpdateCompleteStep(job, 0, 0, JobStepOutcome.Failure);
 
 				ILogFile? log = await LogFileService.GetLogFileAsync(job.Batches[0].Steps[0].LogId!.Value);
-				List<ILogEvent> events = await LogFileService.FindLogEventsAsync(log!);
+				List<ILogEvent> events = await LogFileService.FindEventsAsync(log!);
 				Assert.AreEqual(1, events.Count);
 				Assert.AreEqual(40, events[0].LineCount);
 
diff --git a/Engine/Source/Programs/Horde/Horde.Build/Controllers/LogsController.cs b/Engine/Source/Programs/Horde/Horde.Build/Controllers/LogsController.cs
--- a/Engine/Source/Programs/Horde/Horde.Build/Controllers/LogsController.cs
+++ b/Engine/Source/Programs/Horde/Horde.Build/Controllers/LogsController.cs
@@ -303,7 +303,7 @@ namespace Horde.Build.Controllers
 				return Forbid();
 			}
 
-			List<ILogEvent> logEvents = await _logFileService.FindLogEventsAsync(logFile, index, count);
+			List<ILogEvent> logEvents = await _logFileService.FindEventsAsync(logFile, null, index, count);
 
 			Dictionary<ObjectId, int?> spanIdToIssueId = new Dictionary<ObjectId, int?>();
 
diff --git a/Engine/Source/Programs/Horde/Horde.Build/Issues/Impl/IssueService.cs b/Engine/Source/Programs/Horde/Horde.Build/Issues/Impl/IssueService.cs
--- a/Engine/Source/Programs/Horde/Horde.Build/Issues/Impl/IssueService.cs
+++ b/Engine/Source/Programs/Horde/Horde.Build/Issues/Impl/IssueService.cs
@@ -615,7 +615,7 @@ namespace Horde.Build.Services.Impl
 			}
 
 			// Process all the events for this step
-			List<ILogEvent> stepEvents = await _logFileService.FindLogEventsAsync(logFile);
+			List<ILogEvent> stepEvents = await _logFileService.FindEventsAsync(logFile);
 
 			// Create a new list of event fingerprints for each one
 			bool nonSystemicErrors = false;
diff --git a/Engine/Source/Programs/Horde/Horde.Build/Logs/ILogEventCollection.cs b/Engine/Source/Programs/Horde/Horde.Build/Logs/ILogEventCollection.cs
--- a/Engine/Source/Programs/Horde/Horde.Build/Logs/ILogEventCollection.cs
+++ b/Engine/Source/Programs/Horde/Horde.Build/Logs/ILogEventCollection.cs
@@ -31,10 +31,11 @@ namespace Horde.Build.Collections
 		/// Finds events within a log file
 		/// </summary>
 		/// <param name="logId">Unique id of the log containing this event</param>
+		/// <param name="spanId">Optional span to filter events by</param>
 		/// <param name="index">Start index within the matching results</param>
 		/// <param name="count">Maximum number of results to return</param>
 		/// <returns>List of events matching the query</returns>
-		Task<List<ILogEvent>> FindAsync(LogId logId, int? index = null, int? count = null);
+		Task<List<ILogEvent>> FindAsync(LogId logId, ObjectId? spanId, int? index = null, int? count = null);
 
 		/// <summary>
 		/// Finds a list of events for a set of spans
diff --git a/Engine/Source/Programs/Horde/Horde.Build/Logs/Impl/LogEventCollection.cs b/Engine/Source/Programs/Horde/Horde.Build/Logs/Impl/LogEventCollection.cs
--- a/Engine/Source/Programs/Horde/Horde.Build/Logs/Impl/LogEventCollection.cs
+++ b/Engine/Source/Programs/Horde/Horde.Build/Logs/Impl/LogEventCollection.cs
@@ -133,13 +133,17 @@ namespace Horde.Build.Collections.Impl
 		}
 
 		/// <inheritdoc/>
-		public async Task<List<ILogEvent>> FindAsync(LogId logId, int? index = null, int? count = null)
+		public async Task<List<ILogEvent>> FindAsync(LogId logId, ObjectId? spanId = null, int? index = null, int? count = null)
 		{
 			_logger.LogInformation("Querying for log events for log {LogId} creation time {CreateTime}", logId, logId.Value.CreationTime);
 
 			FilterDefinitionBuilder<LogEventDocument> builder = Builders<LogEventDocument>.Filter;
 
 			FilterDefinition<LogEventDocument> filter = builder.Eq(x => x.Id.LogId, logId);
+			if(spanId != null)
+			{
+				filter &= builder.Eq(x => x.SpanId, spanId.Value);
+			}
 
 			IFindFluent<LogEventDocument, LogEventDocument> results = _logEvents.Find(filter).SortBy(x => x.Id);
 			if (index != null)
@@ -155,19 +159,6 @@ namespace Horde.Build.Collections.Impl
 			return eventDocuments.ConvertAll<ILogEvent>(x => x);
 		}
 
-		/// <inheritdoc/>
-		public async Task<List<ILogEvent>> FindEventsForSpanAsync(ObjectId spanId, LogId[]? logIds, int index, int count)
-		{
-			FilterDefinition<LogEventDocument> filter = Builders<LogEventDocument>.Filter.Eq(x => x.SpanId, spanId);
-			if (logIds != null && logIds.Length > 0)
-			{
-				filter &= Builders<LogEventDocument>.Filter.In(x => x.Id.LogId, logIds);
-			}
-
-			List<LogEventDocument> eventDocuments = await _logEvents.Find(filter).Skip(index).Limit(count).ToListAsync();
-			return eventDocuments.ConvertAll<ILogEvent>(x => x);
-		}
-
 		/// <inheritdoc/>
 		public async Task<List<ILogEvent>> FindEventsForSpansAsync(IEnumerable<ObjectId> spanIds, LogId[]? logIds, int index, int count)
 		{
diff --git a/Engine/Source/Programs/Horde/Horde.Build/Logs/Impl/LogFileService.cs b/Engine/Source/Programs/Horde/Horde.Build/Logs/Impl/LogFileService.cs
--- a/Engine/Source/Programs/Horde/Horde.Build/Logs/Impl/LogFileService.cs
+++ b/Engine/Source/Programs/Horde/Horde.Build/Logs/Impl/LogFileService.cs
@@ -117,10 +117,11 @@ namespace Horde.Build.Services
 		/// Find events for a particular log file
 		/// </summary>
 		/// <param name="logFile">The log file instance</param>
+		/// <param name="spanId">Issue span to return events for</param>
 		/// <param name="index">Index of the first event to retrieve</param>
 		/// <param name="count">Number of events to retrieve</param>
 		/// <returns>List of log events</returns>
-		Task<List<ILogEvent>> FindLogEventsAsync(ILogFile logFile, int? index = null, int? count = null);
+		Task<List<ILogEvent>> FindEventsAsync(ILogFile logFile, ObjectId? spanId = null, int? index = null, int? count = null);
 
 		/// <summary>
 		/// Adds events to a log span
@@ -777,9 +778,9 @@ namespace Horde.Build.Services
 		}
 
 		/// <inheritdoc/>
-		public Task<List<ILogEvent>> FindLogEventsAsync(ILogFile logFile, int? index = null, int? count = null)
+		public Task<List<ILogEvent>> FindEventsAsync(ILogFile logFile, ObjectId? spanId = null, int? index = null, int? count = null)
 		{
-			return _logEvents.FindAsync(logFile.Id, index, count);
+			return _logEvents.FindAsync(logFile.Id, spanId, index, count);
 		}
 
 		class LogEventLine : ILogEventLine
diff --git a/Engine/Source/Programs/Horde/Horde.Build/Notifications/Impl/NotificationService.cs b/Engine/Source/Programs/Horde/Horde.Build/Notifications/Impl/NotificationService.cs
--- a/Engine/Source/Programs/Horde/Horde.Build/Notifications/Impl/NotificationService.cs
+++ b/Engine/Source/Programs/Horde/Horde.Build/Notifications/Impl/NotificationService.cs
@@ -606,7 +606,7 @@ namespace Horde.Build.Notifications.Impl
 				return;
 			}
 
-			List<ILogEvent> jobStepEvents = await _logFileService.FindLogEventsAsync(logFile);
+			List<ILogEvent> jobStepEvents = await _logFileService.FindEventsAsync(logFile);
 			List<ILogEventData> jobStepEventData = new List<ILogEventData>();
 			foreach (ILogEvent logEvent in jobStepEvents)
 			{
diff --git a/Engine/Source/Programs/Horde/Horde.Build/Notifications/Impl/SlackNotificationSink.cs b/Engine/Source/Programs/Horde/Horde.Build/Notifications/Impl/SlackNotificationSink.cs
--- a/Engine/Source/Programs/Horde/Horde.Build/Notifications/Impl/SlackNotificationSink.cs
+++ b/Engine/Source/Programs/Horde/Horde.Build/Notifications/Impl/SlackNotificationSink.cs
@@ -34,6 +34,7 @@ using MongoDB.Driver;
 
 namespace Horde.Build.Notifications.Impl
 {
+	using LogId = ObjectId<ILogFile>;
 	using UserId = ObjectId<IUser>;
 
 	/// <summary>
@@ -195,6 +196,7 @@ namespace Horde.Build.Notifications.Impl
 
 		readonly IIssueService _issueService;
 		readonly IUserCollection _userCollection;
+		readonly LogFileService _logFileService;
 		readonly StreamService _streamService;
 		readonly IWebHostEnvironment _environment;
 		readonly ServerSettings _settings;
@@ -211,17 +213,11 @@ namespace Horde.Build.Notifications.Impl
 		/// <summary>
 		/// Constructor
 		/// </summary>
-		/// <param name="mongoService"></param>
-		/// <param name="issueService"></param>
-		/// <param name="userCollection">The user collection</param>
-		/// <param name="streamService"></param>
-		/// <param name="environment"></param>
-		/// <param name="settings">The current configuration settings</param>
-		/// <param name="logger">Logging device</param>
-		public SlackNotificationSink(MongoService mongoService, IIssueService issueService, IUserCollection userCollection, StreamService streamService, IWebHostEnvironment environment, IOptions<ServerSettings> settings, ILogger<SlackNotificationSink> logger)
+		public SlackNotificationSink(MongoService mongoService, IIssueService issueService, IUserCollection userCollection, LogFileService logFileService, StreamService streamService, IWebHostEnvironment environment, IOptions<ServerSettings> settings, ILogger<SlackNotificationSink> logger)
 		{
 			_issueService = issueService;
 			_userCollection = userCollection;
+			_logFileService = logFileService;
 			_streamService = streamService;
 			_environment = environment;
 			_settings = settings.Value;
@@ -673,6 +669,34 @@ namespace Horde.Build.Notifications.Impl
 			string streamList = StringUtils.FormatList(details.Spans.Select(x => $"*{x.StreamName}*").Distinct(StringComparer.OrdinalIgnoreCase).OrderBy(x => x, StringComparer.OrdinalIgnoreCase));
 			attachment.Blocks.Add(new SectionBlock($"Occurring in {streamList}"));
 
+			IIssueSpan? lastSpan = details.Spans.OrderByDescending(x => x.LastFailure.StepTime).FirstOrDefault();
+			if (lastSpan != null && lastSpan.LastFailure.LogId != null)
+			{
+				LogId logId = lastSpan.LastFailure.LogId.Value;
+				ILogFile? logFile = await _logFileService.GetLogFileAsync(logId);
+				if(logFile != null)
+				{
+					List<ILogEvent> events = await _logFileService.FindEventsAsync(logFile, lastSpan.Id, 0, 20);
+					if (events.Any(x => x.Severity == EventSeverity.Error))
+					{
+						events.RemoveAll(x => x.Severity == EventSeverity.Warning);
+					}
+
+					List<string> eventStrings = new List<string>();
+					for (int idx = 0; idx < Math.Min(events.Count, 3); idx++)
+					{
+						ILogEventData data = await _logFileService.GetEventDataAsync(logFile, events[idx].LineIndex, events[idx].LineCount);
+						eventStrings.Add($"```{data.Message}```");
+					}
+					if (events.Count > 3)
+					{
+						eventStrings.Add("```...```");
+					}
+
+					attachment.Blocks.Add(new SectionBlock(String.Join("\n", eventStrings)));
+				}
+			}
+
 			if (issue.FixChange != null)
 			{
 				IIssueStep? fixFailedStep = issue.FindFixFailedStep(details.Spans);
