Horde: Exclude skipped steps from notifications about jobs failing.

diff --git a/Engine/Source/Programs/Horde/Horde.Build/Notifications/Impl/SlackNotificationSink.cs b/Engine/Source/Programs/Horde/Horde.Build/Notifications/Impl/SlackNotificationSink.cs
--- a/Engine/Source/Programs/Horde/Horde.Build/Notifications/Impl/SlackNotificationSink.cs
+++ b/Engine/Source/Programs/Horde/Horde.Build/Notifications/Impl/SlackNotificationSink.cs
@@ -382,15 +382,18 @@ namespace Horde.Build.Notifications.Impl
 				IReadOnlyDictionary<NodeRef, IJobStep> nodeToStep = job.GetStepForNodeMap();
 				foreach ((NodeRef nodeRef, IJobStep step) in nodeToStep)
 				{
-					INode stepNode = graph.GetNode(nodeRef);
-					string stepName = $"<{jobLink}?step={step.Id}|{stepNode.Name}>";
-					if (step.Outcome == JobStepOutcome.Failure)
+					if (step.State == JobStepState.Completed)
 					{
-						failedStepStrings.Add(stepName);
-					}
-					else if (step.Outcome == JobStepOutcome.Warnings)
-					{
-						warningStepStrings.Add(stepName);
+						INode stepNode = graph.GetNode(nodeRef);
+						string stepName = $"<{jobLink}?step={step.Id}|{stepNode.Name}>";
+						if (step.Outcome == JobStepOutcome.Failure)
+						{
+							failedStepStrings.Add(stepName);
+						}
+						else if (step.Outcome == JobStepOutcome.Warnings)
+						{
+							warningStepStrings.Add(stepName);
+						}
 					}
 				}
 
