Horde: Fix DI error running tests.

diff --git a/Engine/Source/Programs/Horde/Horde.Build/Notifications/Impl/SlackNotificationSink.cs b/Engine/Source/Programs/Horde/Horde.Build/Notifications/Impl/SlackNotificationSink.cs
--- a/Engine/Source/Programs/Horde/Horde.Build/Notifications/Impl/SlackNotificationSink.cs
+++ b/Engine/Source/Programs/Horde/Horde.Build/Notifications/Impl/SlackNotificationSink.cs
@@ -196,7 +196,7 @@ namespace Horde.Build.Notifications.Impl
 
 		readonly IIssueService _issueService;
 		readonly IUserCollection _userCollection;
-		readonly LogFileService _logFileService;
+		readonly ILogFileService _logFileService;
 		readonly StreamService _streamService;
 		readonly IWebHostEnvironment _environment;
 		readonly ServerSettings _settings;
@@ -213,7 +213,7 @@ namespace Horde.Build.Notifications.Impl
 		/// <summary>
 		/// Constructor
 		/// </summary>
-		public SlackNotificationSink(MongoService mongoService, IIssueService issueService, IUserCollection userCollection, LogFileService logFileService, StreamService streamService, IWebHostEnvironment environment, IOptions<ServerSettings> settings, ILogger<SlackNotificationSink> logger)
+		public SlackNotificationSink(MongoService mongoService, IIssueService issueService, IUserCollection userCollection, ILogFileService logFileService, StreamService streamService, IWebHostEnvironment environment, IOptions<ServerSettings> settings, ILogger<SlackNotificationSink> logger)
 		{
 			_issueService = issueService;
 			_userCollection = userCollection;
