Horde: Change the tool namespace to one that already exists.

diff --git a/Engine/Source/Programs/Horde/Horde.Build/ServerSettings.cs b/Engine/Source/Programs/Horde/Horde.Build/ServerSettings.cs
--- a/Engine/Source/Programs/Horde/Horde.Build/ServerSettings.cs
+++ b/Engine/Source/Programs/Horde/Horde.Build/ServerSettings.cs
@@ -515,7 +515,7 @@ namespace Horde.Build
 		/// <summary>
 		/// Namespace to use for storing tools
 		/// </summary>
-		public NamespaceId ToolNamespaceId { get; set; } = new NamespaceId("horde.tools");
+		public NamespaceId ToolNamespaceId { get; set; } = new NamespaceId("horde.p4");
 
 		/// <summary>
 		/// Lazily computed timezone value
