HordeAgent: Wait for exit in task executor before getting exit code, as the process on Mac or Linux it is not guaranteed to have terminated immeditally after the output streams are closed like on windows

#rnx

diff --git a/Engine/Source/Programs/Horde/Horde.Agent/Execution/ComputeTaskExecutor.cs b/Engine/Source/Programs/Horde/Horde.Agent/Execution/ComputeTaskExecutor.cs
--- a/Engine/Source/Programs/Horde/Horde.Agent/Execution/ComputeTaskExecutor.cs
+++ b/Engine/Source/Programs/Horde/Horde.Agent/Execution/ComputeTaskExecutor.cs
@@ -144,6 +144,7 @@ namespace Horde.Agent.Execution
 						_logger.LogInformation("stderr: {Line}", line);
 					}
 
+					process.WaitForExit();
 					_logger.LogInformation("exit: {ExitCode}", process.ExitCode);
 
 					ComputeTaskResult result = new ComputeTaskResult(process.ExitCode);
