Horde: Additional issue modal tweaks
#skipci

diff --git a/Engine/Source/Programs/Horde/HordeDashboard/src/base/utilities/timeUtils.ts b/Engine/Source/Programs/Horde/HordeDashboard/src/base/utilities/timeUtils.ts
--- a/Engine/Source/Programs/Horde/HordeDashboard/src/base/utilities/timeUtils.ts
+++ b/Engine/Source/Programs/Horde/HordeDashboard/src/base/utilities/timeUtils.ts
@@ -335,6 +335,38 @@ export const getLeaseElapsed = (lease: GetAgentLeaseResponse | undefined): strin
 };
 
 
+export const getHumanTime = (timeIn: Date | string | undefined): string => {
+
+    if (!timeIn) {
+        return "";
+    }
+
+    const now = moment.utc().tz(displayTimeZone());
+    const time = moment(timeIn).tz(displayTimeZone());
+
+    const nowDay = now.dayOfYear();
+    const timeDay = time.dayOfYear();
+
+    if (nowDay - timeDay > 2) {
+        return time.format('MMM Do');    
+    }
+
+    if (nowDay - timeDay === 1) {
+        return 'Yesterday';
+    }
+
+    if (nowDay - timeDay === 0) {
+        return 'Today';
+    }
+
+    if (time.day() === 0) {
+        return "Monday";
+    }
+
+    return time.format('MMM Do');    
+
+};
+
 
 export const getShortNiceTime = (timeIn: Date | string | undefined, relative: boolean = false): string => {
 
