Horde: Fix "human" time utiliy function
#rnx
#skipci

diff --git a/Engine/Source/Programs/Horde/HordeDashboard/src/base/utilities/timeUtils.ts b/Engine/Source/Programs/Horde/HordeDashboard/src/base/utilities/timeUtils.ts
--- a/Engine/Source/Programs/Horde/HordeDashboard/src/base/utilities/timeUtils.ts
+++ b/Engine/Source/Programs/Horde/HordeDashboard/src/base/utilities/timeUtils.ts
@@ -347,7 +347,7 @@ export const getHumanTime = (timeIn: Date | string | undefined): string => {
     const nowDay = now.dayOfYear();
     const timeDay = time.dayOfYear();
 
-    if (nowDay - timeDay > 2) {
+    if (nowDay - timeDay >= 2) {
         return time.format('MMM Do');    
     }
 
