Horde: Left align log view preference
#rnx
#skipci

diff --git a/Engine/Source/Programs/Horde/HordeDashboard/src/backend/Api.ts b/Engine/Source/Programs/Horde/HordeDashboard/src/backend/Api.ts
--- a/Engine/Source/Programs/Horde/HordeDashboard/src/backend/Api.ts
+++ b/Engine/Source/Programs/Horde/HordeDashboard/src/backend/Api.ts
@@ -2812,6 +2812,7 @@ export enum DashboardPreference {
 	ColorError = "ColorError",
 	ColorRunning = "ColorRunning",
 	LocalCache = "LocalCache",
+	LeftAlignLog = "LeftAlignLog"
 }
 
 export type DashboardSettings = {
diff --git a/Engine/Source/Programs/Horde/HordeDashboard/src/backend/Dashboard.ts b/Engine/Source/Programs/Horde/HordeDashboard/src/backend/Dashboard.ts
--- a/Engine/Source/Programs/Horde/HordeDashboard/src/backend/Dashboard.ts
+++ b/Engine/Source/Programs/Horde/HordeDashboard/src/backend/Dashboard.ts
@@ -213,6 +213,15 @@ export class Dashboard {
         this.setPreference(DashboardPreference.DisplayUTC, value ? "true" : "false");
     }
 
+    setLeftAlignLog(value: boolean | undefined) {
+        this.setPreference(DashboardPreference.LeftAlignLog, value ? "true" : "false");
+    }
+
+    get leftAlignLog(): boolean {
+        return this.preferences.get(DashboardPreference.LeftAlignLog) === 'true';
+    }
+
+
     private hasLoggedLocalCache = false;
 
     get localCache(): boolean {
