Horde: Issue modal fixes
#rnx
#skipci

diff --git a/Engine/Source/Programs/Horde/HordeDashboard/src/base/utilities/timeUtils.ts b/Engine/Source/Programs/Horde/HordeDashboard/src/base/utilities/timeUtils.ts
--- a/Engine/Source/Programs/Horde/HordeDashboard/src/base/utilities/timeUtils.ts
+++ b/Engine/Source/Programs/Horde/HordeDashboard/src/base/utilities/timeUtils.ts
@@ -347,19 +347,21 @@ export const getHumanTime = (timeIn: Date | string | undefined): string => {
     const nowDay = now.dayOfYear();
     const timeDay = time.dayOfYear();
 
-    if (nowDay - timeDay >= 2) {
+    const delta = nowDay - timeDay;
+
+    if (delta > 2) {
         return time.format('MMM Do');    
     }
 
-    if (nowDay - timeDay === 1) {
+    if (delta === 1) {
         return 'Yesterday';
     }
 
-    if (nowDay - timeDay === 0) {
+    if (delta === 0) {
         return 'Today';
     }
 
-    if (time.day() === 0) {
+    if (time.day() === 1) {
         return "Monday";
     }
 
