Horde: Prevent creation of linker issues if we can't find any symbols to match against.

diff --git a/Engine/Source/Programs/Horde/Horde.Build/Issues/Impl/Handlers/SymbolIssueHandler.cs b/Engine/Source/Programs/Horde/Horde.Build/Issues/Impl/Handlers/SymbolIssueHandler.cs
--- a/Engine/Source/Programs/Horde/Horde.Build/Issues/Impl/Handlers/SymbolIssueHandler.cs
+++ b/Engine/Source/Programs/Horde/Horde.Build/Issues/Impl/Handlers/SymbolIssueHandler.cs
@@ -93,6 +93,13 @@ namespace Horde.Build.IssueHandlers.Impl
 
 			SortedSet<string> symbolNames = new SortedSet<string>();
 			GetSymbolNames(eventData, symbolNames);
+
+			if (symbolNames.Count == 0)
+			{
+				fingerprint = null;
+				return false;
+			}
+
 			fingerprint = new NewIssueFingerprint(Type, symbolNames, null);
 			return true;
 		}
