Horde: Temporarily back out change to symbol event matcher - causing test failures.

diff --git a/Engine/Source/Programs/Horde/Horde.Build/Issues/Impl/Handlers/SymbolIssueHandler.cs b/Engine/Source/Programs/Horde/Horde.Build/Issues/Impl/Handlers/SymbolIssueHandler.cs
--- a/Engine/Source/Programs/Horde/Horde.Build/Issues/Impl/Handlers/SymbolIssueHandler.cs
+++ b/Engine/Source/Programs/Horde/Horde.Build/Issues/Impl/Handlers/SymbolIssueHandler.cs
@@ -94,12 +94,6 @@ namespace Horde.Build.IssueHandlers.Impl
 			SortedSet<string> symbolNames = new SortedSet<string>();
 			GetSymbolNames(eventData, symbolNames);
 
-			if (symbolNames.Count == 0)
-			{
-				fingerprint = null;
-				return false;
-			}
-
 			fingerprint = new NewIssueFingerprint(Type, symbolNames, null);
 			return true;
 		}
