Remove UAT Copying during the Setup Build Step

No step copies down the UAT data from TempStroage, no need for it anymore.

#rnx

diff --git a/Engine/Source/Programs/Horde/Horde.Agent/Execution/BuildGraphExecutor.cs b/Engine/Source/Programs/Horde/Horde.Agent/Execution/BuildGraphExecutor.cs
--- a/Engine/Source/Programs/Horde/Horde.Agent/Execution/BuildGraphExecutor.cs
+++ b/Engine/Source/Programs/Horde/Horde.Agent/Execution/BuildGraphExecutor.cs
@@ -369,7 +369,7 @@ namespace Horde.Agent.Execution
 			if (sharedStorageDir != null && !_preprocessScript) // Do not precompile when preprocessing the script; other agents may have a different view of UAT
 			{
 				DirectoryReference buildDir = GetAutomationToolDir(sharedStorageDir);
-				arguments.Append($" CopyUAT -WithLauncher -TargetDir=\"{buildDir}\"");
+				arguments.Append($" -WithLauncher -TargetDir=\"{buildDir}\"");
 			}
 
 			int result = await ExecuteAutomationToolAsync(step, workspaceDir, arguments.ToString(), envVars, step.Credentials, logger, cancellationToken);
