Horde: Flex update for issue modal and fixed human time spanning year
#rnx
#skipci

diff --git a/Engine/Source/Programs/Horde/HordeDashboard/src/base/utilities/timeUtils.ts b/Engine/Source/Programs/Horde/HordeDashboard/src/base/utilities/timeUtils.ts
--- a/Engine/Source/Programs/Horde/HordeDashboard/src/base/utilities/timeUtils.ts
+++ b/Engine/Source/Programs/Horde/HordeDashboard/src/base/utilities/timeUtils.ts
@@ -349,7 +349,7 @@ export const getHumanTime = (timeIn: Date | string | undefined): string => {
 
     const delta = nowDay - timeDay;
 
-    if (delta > 2) {
+    if (delta > 2 || now.year() !== time.year()) {
         return time.format('MMM Do');    
     }
 
