Horde: Ensure DowntimeService is run in both server and worker run mode

Without it, server processes would check downtime at startup and potentially get stuck in whatever status got reported.

diff --git a/Engine/Source/Programs/Horde/Horde.Build/Startup.cs b/Engine/Source/Programs/Horde/Horde.Build/Startup.cs
--- a/Engine/Source/Programs/Horde/Horde.Build/Startup.cs
+++ b/Engine/Source/Programs/Horde/Horde.Build/Startup.cs
@@ -547,6 +547,9 @@ namespace Horde.Build
 						.RequireAuthenticatedUser()							
 						.Build();
 				});
+			
+			// Hosted service that needs to run no matter the run mode of the process (server vs worker)
+			services.AddHostedService(provider => (DowntimeService)provider.GetRequiredService<IDowntimeService>());
 
 			if (settings.IsRunModeActive(RunMode.Worker))
 			{
@@ -557,7 +560,6 @@ namespace Horde.Build
 				services.AddHostedService(provider => provider.GetRequiredService<AgentService>());
 				services.AddHostedService(provider => provider.GetRequiredService<CommitService>());
 				services.AddHostedService(provider => provider.GetRequiredService<ConsistencyService>());
-				services.AddHostedService(provider => (DowntimeService)provider.GetRequiredService<IDowntimeService>());
 				services.AddHostedService(provider => provider.GetRequiredService<IIssueService>());
 				services.AddHostedService(provider => (LogFileService)provider.GetRequiredService<ILogFileService>());
 				services.AddHostedService(provider => (NotificationService)provider.GetRequiredService<INotificationService>());
