Horde Storage - Blob store consistency check for filesystem will always run, no matter if the instance is leader or not.

diff --git a/Engine/Source/Programs/Horde/HordeStorage/Horde.Storage/Implementation/Blob/BlobStoreConsistencyCheckService.cs b/Engine/Source/Programs/Horde/HordeStorage/Horde.Storage/Implementation/Blob/BlobStoreConsistencyCheckService.cs
--- a/Engine/Source/Programs/Horde/HordeStorage/Horde.Storage/Implementation/Blob/BlobStoreConsistencyCheckService.cs
+++ b/Engine/Source/Programs/Horde/HordeStorage/Horde.Storage/Implementation/Blob/BlobStoreConsistencyCheckService.cs
@@ -59,12 +59,6 @@ namespace Horde.Storage.Implementation
                 return false;
             }
 
-            if (!_leaderElection.IsThisInstanceLeader())
-            {
-                _logger.Information("Skipped running blob store consistency check because this instance was not the leader");
-                return false;
-            }
-
             await RunConsistencyCheck();
 
             return true;
@@ -77,9 +71,16 @@ namespace Horde.Storage.Implementation
                 string blobStoreName = blobStore.GetType().Name;
 
                 bool isRootStore = blobStore is AmazonS3Store or AzureBlobStore;
+                bool requiresLeader = blobStore is not FileSystemStore;
 
                 if (!_settings.CurrentValue.RunBlobStoreConsistencyCheckOnRootStore && isRootStore)
-                    return;
+                    continue;
+
+                if (requiresLeader && !_leaderElection.IsThisInstanceLeader())
+                {
+                    _logger.Information("Skipped running blob store consistency check Blob Store {BlobStore} because this instance was not the leader", blobStoreName);
+                    continue;
+                }
 
                 List<NamespaceId> namespaces = await _refsStore.GetNamespaces().ToListAsync();
                 namespaces.AddRange(await _referencesStore.GetNamespaces().ToListAsync());
