Horde: Fix "nominated by" field not always being set when an owner is assigned to an issue.

diff --git a/Engine/Source/Programs/Horde/Horde.Build/Controllers/IssuesController.cs b/Engine/Source/Programs/Horde/Horde.Build/Controllers/IssuesController.cs
--- a/Engine/Source/Programs/Horde/Horde.Build/Controllers/IssuesController.cs
+++ b/Engine/Source/Programs/Horde/Horde.Build/Controllers/IssuesController.cs
@@ -504,6 +504,10 @@ namespace Horde.Build.Controllers
 			{
 				newNominatedById = new UserId(request.NominatedById);
 			}
+			else if (request.OwnerId != null)
+			{
+				newNominatedById = User.GetUserId();
+			}
 
 			UserId? newDeclinedById = null;
 			if (request.Declined ?? false)
