Horde: Fix linker issues matching together if one of them does not have any listed symbol.

diff --git a/Engine/Source/Programs/Horde/Horde.Build.Tests/IssueServiceTests.cs b/Engine/Source/Programs/Horde/Horde.Build.Tests/IssueServiceTests.cs
--- a/Engine/Source/Programs/Horde/Horde.Build.Tests/IssueServiceTests.cs
+++ b/Engine/Source/Programs/Horde/Horde.Build.Tests/IssueServiceTests.cs
@@ -1014,6 +1014,33 @@ namespace Horde.Build.Tests
 				Assert.AreEqual(105, span.LastSuccess?.Change);
 				Assert.AreEqual(null, span.NextSuccess?.Change);
 			}
+
+			// #3
+			// Scenario: Job step fails at 125 with link error, but does not match symbol name
+			// Expected: Creates new issue
+			{
+				IJob job = CreateJob(_mainStreamId, 125, "Test Build", _graph);
+
+				string[] lines =
+				{
+					@"  Engine\Binaries\Win64\UE4Editor-DatasmithExporter.dll: fatal error LNK1120: 1 unresolved externals",
+				};
+				await ParseEventsAsync(job, 0, 0, lines);
+				await UpdateCompleteStep(job, 0, 0, JobStepOutcome.Failure);
+
+				List<IIssue> issues = await IssueService.FindIssuesAsync();
+				Assert.AreEqual(1, issues.Count);
+
+				IIssue issue = issues[0];
+				List<IIssueSpan> spans = await IssueService.GetIssueSpansAsync(issue);
+				Assert.AreEqual(spans.Count, 1);
+				Assert.AreEqual(issue.Fingerprints.Count, 1);
+				Assert.AreEqual(issue.Fingerprints[0].Type, "Symbol");
+
+				IIssueSpan span = spans[0];
+				Assert.AreEqual(120, span.LastSuccess?.Change);
+				Assert.AreEqual(null, span.NextSuccess?.Change);
+			}
 		}
 
 		[TestMethod]
diff --git a/Engine/Source/Programs/Horde/Horde.Build/Issues/IIssueFingerprint.cs b/Engine/Source/Programs/Horde/Horde.Build/Issues/IIssueFingerprint.cs
--- a/Engine/Source/Programs/Horde/Horde.Build/Issues/IIssueFingerprint.cs
+++ b/Engine/Source/Programs/Horde/Horde.Build/Issues/IIssueFingerprint.cs
@@ -44,9 +44,19 @@ namespace Horde.Build.Models
 			{
 				return false;
 			}
-			if (fingerprint.Keys.Count > 0 && other.Keys.Count > 0 && !fingerprint.Keys.Any(x => other.Keys.Contains(x)))
+			if (fingerprint.Keys.Count == 0)
 			{
-				return false;
+				if (other.Keys.Count > 0)
+				{
+					return false;
+				}
+			}
+			else
+			{
+				if (!fingerprint.Keys.Any(x => other.Keys.Contains(x)))
+				{
+					return false;
+				}
 			}
 			if (fingerprint.RejectKeys != null && fingerprint.RejectKeys.Any(x => other.Keys.Contains(x)))
 			{
