Horde: Store a copy of configuration data in the Horde database, allowing it to be retrieved (and cached) separately to stateful documents.

diff --git a/Engine/Source/Programs/Horde/Horde.Build/Configuration/ConfigCollection.cs b/Engine/Source/Programs/Horde/Horde.Build/Configuration/ConfigCollection.cs
new file mode 100644
--- /dev/null
+++ b/Engine/Source/Programs/Horde/Horde.Build/Configuration/ConfigCollection.cs
@@ -0,0 +1,129 @@
+// Copyright Epic Games, Inc. All Rights Reserved.
+
+using Horde.Build.Server;
+using Microsoft.Extensions.Caching.Memory;
+using MongoDB.Bson;
+using MongoDB.Bson.Serialization.Attributes;
+using MongoDB.Driver;
+using System;
+using System.Linq;
+using System.Text.Json;
+using System.Threading.Tasks;
+
+namespace Horde.Build.Config
+{
+	/// <summary>
+	/// Collection of config documents
+	/// </summary>
+	public class ConfigCollection
+	{
+		class ConfigDoc
+		{
+			[BsonId]
+			public string Id { get; set; } = String.Empty;
+
+			[BsonElement("dat")]
+			public ReadOnlyMemory<byte> Data { get; set; } = Array.Empty<byte>();
+		}
+
+		private readonly IMongoCollection<ConfigDoc> _configs;
+		private readonly IMemoryCache _memoryCache;
+
+		/// <summary>
+		/// Constructor
+		/// </summary>
+		public ConfigCollection(MongoService mongoService, IMemoryCache memoryCache)
+		{
+			_configs = mongoService.GetCollection<ConfigDoc>("Configs");
+			_memoryCache = memoryCache;
+		}
+
+		static string GetDocumentCacheKey(string revision) => $"config:{revision}";
+		static string GetTypedCacheKey<T>(string revision) => $"config-typed:{nameof(T)}:{revision}";
+
+		/// <summary>
+		/// Adds a config document with a particular revision
+		/// </summary>
+		/// <param name="revision"></param>
+		/// <param name="data"></param>
+		/// <returns>Identifier for the config data</returns>
+		public async Task AddConfigDataAsync(string revision, ReadOnlyMemory<byte> data)
+		{
+			ConfigDoc config = new ConfigDoc();
+			config.Id = revision;
+			config.Data = data.ToArray();
+			await _configs.ReplaceOneAsync(x => x.Id == revision, config, new ReplaceOptions { IsUpsert = true });
+
+			AddConfigDataToCache(revision, config.Data);
+		}
+
+		/// <summary>
+		/// Gets raw config data with the given revision number
+		/// </summary>
+		/// <param name="revision">The revision number</param>
+		/// <returns>Raw config data with the given id</returns>
+		public async ValueTask<ReadOnlyMemory<byte>> GetConfigDataAsync(string revision)
+		{
+			string cacheKey = GetDocumentCacheKey(revision);
+			if (!_memoryCache.TryGetValue(cacheKey, out ReadOnlyMemory<byte> data))
+			{
+				ConfigDoc config = await _configs.Find(x => x.Id == revision).FirstAsync();
+				using (ICacheEntry entry = _memoryCache.CreateEntry(GetDocumentCacheKey(config.Id)))
+				{
+					entry.SetValue(config);
+					entry.SetSize(config.Data.Length);
+				}
+				data = config.Data;
+			}
+			return data;
+		}
+
+		/// <summary>
+		/// Adds config data for a given revision to the store, and returns the value of it parsed as a JSON object
+		/// </summary>
+		/// <typeparam name="T"></typeparam>
+		/// <param name="revision"></param>
+		/// <param name="data"></param>
+		/// <returns></returns>
+		public async Task<T> AddConfigAsync<T>(string revision, ReadOnlyMemory<byte> data)
+		{
+			await AddConfigDataAsync(revision, data);
+			return await GetConfigAsync<T>(revision);
+		}
+
+		/// <summary>
+		/// Gets typed config data with a particular id
+		/// </summary>
+		/// <typeparam name="T">Type of data to return</typeparam>
+		/// <param name="revision">The unique revision string</param>
+		/// <returns>Parsed config data for the given revision number</returns>
+		public async ValueTask<T> GetConfigAsync<T>(string revision)
+		{
+			string typedCacheKey = GetTypedCacheKey<T>(revision);
+			if (!_memoryCache.TryGetValue(typedCacheKey, out T config))
+			{
+				ReadOnlyMemory<byte> data = await GetConfigDataAsync(revision);
+
+				JsonSerializerOptions options = new JsonSerializerOptions();
+				Startup.ConfigureJsonSerializer(options);
+				config = JsonSerializer.Deserialize<T>(data.Span, options)!;
+
+				using (ICacheEntry entry = _memoryCache.CreateEntry(typedCacheKey))
+				{
+					entry.SetValue(config);
+					entry.SetSize(data.Length);
+				}
+			}
+			return config;
+		}
+
+		void AddConfigDataToCache(string revision, ReadOnlyMemory<byte> data)
+		{
+			using (ICacheEntry entry = _memoryCache.CreateEntry(GetDocumentCacheKey(revision)))
+			{
+				entry.SetValue(data);
+				entry.SetSize(data.Length);
+			}
+		}
+	}
+}
diff --git a/Engine/Source/Programs/Horde/Horde.Build/Configuration/ConfigUpdateService.cs b/Engine/Source/Programs/Horde/Horde.Build/Configuration/ConfigUpdateService.cs
new file mode 100644
--- /dev/null
+++ b/Engine/Source/Programs/Horde/Horde.Build/Configuration/ConfigUpdateService.cs
@@ -0,0 +1,775 @@
+// Copyright Epic Games, Inc. All Rights Reserved.
+
+using System;
+using System.Collections.Generic;
+using System.Globalization;
+using System.IO;
+using System.Linq;
+using System.Reflection;
+using System.Text.Json;
+using System.Text.Json.Serialization;
+using System.Threading;
+using System.Threading.Tasks;
+using EpicGames.Core;
+using Horde.Build.Acls;
+using Horde.Build.Api;
+using Horde.Build.Collections;
+using Horde.Build.Models;
+using Horde.Build.Notifications;
+using Horde.Build.Server;
+using Horde.Build.Services;
+using Horde.Build.Tools;
+using Horde.Build.Utilities;
+using HordeCommon;
+using Microsoft.AspNetCore.StaticFiles;
+using Microsoft.Extensions.Configuration;
+using Microsoft.Extensions.Hosting;
+using Microsoft.Extensions.Logging;
+using Microsoft.Extensions.Options;
+
+namespace Horde.Build.Config
+{
+	using PoolId = StringId<IPool>;
+	using ProjectId = StringId<IProject>;
+	using StreamId = StringId<IStream>;
+
+	/// <summary>
+	/// Polls Perforce for stream config changes
+	/// </summary>
+	public sealed class ConfigUpdateService : IHostedService, IDisposable
+	{
+		const string FileScheme = "file";
+		const string PerforceScheme = "p4-cluster";
+
+		/// <summary>
+		/// Config file version number
+		/// </summary>
+		const int Version = 10;
+
+		readonly MongoService _mongoService;
+		readonly ConfigCollection _configCollection;
+		readonly ToolCollection _toolCollection;
+		readonly ProjectService _projectService;
+		readonly StreamService _streamService;
+		readonly IPerforceService _perforceService;
+		readonly INotificationService _notificationService;
+		readonly AgentService _agentService;
+		readonly PoolService _poolService;
+		readonly IOptionsMonitor<ServerSettings> _settings;
+		readonly ITicker _ticker;
+		readonly ILogger _logger;
+
+		/// <summary>
+		/// Constructor
+		/// </summary>
+		public ConfigUpdateService(MongoService mongoService, ConfigCollection configCollection, IPerforceService perforceService, ToolCollection toolCollection, ProjectService projectService, StreamService streamService, INotificationService notificationService, PoolService poolService, AgentService agentService, IClock clock, IOptionsMonitor<ServerSettings> settings, ILogger<ConfigUpdateService> logger)
+		{
+			_mongoService = mongoService;
+			_configCollection = configCollection;
+			_perforceService = perforceService;
+			_toolCollection = toolCollection;
+			_projectService = projectService;
+			_streamService = streamService;
+			_notificationService = notificationService;
+			_poolService = poolService;
+			_agentService = agentService;
+			_settings = settings;
+			if (mongoService.ReadOnlyMode)
+			{
+				_ticker = new NullTicker();
+			}
+			else
+			{
+				_ticker = clock.AddSharedTicker<ConfigUpdateService>(TimeSpan.FromMinutes(1.0), TickLeaderAsync, logger);
+			}
+			_logger = logger;
+
+			// This will trigger if the local Horde.json user configuration is changed
+			_settings.OnChange(OnUserConfigUpdated);
+		}
+
+		/// <inheritdoc/>
+		public Task StartAsync(CancellationToken cancellationToken) => _ticker.StartAsync();
+
+		/// <inheritdoc/>
+		public Task StopAsync(CancellationToken cancellationToken) => _ticker.StopAsync();
+
+		/// <inheritdoc/>
+		public void Dispose() => _ticker.Dispose();
+
+		GlobalConfig? _cachedGlobalConfig;
+		string? _cachedGlobalConfigRevision;
+		Dictionary<ProjectId, (ProjectConfig Config, string Revision)> _cachedProjectConfigs = new Dictionary<ProjectId, (ProjectConfig, string)>();
+		readonly Dictionary<ProjectId, string?> _cachedLogoRevisions = new Dictionary<ProjectId, string?>();
+
+		async Task UpdateConfigAsync(Uri configPath)
+		{
+			// Update the globals singleton
+			GlobalConfig globalConfig;
+			for (; ; )
+			{
+				Dictionary<Uri, string> globalRevisions = await FindRevisionsAsync(new[] { configPath });
+				if (globalRevisions.Count == 0)
+				{
+					throw new Exception($"Invalid config path: {configPath}");
+				}
+
+				string revision = globalRevisions.First().Value;
+				if (_cachedGlobalConfig == null || revision != _cachedGlobalConfigRevision)
+				{
+					_logger.LogInformation("Caching global config from {Revision}", revision);
+					try
+					{
+						_cachedGlobalConfig = await ReadDataAsync<GlobalConfig>(revision, configPath);
+						_cachedGlobalConfigRevision = revision;
+					}
+					catch (Exception ex)
+					{
+						await SendFailureNotificationAsync(ex, configPath);
+						return;
+					}
+				}
+				globalConfig = _cachedGlobalConfig;
+
+				Globals globals = await _mongoService.GetGlobalsAsync();
+				if (globals.ConfigRevision == revision)
+				{
+					_logger.LogInformation("Updating configuration from {ConfigPath}", globals.ConfigRevision);
+					break;
+				}
+
+				globals.ConfigRevision = revision;
+				globals.PerforceClusters = _cachedGlobalConfig.PerforceClusters;
+				globals.ScheduledDowntime = _cachedGlobalConfig.Downtime;
+				globals.MaxConformCount = _cachedGlobalConfig.MaxConformCount;
+				globals.ComputeClusters = _cachedGlobalConfig.Compute;
+				globals.RootAcl = Acl.Merge(null, _cachedGlobalConfig.Acl);
+
+				if (await _mongoService.TryUpdateSingletonAsync(globals))
+				{
+					break;
+				}
+			}
+
+			// Update the agent rate table
+			await _agentService.UpdateRateTableAsync(globalConfig.Rates);
+
+			// Update the tools
+			await _toolCollection.ConfigureAsync(globalConfig.Tools);
+
+			// Projects to remove
+			List<IProject> projects = await _projectService.GetProjectsAsync();
+
+			// Get the path to all the project configs
+			List<(ProjectConfigRef ProjectRef, Uri Path)> projectConfigs = globalConfig.Projects.Select(x => (x, CombinePaths(configPath, x.Path))).ToList();
+
+			Dictionary<ProjectId, (ProjectConfig Config, string Revision)> prevCachedProjectConfigs = _cachedProjectConfigs;
+			_cachedProjectConfigs = new Dictionary<ProjectId, (ProjectConfig, string)>();
+
+			List<(ProjectId ProjectId, Uri Path)> projectLogos = new List<(ProjectId ProjectId, Uri Path)>();
+
+			List<(ProjectId ProjectId, StreamConfigRef StreamRef, Uri Path)> streamConfigs = new List<(ProjectId, StreamConfigRef, Uri)>();
+
+			// List of project ids that were not able to be updated. We will avoid removing any existing project or stream definitions for these.
+			HashSet<ProjectId> skipProjectIds = new HashSet<ProjectId>();
+
+			// Update any existing projects
+			Dictionary<Uri, string> projectRevisions = await FindRevisionsAsync(projectConfigs.Select(x => x.Path));
+			for (int idx = 0; idx < projectConfigs.Count; idx++)
+			{
+				// Make sure we were able to fetch metadata for 
+				(ProjectConfigRef projectRef, Uri projectPath) = projectConfigs[idx];
+				if (!projectRevisions.TryGetValue(projectPath, out string? revision))
+				{
+					_logger.LogWarning("Unable to update project {ProjectId} due to missing revision information", projectRef.Id);
+					skipProjectIds.Add(projectRef.Id);
+					continue;
+				}
+
+				IProject? project = projects.FirstOrDefault(x => x.Id == projectRef.Id);
+				bool update = project == null || project.ConfigPath != projectPath.ToString() || project.ConfigRevision != revision;
+
+				ProjectConfig? projectConfig;
+				if (!update && prevCachedProjectConfigs.TryGetValue(projectRef.Id, out (ProjectConfig Config, string Revision) result) && result.Revision == revision)
+				{
+					projectConfig = result.Config;
+				}
+				else
+				{
+					_logger.LogInformation("Caching configuration for project {ProjectId} ({Revision})", projectRef.Id, revision);
+					try
+					{
+						projectConfig = await ReadDataAsync<ProjectConfig>(revision, projectPath);
+						if (update)
+						{
+							_logger.LogInformation("Updating configuration for project {ProjectId} ({Revision})", projectRef.Id, revision);
+							await _projectService.Collection.AddOrUpdateAsync(projectRef.Id, projectPath.ToString(), revision, idx, projectConfig);
+						}
+					}
+					catch (Exception ex)
+					{
+						await SendFailureNotificationAsync(ex, projectPath);
+						skipProjectIds.Add(projectRef.Id);
+						continue;
+					}
+				}
+
+				if (projectConfig.Logo != null)
+				{
+					projectLogos.Add((projectRef.Id, CombinePaths(projectPath, projectConfig.Logo)));
+				}
+
+				_cachedProjectConfigs[projectRef.Id] = (projectConfig, revision);
+				streamConfigs.AddRange(projectConfig.Streams.Select(x => (projectRef.Id, x, CombinePaths(projectPath, x.Path))));
+			}
+
+			// Get the logo revisions
+			Dictionary<Uri, string> logoRevisions = await FindRevisionsAsync(projectLogos.Select(x => x.Path));
+			for (int idx = 0; idx < projectLogos.Count; idx++)
+			{
+				(ProjectId projectId, Uri path) = projectLogos[idx];
+				if (logoRevisions.TryGetValue(path, out string? revision))
+				{
+					string? currentRevision;
+					if (!_cachedLogoRevisions.TryGetValue(projectId, out currentRevision))
+					{
+						currentRevision = (await _projectService.Collection.GetLogoAsync(projectId))?.Revision;
+						_cachedLogoRevisions[projectId] = currentRevision;
+					}
+					if (revision != currentRevision)
+					{
+						_logger.LogInformation("Updating logo for project {ProjectId} ({Revision})", projectId, revision);
+						try
+						{
+							await _projectService.Collection.SetLogoAsync(projectId, path.ToString(), revision, GetMimeTypeFromPath(path), await ReadDataAsync(path));
+							_cachedLogoRevisions[projectId] = revision;
+						}
+						catch (Exception ex)
+						{
+							await SendFailureNotificationAsync(ex, path);
+							continue;
+						}
+					}
+				}
+			}
+
+			// Get the current streams
+			List<IStream> streams = await _streamService.GetStreamsAsync();
+
+			// Get the revisions for all the stream documents
+			Dictionary<Uri, string> streamRevisions = await FindRevisionsAsync(streamConfigs.Select(x => x.Path));
+			for (int idx = 0; idx < streamConfigs.Count; idx++)
+			{
+				(ProjectId projectId, StreamConfigRef streamRef, Uri streamPath) = streamConfigs[idx];
+				if (streamRevisions.TryGetValue(streamPath, out string? revision))
+				{
+					IStream? stream = streams.FirstOrDefault(x => x.Id == streamRef.Id);
+					if (stream == null || stream.ConfigPath != streamPath.ToString() || stream.ConfigRevision != revision)
+					{
+						_logger.LogInformation("Updating configuration for stream {StreamRef} ({Revision})", streamRef.Id, revision);
+						try
+						{
+							StreamConfig streamConfig = await ReadDataAsync<StreamConfig>(revision, streamPath);
+							stream = await _streamService.StreamCollection.CreateOrReplaceAsync(streamRef.Id, stream, streamPath.ToString(), revision, projectId, streamConfig);
+						}
+						catch (Exception ex)
+						{
+							await SendFailureNotificationAsync(ex, streamPath);
+							continue;
+						}
+					}
+				}
+			}
+
+			// Remove any projects which are no longer used
+			HashSet<ProjectId> removeProjectIds = new HashSet<ProjectId>(projects.Select(x => x.Id));
+			removeProjectIds.ExceptWith(projectConfigs.Select(y => y.ProjectRef.Id));
+
+			foreach (ProjectId removeProjectId in removeProjectIds)
+			{
+				_logger.LogInformation("Removing project {ProjectId}", removeProjectId);
+				await _projectService.DeleteProjectAsync(removeProjectId);
+			}
+
+			// Remove any streams that are no longer used
+			HashSet<StreamId> removeStreamIds = new HashSet<StreamId>(streams.Where(x => !skipProjectIds.Contains(x.ProjectId)).Select(x => x.Id));
+			removeStreamIds.ExceptWith(streamConfigs.Select(x => x.StreamRef.Id));
+
+			foreach (StreamId removeStreamId in removeStreamIds)
+			{
+				_logger.LogInformation("Removing stream {StreamId}", removeStreamId);
+				await _streamService.DeleteStreamAsync(removeStreamId);
+			}
+		}
+
+		static readonly FileExtensionContentTypeProvider s_contentTypeProvider = new FileExtensionContentTypeProvider();
+
+		static string GetMimeTypeFromPath(Uri path)
+		{
+			string? contentType;
+			if (!s_contentTypeProvider.TryGetContentType(path.AbsolutePath, out contentType))
+			{
+				contentType = "application/octet-stream";
+			}
+			return contentType;
+		}
+
+		static Uri CombinePaths(Uri baseUri, string path)
+		{
+			if (path.StartsWith("//", StringComparison.Ordinal))
+			{
+				if (baseUri.Scheme == PerforceScheme)
+				{
+					return new Uri($"{PerforceScheme}://{baseUri.Host}{path}");
+				}
+				else
+				{
+					return new Uri($"{PerforceScheme}://{PerforceCluster.DefaultName}{path}");
+				}
+			}
+			return new Uri(baseUri, path);
+		}
+
+		async Task<Dictionary<Uri, string>> FindRevisionsAsync(IEnumerable<Uri> paths)
+		{
+			Dictionary<Uri, string> revisions = new Dictionary<Uri, string>();
+
+			// Find all the Perforce uris
+			List<Uri> perforcePaths = new List<Uri>();
+			foreach (Uri path in paths)
+			{
+				if (path.Scheme == FileScheme)
+				{
+					revisions[path] = $"ver={Version},md5={ContentHash.MD5(new FileReference(path.LocalPath))}";
+				}
+				else if (path.Scheme == PerforceScheme)
+				{
+					perforcePaths.Add(path);
+				}
+				else
+				{
+					throw new Exception($"Invalid path format: {path}");
+				}
+			}
+
+			// Query all the Perforce revisions
+			foreach (IGrouping<string, Uri> perforcePath in perforcePaths.GroupBy(x => x.Host, StringComparer.OrdinalIgnoreCase))
+			{
+				List<FileSummary> files = await _perforceService.FindFilesAsync(perforcePath.Key, perforcePath.Select(x => x.AbsolutePath));
+				foreach (FileSummary file in files)
+				{
+					Uri fileUri = new Uri($"{PerforceScheme}://{perforcePath.Key}{file.DepotPath}");
+					if (file.Error == null)
+					{
+						revisions[fileUri] = $"ver={Version},chg={file.Change},path={fileUri}";
+					}
+					else
+					{
+						_notificationService.NotifyConfigUpdateFailure(file.Error, file.DepotPath);
+					}
+				}
+			}
+
+			return revisions;
+		}
+
+		async Task<T> ReadDataAsync<T>(string revision, Uri configPath) where T : class
+		{
+			byte[] data = await ReadDataAsync(configPath);
+			return await _configCollection.AddConfigAsync<T>(revision, data);
+		}
+
+		Task<byte[]> ReadDataAsync(Uri configPath)
+		{
+			switch (configPath.Scheme)
+			{
+				case FileScheme:
+					return File.ReadAllBytesAsync(configPath.LocalPath);
+				case PerforceScheme:
+					return _perforceService.PrintAsync(configPath.Host, configPath.AbsolutePath);
+				default:
+					throw new Exception($"Invalid config path: {configPath}");
+			}
+		}
+
+		async Task SendFailureNotificationAsync(Exception ex, Uri configPath)
+		{
+			_logger.LogError(ex, "Unable to read data from {ConfigPath}: {Message}", configPath, ex.Message);
+
+			string fileName = configPath.AbsolutePath;
+			int change = -1;
+			IUser? author = null;
+			string? description = null;
+
+			if (configPath.Scheme == PerforceScheme)
+			{
+				try
+				{
+					List<FileSummary> files = await _perforceService.FindFilesAsync(configPath.Host, new[] { fileName });
+					change = files[0].Change;
+
+					List<ChangeSummary> changes = await _perforceService.GetChangesAsync(configPath.Host, change, change, 1);
+					if (changes.Count > 0 && changes[0].Number == change)
+					{
+						(author, description) = (changes[0].Author, changes[0].Description);
+					}
+				}
+				catch (Exception ex2)
+				{
+					_logger.LogError(ex2, "Unable to identify change that last modified {ConfigPath} from Perforce", configPath);
+				}
+			}
+
+			_notificationService.NotifyConfigUpdateFailure(ex.Message, fileName, change, author, description);
+		}
+
+		/// <inheritdoc/>
+		async ValueTask TickLeaderAsync(CancellationToken stoppingToken)
+		{
+
+			Uri? configUri = null;
+
+			if (Path.IsPathRooted(_settings.CurrentValue.ConfigPath) && !_settings.CurrentValue.ConfigPath.StartsWith("//", StringComparison.Ordinal))
+			{
+				// absolute path to config
+				configUri = new Uri(_settings.CurrentValue.ConfigPath);
+			}
+			else if (_settings.CurrentValue.ConfigPath != null)
+			{
+				// relative (development) or perforce path
+				configUri = CombinePaths(new Uri(FileReference.Combine(Program.AppDir, "_").FullName), _settings.CurrentValue.ConfigPath);
+			}
+
+			if (configUri != null)
+			{
+				await UpdateConfigAsync(configUri);
+			}
+		}
+
+		//
+		// On premises configuration handling (aka Horde.json settings, though needs to have perforce global config handled too, checkout/modify/submit)
+		//
+
+		/// <summary>
+		/// Update the global configuration
+		/// </summary>
+		/// <param name="request"></param>
+		/// <returns></returns>
+		public async Task<bool> UpdateGlobalConfig(UpdateGlobalConfigRequest request)
+		{
+			if (_settings.CurrentValue.ConfigPath == null || !Path.IsPathRooted(_settings.CurrentValue.ConfigPath) || _settings.CurrentValue.ConfigPath.StartsWith("//", StringComparison.Ordinal))
+			{
+				throw new Exception($"Global config path must be rooted for service updates.  (Perforce paths are not currently supported for live updates), {_settings.CurrentValue.ConfigPath}");
+			}
+
+			FileReference globalConfigFile = new FileReference(_settings.CurrentValue.ConfigPath);
+			DirectoryReference globalConfigDirectory = globalConfigFile.Directory;
+
+			// make sure the directory exists
+			if (!DirectoryReference.Exists(globalConfigDirectory))
+			{
+				DirectoryReference.CreateDirectory(globalConfigDirectory);
+			}
+
+			bool configDirty = false;
+
+			// projects
+			if (request.ProjectsJson != null)
+			{
+				configDirty = true;
+
+				// write out the projects
+				foreach (KeyValuePair<string, string> project in request.ProjectsJson)
+				{
+					await FileReference.WriteAllTextAsync(FileReference.Combine(globalConfigDirectory, project.Key), project.Value);
+
+					if (request.ProjectLogo != null)
+					{
+						byte[] bytes = Convert.FromBase64String(request.ProjectLogo);
+						await FileReference.WriteAllBytesAsync(FileReference.Combine(globalConfigDirectory, project.Key.Replace(".json", ".png", StringComparison.OrdinalIgnoreCase)), bytes);
+					}
+				}
+			}
+
+			// streams
+			if (request.StreamsJson != null)
+			{
+				configDirty = true;
+
+				// write out the streams
+				foreach (KeyValuePair<string, string> stream in request.StreamsJson)
+				{
+					await FileReference.WriteAllTextAsync(FileReference.Combine(globalConfigDirectory, stream.Key), stream.Value);
+				}
+			}
+
+			// update global config path
+			if (!String.IsNullOrEmpty(request.GlobalsJson))
+			{
+				configDirty = true;
+
+				await FileReference.WriteAllTextAsync(globalConfigFile, request.GlobalsJson);
+			}
+
+			if (configDirty == true)
+			{
+				await UpdateConfigAsync(new Uri("file://" + globalConfigFile.ToString()));
+			}
+
+			// create the default pool 
+			if (request.DefaultPoolName != null)
+			{
+				PoolId poolIdValue = PoolId.Sanitize(request.DefaultPoolName);
+
+				IPool? pool = await _poolService.GetPoolAsync(poolIdValue);
+				if (pool == null)
+				{
+					await _poolService.CreatePoolAsync(request.DefaultPoolName, condition: "OSFamily == 'Windows'", properties: new Dictionary<string, string>() { ["Color"] = "0" });
+				}
+			}
+
+			return true;
+
+		}
+
+		/// <summary>
+		/// Update the server settings
+		/// </summary>
+		/// <param name="request"></param>
+		/// <returns></returns>
+		public async Task<ServerUpdateResponse> UpdateServerSettings(UpdateServerSettingsRequest request)
+		{
+
+			ServerUpdateResponse response = new ServerUpdateResponse();
+
+			Dictionary<string, object> newUserSettings = new Dictionary<string, object>();
+
+			try
+			{
+
+				if (request.Settings == null || request.Settings.Count == 0)
+				{
+					return response;
+				}
+
+				if (_userConfigUpdated != null)
+				{
+					response.Errors.Add("User configuation already being updated");
+					return response;
+				}
+
+				// Load the current user configuration
+				IConfiguration config = new ConfigurationBuilder()
+					.AddJsonFile(Program.UserConfigFile.FullName, optional: true)
+					.Build();
+
+				ServerSettings userSettings = new ServerSettings();
+				config.GetSection("Horde").Bind(userSettings);
+
+				// Figure out current and new settings
+				HashSet<string> userProps = new HashSet<string>(StringComparer.OrdinalIgnoreCase);
+
+				if (FileReference.Exists(Program.UserConfigFile))
+				{
+					byte[] data = await FileReference.ReadAllBytesAsync(Program.UserConfigFile);
+
+					JsonDocument document = JsonDocument.Parse(data);
+					JsonElement hordeElement;
+
+					if (document.RootElement.TryGetProperty("Horde", out hordeElement))
+					{
+						foreach (JsonProperty property in hordeElement.EnumerateObject())
+						{
+							userProps.Add(property.Name);
+						}
+					}
+				}
+
+				foreach (string name in request.Settings.Keys)
+				{
+					userProps.Add(name);
+				}
+
+				// Apply the changes from the request
+				foreach (KeyValuePair<string, object> pair in request.Settings)
+				{
+					PropertyInfo? property = userSettings.GetType().GetProperty(pair.Key, BindingFlags.IgnoreCase | BindingFlags.Public | BindingFlags.Instance);
+					if (property == null)
+					{
+						response.Errors.Add($"Horde configuration property {pair.Key} does not exist when reading server settings");
+						_logger.LogError("Horde configuration property {Key} does not exist when reading server settings", pair.Key);
+						continue;
+					}
+
+					if (property.SetMethod == null)
+					{
+						response.Errors.Add($"Horde configuration property {pair.Key} does not have a set method");
+						_logger.LogError("Horde configuration property {Key} does not have a set method", pair.Key);
+						continue;
+					}
+
+					// explicit setting removal
+					if (pair.Value == null)
+					{
+						userProps.Remove(pair.Key);
+						continue;
+					}
+
+					JsonElement element = (JsonElement)pair.Value;
+
+					object? value = null;
+
+					switch (element.ValueKind)
+					{
+						case JsonValueKind.True:
+							value = true;
+							break;
+						case JsonValueKind.False:
+							value = false;
+							break;
+						case JsonValueKind.Number:
+							value = element.GetDouble();
+							break;
+						case JsonValueKind.String:
+							value = element.GetString();
+							break;
+					}
+
+					// unable to map type
+					if (value == null)
+					{
+						response.Errors.Add($"Unable to map type for Property {property.Name}");
+						continue;
+					}
+
+					// handle common conversions
+					if (property.GetType() != value.GetType())
+					{
+						try
+						{
+							value = Convert.ChangeType(value, property.PropertyType, CultureInfo.CurrentCulture);
+						}
+						catch (Exception ex)
+						{
+							response.Errors.Add($"Property {property.Name} raised exception during conversion, {ex.Message}");
+							continue;
+						}
+
+						if (value == null)
+						{
+							response.Errors.Add($"Property {property.Name} had null value on conversion");
+							continue;
+						}
+					}
+					else
+					{
+						response.Errors.Add($"Property {property.Name} is not assignable to {value}");
+						continue;
+					}
+
+
+					//  Set the value, providing some validation
+					try
+					{
+						property.SetMethod.Invoke(userSettings, new object[] { value });
+					}
+					catch (Exception ex)
+					{
+						response.Errors.Add($"Exception updating property {pair.Key}, {ex.Message}");
+					}
+				}
+
+				// Construct the new user settings
+				foreach (string name in userProps)
+				{
+					PropertyInfo? property = userSettings.GetType().GetProperty(name, BindingFlags.IgnoreCase | BindingFlags.Public | BindingFlags.Instance);
+					if (property == null)
+					{
+						response.Errors.Add($"Horde configuration property {name} does not exist when writing server settings");
+						_logger.LogError("Horde configuration property {Name} does not exist when writing server settings", name);
+
+						continue;
+					}
+
+					if (property.GetMethod == null)
+					{
+						response.Errors.Add($"Horde configuration property {name} does not have a get method");
+						_logger.LogError("Horde configuration property {Name} does not have a get method", name);
+						continue;
+					}
+
+					object? result = property.GetMethod.Invoke(userSettings, null);
+
+					if (result == null)
+					{
+						response.Errors.Add($"Horde configuration property {name} was null while writing and should have already been filtered out");
+						_logger.LogError("Horde configuration property {Name} was null while writing and should have already been filtered out", name);
+						continue;
+					}
+
+					newUserSettings.Add(property.Name, result);
+
+				}
+			}
+			catch (Exception ex)
+			{
+				response.Errors.Add($"Exception while updating settings: {ex.Message}");
+				_logger.LogError(ex, "{Error}", response.Errors.Last());
+			}
+
+			if (response.Errors.Count != 0)
+			{
+				return response;
+			}
+
+			Dictionary<string, object> newLocalSettings = new Dictionary<string, object>();
+			newLocalSettings["Horde"] = newUserSettings;
+
+			try
+			{
+				_userConfigUpdated = new TaskCompletionSource<bool>();
+
+				// This will trigger a setting update as the user config json is set to reload on change
+				try
+				{
+					await FileReference.WriteAllBytesAsync(Program.UserConfigFile, JsonSerializer.SerializeToUtf8Bytes(newLocalSettings, new JsonSerializerOptions { WriteIndented = true, DefaultIgnoreCondition = JsonIgnoreCondition.WhenWritingNull }));
+				}
+				catch (Exception ex)
+				{
+					response.Errors.Add($"Unable to serialize json settings to {Program.UserConfigFile}, {ex.Message}");
+					_logger.LogError(ex, "Unable to serialize json settings to {ConfigFile}, {Message}", Program.UserConfigFile.ToString(), ex.Message);
+				}
+
+				if (response.Errors.Count == 0)
+				{
+					if (await Task.WhenAny(_userConfigUpdated.Task, Task.Delay(5000)) != _userConfigUpdated.Task)
+					{
+						response.Errors.Add("Server update timed out while awaiting write");
+						_logger.LogError("Server update timed out after writing config");
+					}
+				}
+			}
+			finally
+			{
+				_userConfigUpdated = null;
+			}
+
+			return response;
+
+		}
+		void OnUserConfigUpdated(ServerSettings settings, string name)
+		{
+			NumUserConfigUpdates++;
+			_userConfigUpdated?.SetResult(true);
+		}
+
+		private TaskCompletionSource<bool>? _userConfigUpdated;
+
+		/// <summary>
+		/// The number of server configuration updates that have been made while the server is running
+		/// This is used to detect whether server may need to be restarted
+		/// </summary>
+		public int NumUserConfigUpdates { get; private set; } = 0;
+
+	}
+}
diff --git a/Engine/Source/Programs/Horde/Horde.Build/Controllers/ConfigController.cs b/Engine/Source/Programs/Horde/Horde.Build/Controllers/ConfigController.cs
deleted file mode 100644
diff --git a/Engine/Source/Programs/Horde/Horde.Build/Server/ConfigService.cs b/Engine/Source/Programs/Horde/Horde.Build/Server/ConfigService.cs
deleted file mode 100644
diff --git a/Engine/Source/Programs/Horde/Horde.Build/Startup.cs b/Engine/Source/Programs/Horde/Horde.Build/Startup.cs
--- a/Engine/Source/Programs/Horde/Horde.Build/Startup.cs
+++ b/Engine/Source/Programs/Horde/Horde.Build/Startup.cs
@@ -28,6 +28,7 @@ using Horde.Build.Commits;
 using Horde.Build.Commits.Impl;
 using Horde.Build.Compute;
 using Horde.Build.Compute.Impl;
+using Horde.Build.Config;
 using Horde.Build.Fleet.Autoscale;
 using Horde.Build.Jobs;
 using Horde.Build.Logs;
@@ -313,6 +314,7 @@ namespace Horde.Build
 			services.AddSingleton<IDeviceCollection, DeviceCollection>();
 			services.AddSingleton<INoticeCollection, NoticeCollection>();
 
+			services.AddSingleton<ConfigCollection>();
 			services.AddSingleton<ToolCollection>();
 
 			// Auditing
@@ -344,7 +346,6 @@ namespace Horde.Build
 			services.AddSingleton<RequestTrackerService>();
 			services.AddSingleton<CredentialService>();
 			services.AddSingleton<MongoService>();
-			services.AddSingleton<ConfigService>();
 			services.AddSingleton<IDogStatsd>(ctx =>
 			{
 				string? datadogAgentHost = Environment.GetEnvironmentVariable("DD_AGENT_HOST");
@@ -572,7 +573,7 @@ namespace Horde.Build
 				services.AddHostedService(provider => provider.GetRequiredService<PerforceLoadBalancer>());
 				services.AddHostedService<PoolUpdateService>();
 				services.AddHostedService(provider => provider.GetRequiredService<SlackNotificationSink>());
-				services.AddHostedService<ConfigService>();
+				services.AddHostedService<ConfigUpdateService>();
 				services.AddHostedService<TelemetryService>();
 				services.AddHostedService(provider => provider.GetRequiredService<DeviceService>());
 			}
