Horde: Bump the config version number to force an update.

diff --git a/Engine/Source/Programs/Horde/Horde.Build/Configuration/ConfigUpdateService.cs b/Engine/Source/Programs/Horde/Horde.Build/Configuration/ConfigUpdateService.cs
--- a/Engine/Source/Programs/Horde/Horde.Build/Configuration/ConfigUpdateService.cs
+++ b/Engine/Source/Programs/Horde/Horde.Build/Configuration/ConfigUpdateService.cs
@@ -44,7 +44,7 @@ namespace Horde.Build.Config
 		/// <summary>
 		/// Config file version number
 		/// </summary>
-		const int Version = 10;
+		const int Version = 11;
 
 		readonly MongoService _mongoService;
 		readonly ConfigCollection _configCollection;
