Horde: Fix empty fields storing config data in DB.

diff --git a/Engine/Source/Programs/Horde/Horde.Build/Configuration/ConfigCollection.cs b/Engine/Source/Programs/Horde/Horde.Build/Configuration/ConfigCollection.cs
--- a/Engine/Source/Programs/Horde/Horde.Build/Configuration/ConfigCollection.cs
+++ b/Engine/Source/Programs/Horde/Horde.Build/Configuration/ConfigCollection.cs
@@ -23,7 +23,7 @@ namespace Horde.Build.Config
 			public string Id { get; set; } = String.Empty;
 
 			[BsonElement("dat")]
-			public ReadOnlyMemory<byte> Data { get; set; } = Array.Empty<byte>();
+			public byte[] Data { get; set; } = Array.Empty<byte>();
 		}
 
 		private readonly IMongoCollection<ConfigDoc> _configs;
diff --git a/Engine/Source/Programs/Horde/Horde.Build/Configuration/ConfigUpdateService.cs b/Engine/Source/Programs/Horde/Horde.Build/Configuration/ConfigUpdateService.cs
--- a/Engine/Source/Programs/Horde/Horde.Build/Configuration/ConfigUpdateService.cs
+++ b/Engine/Source/Programs/Horde/Horde.Build/Configuration/ConfigUpdateService.cs
@@ -44,7 +44,7 @@ namespace Horde.Build.Config
 		/// <summary>
 		/// Config file version number
 		/// </summary>
-		const int Version = 11;
+		const int Version = 12;
 
 		readonly MongoService _mongoService;
 		readonly ConfigCollection _configCollection;
