Horde: Daily slack channel updates for build health status.

Templates in a stream can now reference a 'workflow' by string identifier. Workflows are defined in the stream config file and specify when (ie. times of day) and where (ie. which Slack channel) to output reports for build health status. More configurable functionality is planned to be added to workflows in the future.

Also fix Slack errors sending header blocks by default, due to markdown being enabled by default. Header blocks only support plain text.

diff --git a/Engine/Source/Programs/Horde/Horde.Build/Api/Streams.cs b/Engine/Source/Programs/Horde/Horde.Build/Api/Streams.cs
--- a/Engine/Source/Programs/Horde/Horde.Build/Api/Streams.cs
+++ b/Engine/Source/Programs/Horde/Horde.Build/Api/Streams.cs
@@ -4,12 +4,14 @@ using System;
 using System.Collections.Generic;
 using System.ComponentModel.DataAnnotations;
 using Horde.Build.Models;
+using Horde.Build.Server;
 using Horde.Build.Utilities;
 using HordeCommon;
 
 namespace Horde.Build.Api
 {
 	using TemplateRefId = StringId<TemplateRef>;
+	using WorkflowId = StringId<WorkflowConfig>;
 
 	/// <summary>
 	/// Query selecting the base changelist to use
@@ -217,6 +219,11 @@ namespace Horde.Build.Api
 		/// </summary>
 		public string? TriageChannel { get; set; }
 
+		/// <summary>
+		/// Issue handling workflow
+		/// </summary>
+		public WorkflowId? Workflow { get; set; }
+
 		/// <summary>
 		/// Schedule to execute this template
 		/// </summary>
diff --git a/Engine/Source/Programs/Horde/Horde.Build/Notifications/INotificationService.cs b/Engine/Source/Programs/Horde/Horde.Build/Notifications/INotificationService.cs
--- a/Engine/Source/Programs/Horde/Horde.Build/Notifications/INotificationService.cs
+++ b/Engine/Source/Programs/Horde/Horde.Build/Notifications/INotificationService.cs
@@ -5,7 +5,9 @@ using System.Diagnostics.CodeAnalysis;
 using System.Security.Claims;
 using System.Threading.Tasks;
 using Horde.Build.Api;
+using Horde.Build.Issues.Impl;
 using Horde.Build.Models;
+using Horde.Build.Server;
 using MongoDB.Bson;
 
 namespace Horde.Build.Notifications
@@ -114,5 +116,11 @@ namespace Horde.Build.Notifications
 		/// <param name="node"></param>
 		/// <param name="user"></param>
 		void NotifyDeviceService(string message, IDevice? device = null, IDevicePool? pool = null, IStream? stream = null, IJob? job = null, IJobStep? step = null, INode? node = null, IUser? user = null);
-    }
+
+		/// <summary>
+		/// Post a notification for the open issues in a stream
+		/// </summary>
+		/// <param name="report">The report data to send</param>
+		void SendIssueReport(IssueReport report);
+	}
 }
diff --git a/Engine/Source/Programs/Horde/Horde.Build/Notifications/INotificationSink.cs b/Engine/Source/Programs/Horde/Horde.Build/Notifications/INotificationSink.cs
--- a/Engine/Source/Programs/Horde/Horde.Build/Notifications/INotificationSink.cs
+++ b/Engine/Source/Programs/Horde/Horde.Build/Notifications/INotificationSink.cs
@@ -4,7 +4,9 @@ using System;
 using System.Collections.Generic;
 using System.Threading.Tasks;
 using Horde.Build.Api;
+using Horde.Build.Issues.Impl;
 using Horde.Build.Models;
+using Horde.Build.Server;
 using HordeCommon;
 
 namespace Horde.Build.Notifications
@@ -68,6 +70,8 @@ namespace Horde.Build.Notifications
 		/// <returns>Async task</returns>
 		Task NotifyLabelCompleteAsync(IUser user, IJob job, IStream stream, ILabel label, int labelIdx, LabelOutcome outcome, List<(string, JobStepOutcome, Uri)> stepData);
 
+		#region Issues
+
 		/// <summary>
 		/// Send notifications that a new issue has been created or assigned
 		/// </summary>
@@ -75,6 +79,15 @@ namespace Horde.Build.Notifications
 		/// <returns></returns>
 		Task NotifyIssueUpdatedAsync(IIssue issue);
 
+		/// <summary>
+		/// Notify the status of issues in a stream
+		/// </summary>
+		/// <param name="report"></param>
+		/// <returns></returns>
+		Task SendIssueReportAsync(IssueReport report);
+
+		#endregion
+
 		/// <summary>
 		/// Notification that a stream has failed to update
 		/// </summary>
diff --git a/Engine/Source/Programs/Horde/Horde.Build/Notifications/Impl/NotificationService.cs b/Engine/Source/Programs/Horde/Horde.Build/Notifications/Impl/NotificationService.cs
--- a/Engine/Source/Programs/Horde/Horde.Build/Notifications/Impl/NotificationService.cs
+++ b/Engine/Source/Programs/Horde/Horde.Build/Notifications/Impl/NotificationService.cs
@@ -10,6 +10,7 @@ using System.Threading.Tasks;
 using EpicGames.Core;
 using Horde.Build.Api;
 using Horde.Build.Collections;
+using Horde.Build.Issues.Impl;
 using Horde.Build.Models;
 using Horde.Build.Server;
 using Horde.Build.Services;
@@ -727,5 +728,11 @@ namespace Horde.Build.Notifications.Impl
 
 			_logger.LogDebug("Finished sending label notifications for label {DashboardName}/{UgsName} in job {JobId}", label.DashboardName, label.UgsName, job.Id);
 		}
+
+		/// <inheritdoc/>
+		public void SendIssueReport(IssueReport report)
+		{
+			EnqueueTasks(sink => sink.SendIssueReportAsync(report));
+		}
 	}
 }
diff --git a/Engine/Source/Programs/Horde/Horde.Build/Notifications/Impl/SlackNotificationSink.cs b/Engine/Source/Programs/Horde/Horde.Build/Notifications/Impl/SlackNotificationSink.cs
--- a/Engine/Source/Programs/Horde/Horde.Build/Notifications/Impl/SlackNotificationSink.cs
+++ b/Engine/Source/Programs/Horde/Horde.Build/Notifications/Impl/SlackNotificationSink.cs
@@ -17,6 +17,7 @@ using System.Threading.Tasks;
 using EpicGames.Core;
 using Horde.Build.Api;
 using Horde.Build.Collections;
+using Horde.Build.Issues.Impl;
 using Horde.Build.Models;
 using Horde.Build.Server;
 using Horde.Build.Services;
@@ -300,7 +301,7 @@ namespace Horde.Build.Notifications.Impl
 			attachment.Color = outcomeColor;
 			attachment.FallbackText = $"Job(s) scheduled in an auto-scaled pool with no agents online. Job IDs {jobIds}";
 
-			attachment.Blocks.Add(new HeaderBlock($"Jobs scheduled in empty pool", false, true));
+			attachment.Blocks.Add(new HeaderBlock($"Jobs scheduled in empty pool", true));
 			attachment.Blocks.Add(new SectionBlock($"One or more jobs were scheduled in an auto-scaled pool but with no current agents online."));
 			attachment.Blocks.Add(new SectionBlock(sb.ToString()));
 			
@@ -911,6 +912,86 @@ namespace Horde.Build.Notifications.Impl
 			return $"<!date^{time.ToUnixTimeSeconds()}^{{time}}|{time}>";
 		}
 
+		/// <inheritdoc/>
+		public async Task SendIssueReportAsync(IssueReport report)
+		{
+			if (report.Workflow.ReportChannel != null)
+			{
+				List<BlockBase> blocks = new List<BlockBase>();
+				blocks.Add(new HeaderBlock($"{report.Stream.Name}: {report.Time:d}"));
+				await SendMessageAsync(report.Workflow.ReportChannel, blocks: blocks.ToArray());// B, message.ToString());/// "Hello world");//, blocks: blocks.ToArray());
+
+				StringBuilder body = new StringBuilder();
+				if (report.Issues.Count == 0)
+				{
+					body.Append(":tick: No issues open.");
+				}
+				else
+				{
+					TimeSpan averageAge = TimeSpan.FromHours(report.Issues.Select(x => (report.Time - x.CreatedAt).TotalHours).Average());
+					body.Append($"*{report.Issues.Count} issues* are currently open (average age {FormatReadableTimeSpan(averageAge)}):");
+				}
+				body.Append('\n');
+
+				foreach(IIssue issue in report.Issues.OrderByDescending(x => x.Id))
+				{
+					IIssueSpan? span = report.IssueSpans.FirstOrDefault(x => x.IssueId == issue.Id);
+
+					Uri issueUrl = _settings.DashboardUrl;
+					if(span != null)
+					{
+						IIssueStep lastFailure = span.LastFailure;
+						issueUrl = new Uri(issueUrl, $"job/{lastFailure.JobId}?step={lastFailure.StepId}&issue={issue.Id}");
+					}
+
+					string status = "*Unassigned*";
+					if(issue.OwnerId != null)
+					{
+						IUser? user = await _userCollection.GetCachedUserAsync(issue.OwnerId.Value);
+						if (user == null)
+						{
+							status = $"Assigned to user {issue.OwnerId.Value}";
+						}
+						else
+						{
+							status = $"Assigned to {user.Name}";
+						}
+						if(issue.AcknowledgedAt == null)
+						{
+							status = $"{status} (unacknowledged)";
+						}
+					}
+
+					body.Append($"\n\t\u2022 *Issue <{issueUrl}|{issue.Id}>*: {issue.Summary} [{FormatReadableTimeSpan(report.Time - issue.CreatedAt)}] - {status}");
+				}
+
+				if (report.NumSteps > 0)
+				{
+					double totalPct = (report.NumPassingSteps * 100.0) / report.NumSteps;
+					body.Append($"\n\n{report.NumPassingSteps:n0} of {report.NumSteps:n0} (*{totalPct:0.0}%*) build steps succeeded since last status update.");
+				}
+
+				await SendMessageAsync(report.Workflow.ReportChannel, text: body.ToString());
+			}
+		}
+
+		static string FormatReadableTimeSpan(TimeSpan timeSpan)
+		{
+			double totalDays = timeSpan.TotalDays;
+			if (totalDays > 1.0)
+			{
+				return $"{totalDays:n0}d";
+			}
+
+			double totalHours = timeSpan.TotalHours;
+			if(totalHours > 1.0)
+			{
+				return $"{totalHours:n0}h";
+			}
+
+			return $"{timeSpan.TotalMinutes:n0}m";
+		}
+
 		#endregion
 
 		#region Stream updates
@@ -951,7 +1032,7 @@ namespace Horde.Build.Notifications.Impl
 			attachment.Color = outcomeColor;
 			attachment.FallbackText = $"Update Failure: {fileName}";
 
-			attachment.Blocks.Add(new HeaderBlock($"Config Update Failure :rip:", false, true));
+			attachment.Blocks.Add(new HeaderBlock($"Config Update Failure :rip:", true));
 
 			attachment.Blocks.Add(new SectionBlock($"Horde was unable to update {fileName}"));
 			attachment.Blocks.Add(new SectionBlock($"```{errorMessage}```"));
@@ -1001,7 +1082,7 @@ namespace Horde.Build.Notifications.Impl
 			attachment.Color = outcomeColor;
 			attachment.FallbackText = $"{file.DepotPath} - Update Failure";
 
-			attachment.Blocks.Add(new HeaderBlock($"Stream Update Failure :rip:", false, true));
+			attachment.Blocks.Add(new HeaderBlock($"Stream Update Failure :rip:", true));
 
 			attachment.Blocks.Add(new SectionBlock($"<!here> Horde was unable to update {file.DepotPath}"));
 			attachment.Blocks.Add(new SectionBlock($"```{file.Error}```"));
@@ -1074,7 +1155,7 @@ namespace Horde.Build.Notifications.Impl
 				attachment.FallbackText += $" - Device: {device.Name} Pool: {pool.Name}";
 			}
 				
-			attachment.Blocks.Add(new HeaderBlock($"{message}", false, false));
+			attachment.Blocks.Add(new HeaderBlock($"{message}", false));
 
 			if (stream != null && job != null && step != null && node != null)
 			{
@@ -1094,7 +1175,7 @@ namespace Horde.Build.Notifications.Impl
 			return SendMessageAsync(recipient, attachments: new[] { attachment });
 
 		}
-		
+
 		#endregion
 
 		const int MaxJobStepEvents = 5;
@@ -1236,7 +1317,10 @@ namespace Horde.Build.Notifications.Impl
 
 			BlockKitMessage message = new BlockKitMessage();
 			message.Channel = recipient;
-			message.Text = AddEnvironmentAnnotation(text);
+			if (text != null)
+			{
+				message.Text = AddEnvironmentAnnotation(text);
+			}
 			if (blocks != null)
 			{
 				message.Blocks.AddRange(blocks);
diff --git a/Engine/Source/Programs/Horde/Horde.Build/Server/Config.cs b/Engine/Source/Programs/Horde/Horde.Build/Server/Config.cs
--- a/Engine/Source/Programs/Horde/Horde.Build/Server/Config.cs
+++ b/Engine/Source/Programs/Horde/Horde.Build/Server/Config.cs
@@ -15,7 +15,8 @@ namespace Horde.Build.Server
 {
 	using ProjectId = StringId<IProject>;
 	using StreamId = StringId<IStream>;
-
+	using WorkflowId = StringId<WorkflowConfig>;
+	
 	/// <summary>
 	/// Global configuration
 	/// </summary>
@@ -262,6 +263,32 @@ namespace Horde.Build.Server
 		/// Filter for paths to be replicated to storage, as a Perforce wildcard relative to the root of the workspace.
 		/// </summary>
 		public string? ReplicationFilter { get; set; }
+
+		/// <summary>
+		/// Workflows for dealing with new issues
+		/// </summary>
+		public List<WorkflowConfig> Workflows { get; set; } = new List<WorkflowConfig>();
+	}
+
+	/// <summary>
+	/// Configuration for an issue workflow
+	/// </summary>
+	public class WorkflowConfig
+	{
+		/// <summary>
+		/// Identifier for this workflow
+		/// </summary>
+		public WorkflowId Id { get; set; } = WorkflowId.Empty;
+
+		/// <summary>
+		/// Times of day at which to send a report
+		/// </summary>
+		public List<TimeSpan> ReportTimes { get; set; } = new List<TimeSpan> { TimeSpan.Zero };
+
+		/// <summary>
+		/// Channel to post summary information for these templates.
+		/// </summary>
+		public string? ReportChannel { get; set; }
 	}
 
 	/// <summary>
diff --git a/Engine/Source/Programs/Horde/Horde.Build/Startup.cs b/Engine/Source/Programs/Horde/Horde.Build/Startup.cs
--- a/Engine/Source/Programs/Horde/Horde.Build/Startup.cs
+++ b/Engine/Source/Programs/Horde/Horde.Build/Startup.cs
@@ -30,6 +30,7 @@ using Horde.Build.Compute;
 using Horde.Build.Compute.Impl;
 using Horde.Build.Config;
 using Horde.Build.Fleet.Autoscale;
+using Horde.Build.Issues.Impl;
 using Horde.Build.Jobs;
 using Horde.Build.Logs;
 using Horde.Build.Logs.Builder;
@@ -562,6 +563,7 @@ namespace Horde.Build
 				services.AddHostedService(provider => provider.GetRequiredService<CommitService>());
 				services.AddHostedService(provider => provider.GetRequiredService<ConsistencyService>());
 				services.AddHostedService(provider => provider.GetRequiredService<IIssueService>());
+				services.AddHostedService<IssueReportService>();
 				services.AddHostedService(provider => (LogFileService)provider.GetRequiredService<ILogFileService>());
 				services.AddHostedService(provider => (NotificationService)provider.GetRequiredService<INotificationService>());
 				if (!settings.DisableSchedules)
diff --git a/Engine/Source/Programs/Horde/Horde.Build/Utilities/BlockKit/HeaderBlock.cs b/Engine/Source/Programs/Horde/Horde.Build/Utilities/BlockKit/HeaderBlock.cs
--- a/Engine/Source/Programs/Horde/Horde.Build/Utilities/BlockKit/HeaderBlock.cs
+++ b/Engine/Source/Programs/Horde/Horde.Build/Utilities/BlockKit/HeaderBlock.cs
@@ -13,11 +13,10 @@ namespace Horde.Build.Utilities.Slack.BlockKit
 		/// Constructor
 		/// </summary>
 		/// <param name="text">Any text to initiale the Section with</param>
-		/// <param name="isMarkdown">If true, text is markdown, else text is considered plain text</param>
 		/// <param name="emoji">If true, allows escaping of emojis with colons.</param>
-		public HeaderBlock(string text = "", bool isMarkdown = true, bool emoji = false)
+		public HeaderBlock(string text = "", bool emoji = false)
 		{
-			Text = new TextObject(text, isMarkdown, emoji);
+			Text = new TextObject(text, false, emoji);
 		}
 
 		/// <summary>
