Horde: Missing file.

diff --git a/Engine/Source/Programs/Horde/Horde.Build/Issues/Impl/IssueReportService.cs b/Engine/Source/Programs/Horde/Horde.Build/Issues/Impl/IssueReportService.cs
new file mode 100644
--- /dev/null
+++ b/Engine/Source/Programs/Horde/Horde.Build/Issues/Impl/IssueReportService.cs
@@ -0,0 +1,212 @@
+// Copyright Epic Games, Inc. All Rights Reserved.
+
+using System;
+using System.Collections.Generic;
+using System.Linq;
+using System.Threading;
+using System.Threading.Tasks;
+using Horde.Build.Collections;
+using Horde.Build.Config;
+using Horde.Build.Models;
+using Horde.Build.Notifications;
+using Horde.Build.Server;
+using Horde.Build.Utilities;
+using HordeCommon;
+using Microsoft.Extensions.Hosting;
+using Microsoft.Extensions.Logging;
+
+namespace Horde.Build.Issues.Impl
+{
+	using TemplateRefId = StringId<TemplateRef>;
+
+#pragma warning disable CS1591 // Missing XML comment for publicly visible type or member
+	public class IssueReport
+	{
+		public DateTimeOffset Time { get; }
+		public IStream Stream { get; }
+		public WorkflowConfig Workflow { get; }
+		public int NumSteps { get; set; }
+		public int NumPassingSteps { get; set; }
+		public List<IIssue> Issues { get; } = new List<IIssue>();
+		public List<IIssueSpan> IssueSpans { get; } = new List<IIssueSpan>();
+
+		public IssueReport(DateTimeOffset time, IStream stream, WorkflowConfig workflow)
+		{
+			Time = time;
+			Stream = stream;
+			Workflow = workflow;
+		}
+	}
+#pragma warning restore CS1591 // Missing XML comment for publicly visible type or member
+
+	[SingletonDocument("6268871c211d05611b3e4fd8")]
+	class IssueReportState : SingletonBase
+	{
+		public Dictionary<string, DateTime> KeyToLastReportTime { get; set; } = new Dictionary<string, DateTime>();
+	}
+
+	/// <summary>
+	/// Posts summaries for all the open issues in different streams to Slack channels
+	/// </summary>
+	public class IssueReportService : IHostedService
+	{
+		readonly SingletonDocument<IssueReportState> _state;
+		readonly IStreamCollection _streamCollection;
+		readonly IIssueCollection _issueCollection;
+		readonly IJobCollection _jobCollection;
+		readonly ConfigCollection _configCollection;
+		readonly INotificationService _notificationService;
+		readonly IClock _clock;
+		readonly ITicker _ticker;
+		readonly ILogger<IssueReportService> _logger;
+
+		/// <summary>
+		/// Constructor
+		/// </summary>
+		public IssueReportService(MongoService mongoService, IStreamCollection streamCollection, IIssueCollection issueCollection, IJobCollection jobCollection, ConfigCollection configCollection, INotificationService notificationService, IClock clock, ILogger<IssueReportService> logger)
+		{
+			_state = new SingletonDocument<IssueReportState>(mongoService);
+			_streamCollection = streamCollection;
+			_issueCollection = issueCollection;
+			_jobCollection = jobCollection;
+			_configCollection = configCollection;
+			_notificationService = notificationService;
+			_clock = clock;
+			_ticker = clock.AddSharedTicker<IssueReportService>(TimeSpan.FromMinutes(5.0), TickAsync, logger);
+			_logger = logger;
+		}
+
+		/// <inheritdoc/>
+		public async Task StartAsync(CancellationToken cancellationToken)
+		{
+			await _ticker.StartAsync();
+		}
+
+		/// <inheritdoc/>
+		public async Task StopAsync(CancellationToken cancellationToken)
+		{
+			await _ticker.StopAsync();
+		}
+
+		async ValueTask TickAsync(CancellationToken cancellationToken)
+		{
+			IssueReportState state = await _state.GetAsync();
+			HashSet<string> invalidKeys = new HashSet<string>(state.KeyToLastReportTime.Keys, StringComparer.Ordinal);
+
+			DateTime currentTime = _clock.UtcNow;
+
+			List<IStream> streams = await _streamCollection.FindAllAsync();
+			foreach (IStream stream in streams)
+			{
+				StreamConfig config = await _configCollection.GetConfigAsync<StreamConfig>(stream.ConfigRevision);
+				if (config.Workflows.Count > 0)
+				{
+					List<IIssue>? issues = null;
+					List<IIssueSpan>? spans = null;
+					List<IJob>? jobs = null;
+
+					foreach (WorkflowConfig workflow in config.Workflows)
+					{
+						string key = $"{stream.Id}:{workflow.Id}";
+						invalidKeys.Remove(key);
+
+						DateTime lastReportTime;
+						if (!state.KeyToLastReportTime.TryGetValue(key, out lastReportTime))
+						{
+							state = await _state.UpdateAsync(s => s.KeyToLastReportTime[key] = currentTime);
+							continue;
+						}
+
+						DateTime lastScheduledReportTime = GetLastScheduledReportTime(workflow, currentTime);
+						if (lastReportTime > lastScheduledReportTime)
+						{
+							continue;
+						}
+
+						DateTime prevScheduledReportTime = GetLastScheduledReportTime(workflow, lastScheduledReportTime - TimeSpan.FromMinutes(1.0));
+
+						_logger.LogInformation("Creating report for {StreamId} workflow {WorkflowId}", stream.Id, workflow.Id);
+
+						issues ??= await _issueCollection.FindIssuesAsync(streamId: stream.Id);
+						spans ??= await _issueCollection.FindSpansAsync(issueIds: issues.Select(x => x.Id).ToArray());
+						jobs ??= await _jobCollection.FindAsync(streamId: stream.Id, minCreateTime: prevScheduledReportTime);
+
+						HashSet<TemplateRefId> wfTemplateRefIds = new HashSet<TemplateRefId>();
+						foreach (Api.CreateTemplateRefRequest templateRef in config.Templates)
+						{
+							if (templateRef.Id != null && templateRef.Workflow != null && templateRef.Workflow == workflow.Id)
+							{
+								wfTemplateRefIds.Add(new TemplateRefId(templateRef.Id));
+							}
+						}
+
+						IssueReport report = new IssueReport(currentTime, stream, workflow);
+
+						foreach (IJob job in jobs)
+						{
+							if (wfTemplateRefIds.Contains(job.TemplateId))
+							{
+								foreach (IJobStep step in job.Batches.SelectMany(x => x.Steps).Where(x => x.State == JobStepState.Completed))
+								{
+									report.NumSteps++;
+									if (step.Outcome == JobStepOutcome.Success)
+									{
+										report.NumPassingSteps++;
+									}
+								}
+							}
+						}
+
+						foreach (IIssueSpan span in spans)
+						{
+							if (span.LastSuccess != null && wfTemplateRefIds.Contains(span.TemplateRefId))
+							{
+								report.IssueSpans.Add(span);
+							}
+						}
+
+						HashSet<int> issueIds = new HashSet<int>(report.IssueSpans.Select(x => x.IssueId));
+						report.Issues.AddRange(issues.Where(x => issueIds.Contains(x.Id)));
+
+						_notificationService.SendIssueReport(report);
+
+						state = await _state.UpdateAsync(s => s.KeyToLastReportTime[key] = currentTime);
+					}
+				}
+			}
+
+			if (invalidKeys.Count > 0)
+			{
+				void RemoveInvalidKeys(IssueReportState state)
+				{
+					foreach (string invalidKey in invalidKeys)
+					{
+						state.KeyToLastReportTime.Remove(invalidKey);
+					}
+				}
+				state = await _state.UpdateAsync(RemoveInvalidKeys);
+			}
+		}
+
+		static DateTime GetLastScheduledReportTime(WorkflowConfig workflow, DateTime currentTime)
+		{
+			if (workflow.ReportTimes.Count > 0)
+			{
+				DateTime startOfDay = currentTime - currentTime.TimeOfDay;
+				for (; ; )
+				{
+					for (int idx = workflow.ReportTimes.Count - 1; idx >= 0; idx--)
+					{
+						DateTime reportTime = startOfDay + workflow.ReportTimes[idx];
+						if (reportTime < currentTime)
+						{
+							return reportTime;
+						}
+					}
+					startOfDay -= TimeSpan.FromDays(1.0);
+				}
+			}
+			return DateTime.MinValue;
+		}
+	}
+}
