Horde: Fix tests.

diff --git a/Engine/Source/Programs/Horde/Horde.Build.Tests/NotificationServiceTests.cs b/Engine/Source/Programs/Horde/Horde.Build.Tests/NotificationServiceTests.cs
--- a/Engine/Source/Programs/Horde/Horde.Build.Tests/NotificationServiceTests.cs
+++ b/Engine/Source/Programs/Horde/Horde.Build.Tests/NotificationServiceTests.cs
@@ -4,6 +4,7 @@ using System;
 using System.Collections.Generic;
 using System.Threading.Tasks;
 using Horde.Build.Api;
+using Horde.Build.Issues.Impl;
 using Horde.Build.Models;
 using Horde.Build.Notifications;
 using Horde.Build.Notifications.Impl;
@@ -37,6 +38,7 @@ namespace Horde.Build.Tests
 		public Task NotifyIssueUpdatedAsync(IIssue issue) { throw new NotImplementedException(); }
 		public Task NotifyConfigUpdateFailureAsync(string errorMessage, string fileName, int? change = null, IUser? author = null, string? description = null) { throw new NotImplementedException(); }
 		public Task NotifyDeviceServiceAsync(string message, IDevice? device = null, IDevicePool? pool = null, IStream? stream = null, IJob? job = null, IJobStep? step = null, INode? node = null, IUser? user = null) { throw new NotImplementedException(); }
+		public Task SendIssueReportAsync(IssueReport report) => throw new NotImplementedException();
 	}
 
 	[TestClass]
diff --git a/Engine/Source/Programs/Horde/Horde.Build/Startup.cs b/Engine/Source/Programs/Horde/Horde.Build/Startup.cs
--- a/Engine/Source/Programs/Horde/Horde.Build/Startup.cs
+++ b/Engine/Source/Programs/Horde/Horde.Build/Startup.cs
@@ -227,6 +227,24 @@ namespace Horde.Build
 			}
 		}
 
+		class JsonTimeSpanConverter : JsonConverter<TimeSpan>
+		{
+			public override TimeSpan Read(ref Utf8JsonReader reader, Type typeToConvert, JsonSerializerOptions options)
+			{
+				string? str = reader.GetString();
+				if (str == null)
+				{
+					throw new InvalidDataException("Unable to parse TimeSpan");
+				}
+				return TimeSpan.Parse(str, CultureInfo.CurrentCulture);
+			}
+
+			public override void Write(Utf8JsonWriter writer, TimeSpan timeSpan, JsonSerializerOptions options)
+			{
+				writer.WriteStringValue(timeSpan.ToString("c"));
+			}
+		}
+
 		public Startup(IConfiguration configuration)
 		{
 			Configuration = configuration;
@@ -678,6 +696,7 @@ namespace Horde.Build
 			options.Converters.Add(new JsonKnownTypesConverterFactory());
 			options.Converters.Add(new JsonStringIdConverterFactory());
 			options.Converters.Add(new JsonDateTimeConverter());
+			options.Converters.Add(new JsonTimeSpanConverter());
 		}
 
 		public class StorageBackendSettings<T> : IFileSystemStorageOptions, IAwsStorageOptions
