Horde: Ignore allow list for Slack channel notifications.

diff --git a/Engine/Source/Programs/Horde/Horde.Build/Notifications/Impl/SlackNotificationSink.cs b/Engine/Source/Programs/Horde/Horde.Build/Notifications/Impl/SlackNotificationSink.cs
--- a/Engine/Source/Programs/Horde/Horde.Build/Notifications/Impl/SlackNotificationSink.cs
+++ b/Engine/Source/Programs/Horde/Horde.Build/Notifications/Impl/SlackNotificationSink.cs
@@ -1309,7 +1309,7 @@ namespace Horde.Build.Notifications.Impl
 
 		private async Task SendMessageAsync(string recipient, string? text = null, BlockBase[]? blocks = null, BlockKitAttachment[]? attachments = null)
 		{
-			if (_allowUsers != null && !_allowUsers.Contains(recipient))
+			if (_allowUsers != null && !_allowUsers.Contains(recipient) && !recipient.StartsWith("#", StringComparison.Ordinal))
 			{
 				_logger.LogDebug("Suppressing message to {Recipient}: {Text}", recipient, text);
 				return;
