Horde: Add an option to force use of the local Perforce environment. Useful when running locally against an external database.

diff --git a/Engine/Source/Programs/Horde/Horde.Build/Perforce/Impl/PerforceService.cs b/Engine/Source/Programs/Horde/Horde.Build/Perforce/Impl/PerforceService.cs
--- a/Engine/Source/Programs/Horde/Horde.Build/Perforce/Impl/PerforceService.cs
+++ b/Engine/Source/Programs/Horde/Horde.Build/Perforce/Impl/PerforceService.cs
@@ -64,6 +64,10 @@ namespace Horde.Build.Services
 		readonly LazyCachedValue<Task<Globals>> _cachedGlobals;
 		readonly ILogger _logger;
 
+		// Useful overrides for local debugging with read-only data
+		string? _perforceServerOverride;
+		string? _perforceUserOverride;
+
 		/// <summary>
 		/// Object used for controlling access to the access user tickets
 		/// </summary>
@@ -102,6 +106,13 @@ namespace Horde.Build.Services
 			P4.P4Debugging.SetBridgeLogFunction(_logBridgeDelegate);
 
 			P4.LogFile.SetLoggingFunction(LogPerforce);
+
+			if(settings.Value.UseLocalPerforceEnv)
+			{
+				IPerforceSettings perforceSettings = PerforceSettings.Default;
+				_perforceServerOverride = perforceSettings.ServerAndPort;
+				_perforceUserOverride = perforceSettings.UserName;
+			}
 		}
 
 		public void Dispose()
@@ -148,6 +159,11 @@ namespace Horde.Build.Services
 			return cluster;
 		}
 
+		async ValueTask<string> SelectServerAddressAsync(PerforceCluster cluster)
+		{
+			return (await SelectServer(cluster)).ServerAndPort;
+		}
+
 		async Task<IPerforceServer> SelectServer(PerforceCluster cluster)
 		{
 			using IScope scope = GlobalTracer.Instance.BuildSpan("PerforceService.SelectServer").StartActive();
@@ -210,22 +226,40 @@ namespace Horde.Build.Services
 
 			PerforceCluster cluster = await GetClusterAsync(clusterName);
 
-			string? userName = cluster.ServiceAccount ?? Environment.UserName;
-
-			string? password = null;
-			if (cluster.ServiceAccount != null)
+			string? userName, password;
+			if (_perforceUserOverride != null)
+			{
+				userName = _perforceUserOverride;
+				password = null;
+			}
+			else if (cluster.ServiceAccount != null)
 			{
-				PerforceCredentials? credentials = cluster.Credentials.FirstOrDefault(x => x.UserName.Equals(userName, StringComparison.OrdinalIgnoreCase));
+				PerforceCredentials? credentials = cluster.Credentials.FirstOrDefault(x => x.UserName.Equals(cluster.ServiceAccount, StringComparison.OrdinalIgnoreCase));
 				if (credentials == null)
 				{
 					throw new Exception($"No credentials defined for {cluster.ServiceAccount} on {cluster.Name}");
 				}
+
+				userName = credentials.UserName;
 				password = credentials.Password;
 			}
+			else
+			{
+				userName = PerforceSettings.Default.UserName;
+				password = null;
+			}
 
-			IPerforceServer server = await SelectServer(cluster);
+			string serverAndPort;
+			if (_perforceServerOverride != null)
+			{
+				serverAndPort = _perforceServerOverride;
+			}
+			else
+			{
+				serverAndPort = await SelectServerAddressAsync(cluster);
+			}
 
-			PerforceSettings settings = new PerforceSettings(server.ServerAndPort, userName);
+			PerforceSettings settings = new PerforceSettings(serverAndPort, userName);
 			settings.Password = password;
 			settings.AppName = "Horde.Build";
 			settings.ClientName = "__DOES_NOT_EXIST__";
diff --git a/Engine/Source/Programs/Horde/Horde.Build/ServerSettings.cs b/Engine/Source/Programs/Horde/Horde.Build/ServerSettings.cs
--- a/Engine/Source/Programs/Horde/Horde.Build/ServerSettings.cs
+++ b/Engine/Source/Programs/Horde/Horde.Build/ServerSettings.cs
@@ -517,6 +517,11 @@ namespace Horde.Build
 		/// </summary>
 		public NamespaceId ToolNamespaceId { get; set; } = new NamespaceId("horde.p4");
 
+		/// <summary>
+		/// Whether to 
+		/// </summary>
+		public bool UseLocalPerforceEnv { get; set; }
+
 		/// <summary>
 		/// Lazily computed timezone value
 		/// </summary>
