Horde: Disable hosted services when running in read-only mode.

diff --git a/Engine/Source/Programs/Horde/Horde.Build/Startup.cs b/Engine/Source/Programs/Horde/Horde.Build/Startup.cs
--- a/Engine/Source/Programs/Horde/Horde.Build/Startup.cs
+++ b/Engine/Source/Programs/Horde/Horde.Build/Startup.cs
@@ -571,7 +571,7 @@ namespace Horde.Build
 			// Hosted service that needs to run no matter the run mode of the process (server vs worker)
 			services.AddHostedService(provider => (DowntimeService)provider.GetRequiredService<IDowntimeService>());
 
-			if (settings.IsRunModeActive(RunMode.Worker))
+			if (settings.IsRunModeActive(RunMode.Worker) && !settings.DatabaseReadOnlyMode)
 			{
 				services.AddHostedService<MongoUpgradeService>();
 
